#ue4
#ios
* fix for not checking if a file exists before trying to change its attributes

diff --git a/Engine/Source/Programs/UnrealBuildTool/IOS/UEDeployIOS.cs b/Engine/Source/Programs/UnrealBuildTool/IOS/UEDeployIOS.cs
--- a/Engine/Source/Programs/UnrealBuildTool/IOS/UEDeployIOS.cs
+++ b/Engine/Source/Programs/UnrealBuildTool/IOS/UEDeployIOS.cs
@@ -162,8 +162,11 @@ namespace UnrealBuildTool.IOS
 			CopyFolder(InEngineDir + "/Content/Stats", AppDirectory + "/cookeddata/engine/content/stats", true);
 
 			// ensure the destination is writable
-			FileInfo GameFileInfo = new FileInfo (AppDirectory + "/" + GameName);
-			GameFileInfo.Attributes = GameFileInfo.Attributes & ~FileAttributes.ReadOnly;
+			if (File.Exists(AppDirectory + "/" + GameName))
+			{
+				FileInfo GameFileInfo = new FileInfo(AppDirectory + "/" + GameName);
+				GameFileInfo.Attributes = GameFileInfo.Attributes & ~FileAttributes.ReadOnly;
+			}
 
 			// copy the GameName binary
 			File.Copy(BinaryPath + "/" + GameExeName, AppDirectory + "/" + GameName, true);
