Fixed TTP #331034 LIVE: Mac: Compiling game in debug from the editor is looking for dylibs in a wrong place

diff --git a/Engine/Source/Programs/UnrealBuildTool/Configuration/UEBuildTarget.cs b/Engine/Source/Programs/UnrealBuildTool/Configuration/UEBuildTarget.cs
--- a/Engine/Source/Programs/UnrealBuildTool/Configuration/UEBuildTarget.cs
+++ b/Engine/Source/Programs/UnrealBuildTool/Configuration/UEBuildTarget.cs
@@ -1147,6 +1147,12 @@ namespace UnrealBuildTool
 			}
 			AppBinaries.AddRange( NewBinaries );
 
+			// On Mac AppBinaries paths for non-console targets need to be adjusted to be inside the app bundle
+			if (GlobalLinkEnvironment.Config.TargetPlatform == CPPTargetPlatform.Mac && !GlobalLinkEnvironment.Config.bIsBuildingConsoleApplication)
+			{
+				MacToolChain.FixBundleBinariesPaths( this, AppBinaries );
+			}
+
 			// If we're building a single module, then find the binary for that module and add it to our target
 			if (OnlyModules.Count > 0)
 			{
@@ -1333,12 +1339,6 @@ namespace UnrealBuildTool
 			// On Mac we have actions that should be executed after all the binaries are created
 			if (GlobalLinkEnvironment.Config.TargetPlatform == CPPTargetPlatform.Mac)
 			{
-				// Also, for non-console targets, AppBinaries paths need to be adjusted to be inside the app bundle
-				if (!GlobalLinkEnvironment.Config.bIsBuildingConsoleApplication)
-				{
-					MacToolChain.FixBundleBinariesPaths( this, AppBinaries );
-				}
-
 				MacToolChain.SetupBundleDependencies(AppBinaries, GameName);
 			}
 
