Added Action.CommandDescription to UBT and used it in MacToolChain to make the build log nicer.

diff --git a/Engine/Source/Programs/UnrealBuildTool/Mac/MacToolChain.cs b/Engine/Source/Programs/UnrealBuildTool/Mac/MacToolChain.cs
--- a/Engine/Source/Programs/UnrealBuildTool/Mac/MacToolChain.cs
+++ b/Engine/Source/Programs/UnrealBuildTool/Mac/MacToolChain.cs
@@ -393,7 +393,8 @@ namespace UnrealBuildTool
 				CompileAction.WorkingDirectory = GetMacDevSrcRoot();
 				CompileAction.CommandPath = "xcrun";
 				CompileAction.CommandArguments = MacCompiler + Arguments + FileArguments + CompileEnvironment.Config.AdditionalArguments;
-				CompileAction.StatusDescription = string.Format("{0}", Path.GetFileName(SourceFile.AbsolutePath));
+				CompileAction.CommandDescription = "Compile";
+				CompileAction.StatusDescription = Path.GetFileName(SourceFile.AbsolutePath);
 				CompileAction.StatusDetailedDescription = SourceFile.Description;
 				CompileAction.bIsGCCCompiler = true;
 				// We're already distributing the command by execution on Mac.
@@ -494,6 +495,7 @@ namespace UnrealBuildTool
 
 			LinkAction.WorkingDirectory = GetMacDevSrcRoot();
 			LinkAction.CommandPath = "/bin/sh";
+			LinkAction.CommandDescription = "Link";
 
 			string EngineCLVersion = LoadEngineCLVersion();
 			string EngineDisplayVersion = LoadEngineDisplayVersion(true);
@@ -740,7 +742,7 @@ namespace UnrealBuildTool
 			// Only execute linking on the local Mac.
 			LinkAction.bCanExecuteRemotely = false;
 
-			LinkAction.StatusDescription = string.Format("{0}", OutputFile.AbsolutePath);
+			LinkAction.StatusDescription = Path.GetFileName(OutputFile.AbsolutePath);
 			LinkAction.OutputEventHandler = new DataReceivedEventHandler(RemoteOutputReceivedEventHandler);
 			
 			LinkAction.ProducedItems.Add(RemoteOutputFile);
@@ -929,6 +931,7 @@ namespace UnrealBuildTool
 			Action LinkAction = new Action(ActionType.Link);
 			LinkAction.WorkingDirectory = Path.GetFullPath(".");
 			LinkAction.CommandPath = "/bin/sh";
+			LinkAction.CommandDescription = "";
 
 			// Call the FixDylibDependencies.sh script which will link the dylibs and the main executable, this time proper ones, as it's called
 			// once all are already created, so the cross dependency problem no longer prevents linking.
@@ -1019,6 +1022,7 @@ namespace UnrealBuildTool
 			Action CreateAppBundleAction = new Action(ActionType.CreateAppBundle);
 			CreateAppBundleAction.WorkingDirectory = Path.GetFullPath(".");
 			CreateAppBundleAction.CommandPath = "/bin/sh";
+			CreateAppBundleAction.CommandDescription = "";
 
 			// make path to the script
 			FileItem BundleScript = FileItem.GetItemByFullPath(Path.Combine(LinkEnvironment.Config.IntermediateDirectory, "CreateAppBundle.sh"));
@@ -1033,7 +1037,7 @@ namespace UnrealBuildTool
 			CreateAppBundleAction.CommandArguments = "\"" + RemoteBundleScript.AbsolutePath + "\"";
 			CreateAppBundleAction.PrerequisiteItems.Add(FixDylibOutputFile);
 			CreateAppBundleAction.ProducedItems.Add(RemoteDestFile);
-			CreateAppBundleAction.StatusDescription = string.Format("{0}.app - creating app bundle", Path.GetFileName(Executable.AbsolutePath));
+			CreateAppBundleAction.StatusDescription = string.Format("Creating app bundle: {0}.app", Path.GetFileName(Executable.AbsolutePath));
 			CreateAppBundleAction.bCanExecuteRemotely = false;
 
 			return RemoteDestFile;
diff --git a/Engine/Source/Programs/UnrealBuildTool/System/ActionGraph.cs b/Engine/Source/Programs/UnrealBuildTool/System/ActionGraph.cs
--- a/Engine/Source/Programs/UnrealBuildTool/System/ActionGraph.cs
+++ b/Engine/Source/Programs/UnrealBuildTool/System/ActionGraph.cs
@@ -38,6 +38,7 @@ namespace UnrealBuildTool
 		public string WorkingDirectory = null;
 		public string CommandPath = null;
 		public string CommandArguments = null;
+		public string CommandDescription = null;
 		public string StatusDescription = "...";
 		public string StatusDetailedDescription = "";
 		public bool bCanExecuteRemotely = false;
diff --git a/Engine/Source/Programs/UnrealBuildTool/System/LocalExecutor.cs b/Engine/Source/Programs/UnrealBuildTool/System/LocalExecutor.cs
--- a/Engine/Source/Programs/UnrealBuildTool/System/LocalExecutor.cs
+++ b/Engine/Source/Programs/UnrealBuildTool/System/LocalExecutor.cs
@@ -141,7 +141,15 @@ namespace UnrealBuildTool
 				// Log summary if wanted.
 				else if (Action.bShouldOutputStatusDescription)
 				{
-					Log.TraceInformation("{0} {1}", Path.GetFileName(ExpandedCommandPath), Action.StatusDescription);
+					string CommandDescription = Action.CommandDescription != null ? Action.CommandDescription : Path.GetFileName(ExpandedCommandPath);
+					if (string.IsNullOrEmpty(CommandDescription))
+					{
+						Log.TraceInformation(Action.StatusDescription);
+					}
+					else
+					{
+						Log.TraceInformation("{0} {1}", CommandDescription, Action.StatusDescription);
+					}
 				}
 
 				// Try to launch the action's process, and produce a friendly error message if it fails.
