#UE4
#HTML5

       - fix executable name for HTML5 Win32 architecture.

diff --git a/Engine/Source/Programs/UnrealBuildTool/HTML5/UEBuildHTML5.cs b/Engine/Source/Programs/UnrealBuildTool/HTML5/UEBuildHTML5.cs
--- a/Engine/Source/Programs/UnrealBuildTool/HTML5/UEBuildHTML5.cs
+++ b/Engine/Source/Programs/UnrealBuildTool/HTML5/UEBuildHTML5.cs
@@ -22,7 +22,7 @@ namespace UnrealBuildTool
 		public override string ModifyNMakeOutput(string ExeName)
 		{
 			// nmake Run should always run the win32 version
-			return Path.ChangeExtension(ExeName, ".exe");
+			return Path.ChangeExtension(ExeName+GetActiveArchitecture(), ".exe");
 		}
 
 		/**
