#ue4
* fix for bad BranchDirectory in the RemoteToolChain

diff --git a/Engine/Source/Programs/UnrealBuildTool/ToolChain/RemoteToolChain.cs b/Engine/Source/Programs/UnrealBuildTool/ToolChain/RemoteToolChain.cs
--- a/Engine/Source/Programs/UnrealBuildTool/ToolChain/RemoteToolChain.cs
+++ b/Engine/Source/Programs/UnrealBuildTool/ToolChain/RemoteToolChain.cs
@@ -34,7 +34,7 @@ namespace UnrealBuildTool
 		public static string UserDevRootMac = "/UE4/Builds/";
 
 		/** The directory that this local branch is in, without drive information (strip off X:\ from X:\UE4\iOS) */
-		public static string BranchDirectory = Environment.MachineName + "\\" + Path.GetFullPath("..\\..\\").Substring(3);
+		public static string BranchDirectory = Environment.MachineName + "\\" + Path.GetFullPath(".\\").Substring(3);
 		public static string BranchDirectoryMac = BranchDirectory.Replace("\\", "/");
 
 
@@ -52,6 +52,9 @@ namespace UnrealBuildTool
             ErrorMessageTokens.Add("[BEROR]");
             ErrorMessageTokens.Add("IPP ERROR");
             ErrorMessageTokens.Add("System.Net.Sockets.SocketException");
+
+			BranchDirectory = BranchDirectory.Replace("Engine\\Source\\", "");
+			BranchDirectoryMac = BranchDirectory.Replace("\\", "/");
         }
 
         // Do any one-time, global initialization for the tool chain
