#UE4 Quick fix for Mono builds that doesn't include some classes overridden by XML system. Now it's ignoring missing classes.

diff --git a/Engine/Source/Programs/UnrealBuildTool/Utilities/XmlConfigLoader.cs b/Engine/Source/Programs/UnrealBuildTool/Utilities/XmlConfigLoader.cs
--- a/Engine/Source/Programs/UnrealBuildTool/Utilities/XmlConfigLoader.cs
+++ b/Engine/Source/Programs/UnrealBuildTool/Utilities/XmlConfigLoader.cs
@@ -163,7 +163,8 @@ namespace UnrealBuildTool
 
 				if(ClassType == null)
 				{
-					throw new BuildException("XmlConfig Loading: class '{0}' doesn't exist.", XmlClass.Name);
+					Log.TraceVerbose("XmlConfig Loading: class '{0}' doesn't exist.", XmlClass.Name);
+					continue;
 				}
 
 				XmlConfigLoaderClassData ClassData;
