#UE4 Set xml parsing to use invariant culture.

diff --git a/Engine/Source/Programs/UnrealBuildTool/Utilities/XmlConfigLoader.cs b/Engine/Source/Programs/UnrealBuildTool/Utilities/XmlConfigLoader.cs
--- a/Engine/Source/Programs/UnrealBuildTool/Utilities/XmlConfigLoader.cs
+++ b/Engine/Source/Programs/UnrealBuildTool/Utilities/XmlConfigLoader.cs
@@ -1,9 +1,11 @@
 ﻿using System;
 using System.Collections.Generic;
+using System.Globalization;
 using System.IO;
 using System.Linq;
 using System.Reflection;
 using System.Text;
+using System.Threading;
 using System.Xml;
 
 namespace UnrealBuildTool
@@ -236,12 +238,19 @@ namespace UnrealBuildTool
 				// Declaring parameters array used by TryParse method.
 				// Second parameter is "out", so you have to just
 				// assign placeholder null to it.
+
+				var OldCulture = Thread.CurrentThread.CurrentCulture;
+				Thread.CurrentThread.CurrentCulture = CultureInfo.InvariantCulture;
+
 				var Parameters = new object[] { Text, null };
 				if (!(bool)TryParseMethod.Invoke(null, Parameters))
 				{
+					Thread.CurrentThread.CurrentCulture = OldCulture;
 					throw new BuildException("BuildConfiguration Loading: Parsing {0} value from \"{1}\" failed.", FieldType.Name, Text);
 				}
 
+				Thread.CurrentThread.CurrentCulture = OldCulture;
+
 				// If Invoke returned true, the second object of the
 				// parameters array is set to the parsed value.
 				return Parameters[1];
