#UE4 Fixing hack for parsing with InvariantCulture in UBT xml system.
#summary The parsing functions now look for Parse method instead of TryParse and try to find one with IFormatProvider to pass InvariantCulture. If such Parse method exists uses it, if not then tries to find one without format provider and uses that one.

diff --git a/Engine/Source/Programs/UnrealBuildTool/Utilities/XmlConfigLoader.cs b/Engine/Source/Programs/UnrealBuildTool/Utilities/XmlConfigLoader.cs
--- a/Engine/Source/Programs/UnrealBuildTool/Utilities/XmlConfigLoader.cs
+++ b/Engine/Source/Programs/UnrealBuildTool/Utilities/XmlConfigLoader.cs
@@ -223,40 +223,87 @@ namespace UnrealBuildTool
 			}
 			else
 			{
-				// Getting TryParse method for FieldType which is required,
-				// if it doesn't exists for complex type, author should add
-				// one. The signature should be:
-				// static bool TryParse(string Input, out T value),
-				// where T is containing type.
-				var TryParseMethod = FieldType.GetMethod("TryParse", new Type[] { typeof(System.String), FieldType.MakeByRefType() });
-
-				if (TryParseMethod == null)
-				{
-					throw new BuildException("BuildConfiguration Loading: Parsing of the type {0} is not supported.", FieldType.Name);
-				}
-
 				// Declaring parameters array used by TryParse method.
 				// Second parameter is "out", so you have to just
 				// assign placeholder null to it.
-
-				var OldCulture = Thread.CurrentThread.CurrentCulture;
-				Thread.CurrentThread.CurrentCulture = CultureInfo.InvariantCulture;
-
-				var Parameters = new object[] { Text, null };
-				if (!(bool)TryParseMethod.Invoke(null, Parameters))
+				object ParsedValue;
+				if (!TryParse(FieldType, Text, out ParsedValue))
 				{
-					Thread.CurrentThread.CurrentCulture = OldCulture;
 					throw new BuildException("BuildConfiguration Loading: Parsing {0} value from \"{1}\" failed.", FieldType.Name, Text);
 				}
 
-				Thread.CurrentThread.CurrentCulture = OldCulture;
-
 				// If Invoke returned true, the second object of the
 				// parameters array is set to the parsed value.
-				return Parameters[1];
+				return ParsedValue;
 			}
 		}
 
+		/// <summary>
+		/// Emulates TryParse behavior on custom type. If the type implements
+		/// Parse(string, IFormatProvider) or Parse(string) static method uses
+		/// one of them to parse with preference of the one with format
+		/// provider (but passes invariant culture).
+		/// </summary>
+		/// <param name="ParsingType">Type to parse.</param>
+		/// <param name="UnparsedValue">String representation of the value.</param>
+		/// <param name="ParsedValue">Output parsed value.</param>
+		/// <returns>True if parsing succeeded. False otherwise.</returns>
+		private static bool TryParse(Type ParsingType, string UnparsedValue, out object ParsedValue)
+		{
+			// Getting Parse method for FieldType which is required,
+			// if it doesn't exists for complex type, author should add
+			// one. The signature should be one of:
+			//     static T Parse(string Input, IFormatProvider Provider) or
+			//     static T Parse(string Input)
+			// where T is containing type.
+			// The one with format provider is preferred and invoked with
+			// InvariantCulture.
+
+			bool bWithCulture = true;
+			var ParseMethod = ParsingType.GetMethod("Parse", new Type[] { typeof(System.String), typeof(IFormatProvider) });
+
+			if(ParseMethod == null)
+			{
+				ParseMethod = ParsingType.GetMethod("Parse", new Type[] { typeof(System.String) });
+				bWithCulture = false;
+			}
+
+			if (ParseMethod == null)
+			{
+				throw new BuildException("BuildConfiguration Loading: Parsing of the type {0} is not supported.", ParsingType.Name);
+			}
+
+			var ParametersList = new List<object> { UnparsedValue };
+
+			if(bWithCulture)
+			{
+				ParametersList.Add(CultureInfo.InvariantCulture);
+			}
+
+			try
+			{
+				ParsedValue = ParseMethod.Invoke(null, ParametersList.ToArray());
+			}
+			catch(Exception e)
+			{
+				if(e is TargetInvocationException &&
+					(
+						e.InnerException is ArgumentNullException ||
+						e.InnerException is FormatException ||
+						e.InnerException is OverflowException
+					)
+				)
+				{
+					ParsedValue = null;
+					return false;
+				}
+
+				throw;
+			}
+
+			return true;
+		}
+
 		// Map to store class data in.
 		private static readonly Dictionary<Type, XmlConfigLoaderClassData> Data = new Dictionary<Type, XmlConfigLoaderClassData>();
 	}
