GitHub pull request integrate

https://github.com/EpicGames/UnrealEngine/pull/45/

Correct spelling error

From GitHub user "abergmeier"

diff --git a/Engine/Source/Programs/UnrealBuildTool/System/RulesCompiler.cs b/Engine/Source/Programs/UnrealBuildTool/System/RulesCompiler.cs
--- a/Engine/Source/Programs/UnrealBuildTool/System/RulesCompiler.cs
+++ b/Engine/Source/Programs/UnrealBuildTool/System/RulesCompiler.cs
@@ -303,7 +303,7 @@ namespace UnrealBuildTool
         public bool bUsesSteam;
 
 		/// <summary>
-		/// Whether the project uses visual Slate UI (as opposed to the low level windowing/messaging which is alway used)
+		/// Whether the project uses visual Slate UI (as opposed to the low level windowing/messaging which is always used)
 		/// </summary>
 		public bool bUsesSlate = true;
 
