Added opt-in framework for switching SDKs on the fly.
Made Linux toolchain support that.

diff --git a/Engine/Source/Programs/UnrealBuildTool/Configuration/UEBuildPlatform.cs b/Engine/Source/Programs/UnrealBuildTool/Configuration/UEBuildPlatform.cs
--- a/Engine/Source/Programs/UnrealBuildTool/Configuration/UEBuildPlatform.cs
+++ b/Engine/Source/Programs/UnrealBuildTool/Configuration/UEBuildPlatform.cs
@@ -15,11 +15,24 @@ namespace UnrealBuildTool
 		// a mapping of a group to the platforms in the group (ie, Microsoft contains Win32 and Win64)
 		static Dictionary<UnrealPlatformGroup, List<UnrealTargetPlatform>> PlatformGroupDictionary = new Dictionary<UnrealPlatformGroup, List<UnrealTargetPlatform>>();
 
-		/**
-		 * Attempt to convert a string to an UnrealTargetPlatform enum eentry
-		 * 
-		 * @return UnrealTargetPlatform.Unknown on failure (the platform didn't match the enum)
-		 */
+        /** Name of the file that holds currently install SDK version string */
+        static string CurrentlyInstalledSDKStringManifest = "CurrentlyInstalled.txt";
+
+        /** Name of the file that holds environment variables of current SDK */
+        static string SDKEnvironmentVarsFile = "OutputEnvVars.txt";
+
+        /** Enum describing types of hooks a platform SDK can have */
+        public enum SDKHookType
+        {
+            Install,
+            Uninstall
+        };
+
+        /**
+         * Attempt to convert a string to an UnrealTargetPlatform enum entry
+         * 
+         * @return UnrealTargetPlatform.Unknown on failure (the platform didn't match the enum)
+         */
 		public static UnrealTargetPlatform ConvertStringToPlatform(string InPlatformName)
 		{
 			// special case x64, to not break anything
@@ -163,11 +176,247 @@ namespace UnrealBuildTool
 			}
 		}
 
-		/**
-		 *	Whether the required external SDKs are installed for this platform
-		 */
+        /** 
+         * Whether platform supports switching SDKs during runtime
+         * 
+         * @return true if supports
+         */
+        public virtual bool PlatformSupportsSDKSwitching()
+        {
+            return false;
+        }
+
+        /** 
+         * Returns platform-specific name used in SDK repository
+         * 
+         * @return path to SDK Repository
+         */
+        public virtual string GetSDKTargetPlatformName()
+        {
+            return "";
+        }
+
+        /** 
+         * Returns SDK string as required by the platform 
+         * 
+         * @return Valid SDK string
+         */
+        public virtual string GetRequiredSDKString()
+        {
+            return "";
+        }
+
+        /** 
+         * Returns path to platform SDKs
+         * 
+         * @return Valid SDK string
+         */
+        public virtual string GetPathToPlatformSDKs()
+        {
+            string SDKRoot = Environment.GetEnvironmentVariable("UE_SDKS_ROOT");
+            if (SDKRoot != null && SDKRoot != "")
+            {
+                return Path.Combine(SDKRoot, "Host" + ExternalExecution.GetRuntimePlatform(), GetSDKTargetPlatformName());
+            }
+
+            return "";
+        }
+
+        /**
+         * Gets currently installed version
+         * 
+         * @param PlatformSDKRoot absolute path to platform SDK root
+         * @param OutInstalledSDKVersionString version string as currently installed
+         * 
+         * @return true if was able to read it
+         */
+        public virtual bool GetCurrentlyInstalledSDKString(string PlatformSDKRoot, out string OutInstalledSDKVersionString)
+        {
+            if (Directory.Exists(PlatformSDKRoot))
+            {
+                string VersionFilename = Path.Combine(PlatformSDKRoot, CurrentlyInstalledSDKStringManifest);
+                if (File.Exists(VersionFilename))
+                {
+                    using (StreamReader Reader = new StreamReader(VersionFilename))
+                    {
+                        string Version = Reader.ReadLine();
+                        if (Version != null)
+                        {
+                            OutInstalledSDKVersionString = Version;
+                            return true;
+                        }
+                    }
+                }
+            }
+
+            OutInstalledSDKVersionString = "";
+            return false;
+        }
+
+        /**
+         * Sets currently installed version
+         * 
+         * @param PlatformSDKRoot absolute path to platform SDK root
+         * @param InstalledSDKVersionString SDK version string to set
+         * 
+         * @return true if was able to set it
+         */
+        public virtual bool SetCurrentlyInstalledSDKString(string PlatformSDKRoot, string InstalledSDKVersionString)
+        {
+            if (Directory.Exists(PlatformSDKRoot))
+            {
+                string VersionFilename = Path.Combine(PlatformSDKRoot, CurrentlyInstalledSDKStringManifest);
+                if (File.Exists(VersionFilename))
+                {
+                    File.Delete(VersionFilename);
+                }
+
+                using (StreamWriter Writer = File.CreateText(VersionFilename))
+                {
+                    Writer.WriteLine(InstalledSDKVersionString);
+                    return true;
+                }
+            }
+
+            return false;
+        }
+
+        /**
+         * Returns Hook names as needed by the platform
+         * (e.g. can be overriden with custom executables or scripts)
+         *
+         * @param Hook Hook type
+         */
+        public virtual string GetHookExecutableName(SDKHookType Hook)
+        {
+            if (Hook == SDKHookType.Uninstall)
+            {
+                return "unsetup.bat";
+            }
+
+            return "setup.bat";
+        }
+
+        /**
+         * Runs install/uninstall hooks for SDK
+         * 
+         * @param PlatformSDKRoot absolute path to platform SDK root
+         * @param SDKVersionString version string to run for (can be empty!)
+         * @param Hook which one of hooks to run
+         * @param bHookCanBeNonExistent whether a non-existing hook means failure
+         * 
+         * @return true if succeeded
+         */
+        public virtual bool RunSDKHooks(string PlatformSDKRoot, string SDKVersionString, SDKHookType Hook, bool bHookCanBeNonExistent = true)
+        {
+            if (SDKVersionString != "")
+            {
+                string SDKDirectory = Path.Combine(PlatformSDKRoot, SDKVersionString);
+                string HookExe = Path.Combine(SDKDirectory, GetHookExecutableName(Hook));
+
+                if (File.Exists(HookExe))
+                {
+                    Console.WriteLine("Running {0} hook {1}", Hook, HookExe);
+
+                    // run it
+                    Process HookProcess = new Process();
+                    HookProcess.StartInfo.WorkingDirectory = SDKDirectory;
+                    HookProcess.StartInfo.FileName = HookExe;
+                    HookProcess.StartInfo.Arguments = "";
+
+                    HookProcess.Start();
+                    HookProcess.WaitForExit();
+
+                    if (HookProcess.ExitCode != 0)
+                    {
+                        Console.WriteLine("Hook exited uncleanly (returned {0}), considering it failed.", HookProcess.ExitCode);
+                        return false;
+                    }
+
+                    return true;
+                }
+            }
+
+            return bHookCanBeNonExistent;
+        }
+
+        /**
+         * Loads environment variables from SDK
+         * 
+         * @param PlatformSDKRoot absolute path to platform SDK
+         * @param SDKVersionString SDK version string (cannot be empty!)
+         * 
+         * @return true if succeeded
+         */
+        public virtual bool SetupEnvironmentFromSDK(string PlatformSDKRoot, string SDKVersionString)
+        {
+            string EnvVarFile = Path.Combine(PlatformSDKRoot, SDKVersionString, SDKEnvironmentVarsFile);
+            if (File.Exists(EnvVarFile))
+            {
+                using (StreamReader Reader = new StreamReader(EnvVarFile))
+                {
+                    for (; ; )
+                    {
+                        string VariableString = Reader.ReadLine();
+                        if (VariableString == null)
+                        {
+                            return true;
+                        }
+
+                        string[] Parts = VariableString.Split('=');
+                        if (Parts.Length != 2)
+                        {
+                            Console.WriteLine("Incorrect environment variable declaration:");
+                            Console.WriteLine(VariableString);
+                            return false;
+                        }
+
+                        Console.WriteLine("Setting variable '{0}' to '{1}'", Parts[0], Parts[1]);
+                        Environment.SetEnvironmentVariable(Parts[0], Parts[1]);
+                    }
+                }
+            }
+
+            return false;
+        }
+
+        /**
+         *	Whether the required external SDKs are installed for this platform
+         */
 		public virtual bool HasRequiredSDKsInstalled()
 		{
+            if (PlatformSupportsSDKSwitching())
+            {
+                // if we don't have a path to platform SDK's, take no action
+                string PlatformSDKRoot = GetPathToPlatformSDKs();
+                if (PlatformSDKRoot != "")
+                {
+                    string CurrentSDKString;
+                    if (!GetCurrentlyInstalledSDKString(PlatformSDKRoot, out CurrentSDKString) || CurrentSDKString != GetRequiredSDKString())
+                    {
+                        // switch over (note that version string can be empty)
+                        if (!RunSDKHooks(PlatformSDKRoot, CurrentSDKString, SDKHookType.Uninstall))
+                        {
+                            throw new BuildException("Failed to uninstall currently installed SDK {0}", CurrentSDKString);
+                        }
+                        // delete Manifest file to avoid multiple uninstalls
+                        SetCurrentlyInstalledSDKString(PlatformSDKRoot, "");
+
+                        if (!RunSDKHooks(PlatformSDKRoot, GetRequiredSDKString(), SDKHookType.Install, false))
+                        {
+                            throw new BuildException("Failed to install required SDK {0}", GetRequiredSDKString());
+                        }
+                        SetCurrentlyInstalledSDKString(PlatformSDKRoot, GetRequiredSDKString());
+                    }
+
+                    // load environment variables from current SDK
+                    if (!SetupEnvironmentFromSDK(PlatformSDKRoot, GetRequiredSDKString()))
+                    {
+                        throw new BuildException("Failed to load environment from required SDK {0}", GetRequiredSDKString());
+                    }
+                }
+            }
+
 			return true;
 		}
 
diff --git a/Engine/Source/Programs/UnrealBuildTool/Linux/UEBuildLinux.cs b/Engine/Source/Programs/UnrealBuildTool/Linux/UEBuildLinux.cs
--- a/Engine/Source/Programs/UnrealBuildTool/Linux/UEBuildLinux.cs
+++ b/Engine/Source/Programs/UnrealBuildTool/Linux/UEBuildLinux.cs
@@ -10,6 +10,42 @@ namespace UnrealBuildTool.Linux
 {
     class LinuxPlatform : UEBuildPlatform
     {
+        /** This is the SDK version we support */
+        static private string ExpectedSDKVersion = "v3_clang-3.3_ld-2.24_glibc-2.12.2";
+
+        /** Platform name (embeds architecture for now) */
+        static private string TargetPlatformName = "Linux_x64";
+
+        /** 
+         * Whether platform supports switching SDKs during runtime
+         * 
+         * @return true if supports
+         */
+        public override bool PlatformSupportsSDKSwitching()
+        {
+            return true;
+        }
+
+        /** 
+         * Returns platform-specific name used in SDK repository
+         * 
+         * @return path to SDK Repository
+         */
+        public override string GetSDKTargetPlatformName()
+        {
+            return TargetPlatformName;
+        }
+
+        /** 
+         * Returns SDK string as required by the platform 
+         * 
+         * @return Valid SDK string
+         */
+        public override string GetRequiredSDKString()
+        {
+            return ExpectedSDKVersion;
+        }
+
         /**
          *	Whether the required external SDKs are installed for this platform
          */
@@ -20,6 +56,12 @@ namespace UnrealBuildTool.Linux
                 return true;
             }
 
+            // attempt to switch SDKs
+            if (!base.HasRequiredSDKsInstalled())
+            {
+                return false;
+            }
+
             string BaseLinuxPath = Environment.GetEnvironmentVariable("LINUX_ROOT");
 
             // we don't have an LINUX_ROOT specified
diff --git a/Engine/Source/Programs/UnrealBuildTool/System/UnrealBuildTool.cs b/Engine/Source/Programs/UnrealBuildTool/System/UnrealBuildTool.cs
--- a/Engine/Source/Programs/UnrealBuildTool/System/UnrealBuildTool.cs
+++ b/Engine/Source/Programs/UnrealBuildTool/System/UnrealBuildTool.cs
@@ -374,6 +374,20 @@ namespace UnrealBuildTool
 
 				List<System.Type> ProjectGeneratorList = new List<System.Type>();
 				var AllTypes = UBTAssembly.GetTypes();
+                // register all build platforms first, since they implement SDK-switching logic that can set environment variables
+                foreach (var CheckType in AllTypes)
+                {
+                    if (CheckType.IsClass && !CheckType.IsAbstract)
+                    {
+                        if (CheckType.IsSubclassOf(typeof(UEBuildPlatform)))
+                        {
+                            Log.TraceVerbose("    Registering build platform: {0}", CheckType.ToString());
+                            UEBuildPlatform TempInst = (UEBuildPlatform)(UBTAssembly.CreateInstance(CheckType.FullName, true));
+                            TempInst.RegisterBuildPlatform();
+                        }
+                    }
+                }
+                // register all the other classes
 				foreach (var CheckType in AllTypes)
 				{
 					if (CheckType.IsClass && !CheckType.IsAbstract)
@@ -384,12 +398,6 @@ namespace UnrealBuildTool
 							UEToolChain TempInst = (UEToolChain)(UBTAssembly.CreateInstance(CheckType.FullName, true));
 							TempInst.RegisterToolChain();
 						}
-						else if (CheckType.IsSubclassOf(typeof(UEBuildPlatform)))
-						{
-							Log.TraceVerbose("    Registering build platform: {0}", CheckType.ToString());
-							UEBuildPlatform TempInst = (UEBuildPlatform)(UBTAssembly.CreateInstance(CheckType.FullName, true));
-							TempInst.RegisterBuildPlatform();
-						}
 						else if (CheckType.IsSubclassOf(typeof(UEBuildDeploy)))
 						{
 							Log.TraceVerbose("    Registering build deploy: {0}", CheckType.ToString());
