#UE4 - Fix for UBT throwing exceptions when passing -define arguments with quotation marks.

diff --git a/Engine/Source/Programs/UnrealBuildTool/System/UnrealBuildTool.cs b/Engine/Source/Programs/UnrealBuildTool/System/UnrealBuildTool.cs
--- a/Engine/Source/Programs/UnrealBuildTool/System/UnrealBuildTool.cs
+++ b/Engine/Source/Programs/UnrealBuildTool/System/UnrealBuildTool.cs
@@ -760,6 +760,10 @@ namespace UnrealBuildTool
 						{
 							UEBuildConfiguration.bCompileSpeedTree = false;
 						}
+						else if (LowercaseArg == "-define")
+						{
+							// Skip -define
+						}
 						else if (CheckPlatform.ToString().ToLowerInvariant() == LowercaseArg)
 						{
 							// It's the platform set...
@@ -784,7 +788,8 @@ namespace UnrealBuildTool
 							{
 								// This arg may be a game name. Check for the existence of a game folder with this name.
 								// "Engine" is not a valid game name.
-								if (LowercaseArg != "engine" && Directory.Exists(Path.Combine(ProjectFileGenerator.RootRelativePath, Arg, "Config")))
+								if (LowercaseArg != "engine" && Arg.IndexOfAny(Path.GetInvalidPathChars()) == -1 &&
+									Directory.Exists(Path.Combine(ProjectFileGenerator.RootRelativePath, Arg, "Config")))
 								{
 									GameName = Arg;
 									Log.TraceVerbose("CommandLine: Found game name '{0}'", GameName);
