#UE4 - UBT: Escape custom definitions containing quotation marks for the compiler commandline.

diff --git a/Engine/Source/Programs/UnrealBuildTool/Windows/VCToolChain.cs b/Engine/Source/Programs/UnrealBuildTool/Windows/VCToolChain.cs
--- a/Engine/Source/Programs/UnrealBuildTool/Windows/VCToolChain.cs
+++ b/Engine/Source/Programs/UnrealBuildTool/Windows/VCToolChain.cs
@@ -690,7 +690,9 @@ namespace UnrealBuildTool
 			// Add preprocessor definitions to the argument list.
 			foreach (string Definition in CompileEnvironment.Config.Definitions)
 			{
-				Arguments += string.Format(" /D \"{0}\"", Definition);
+				// Escape all quotation marks so that they get properly passed with the command line.
+				var DefinitionArgument = Definition.Contains("\"") ? Definition.Replace("\"", "\\\"") : Definition;
+				Arguments += string.Format(" /D \"{0}\"", DefinitionArgument);
 			}
 
 			// Create a compile action for each source file.
