Print out environment changes only when displaying debug info (too much spam during project generation).

diff --git a/Engine/Source/Programs/UnrealBuildTool/Configuration/UEBuildPlatform.cs b/Engine/Source/Programs/UnrealBuildTool/Configuration/UEBuildPlatform.cs
--- a/Engine/Source/Programs/UnrealBuildTool/Configuration/UEBuildPlatform.cs
+++ b/Engine/Source/Programs/UnrealBuildTool/Configuration/UEBuildPlatform.cs
@@ -371,7 +371,10 @@ namespace UnrealBuildTool
                             return false;
                         }
 
-                        Console.WriteLine("Setting variable '{0}' to '{1}'", Parts[0], Parts[1]);
+                        if (BuildConfiguration.bPrintDebugInfo)
+                        {
+                            Console.WriteLine("Setting variable '{0}' to '{1}'", Parts[0], Parts[1]);
+                        }
                         Environment.SetEnvironmentVariable(Parts[0], Parts[1]);
                     }
                 }
