#ue4 Add a platform packaging data file so PS4 can package content only projects.
#ttp 331317

diff --git a/Engine/Source/Programs/UnrealBuildTool/System/RulesCompiler.cs b/Engine/Source/Programs/UnrealBuildTool/System/RulesCompiler.cs
--- a/Engine/Source/Programs/UnrealBuildTool/System/RulesCompiler.cs
+++ b/Engine/Source/Programs/UnrealBuildTool/System/RulesCompiler.cs
@@ -372,15 +372,7 @@ namespace UnrealBuildTool
 		/// <summary>
 		/// Type of target
 		/// </summary>
-		public TargetType Type = TargetType.Game;
-
-		/// <summary>
-		/// Gets the TitleID for the target.  Mostly for consoles.
-		/// </summary>		
-		public virtual String GetTitleID(UnrealTargetPlatform Platform)
-		{
-			return "DefaultTitleID";
-		}
+		public TargetType Type = TargetType.Game;		
 
 		/// <summary>
 		/// The name of this target's 'configuration' within the development IDE.  No project may have more than one target with the same configuration name.
@@ -711,46 +703,6 @@ namespace UnrealBuildTool
             }
             return Result;
         }
-
-		/// <summary>
-		/// Various metadata properties about the target (like some per-platform IDs, etc). It's a map of a string to an object (string, array, etc), and
-		/// it's up to the user of the metadata to get the right type
-		/// </summary>
-		public Dictionary<string, object> TargetMetadata = new Dictionary<string,object>();
-
-		/// <summary>
-		/// Return the raw object for the given key, or null
-		/// </summary>
-		public object GetMetadata(string Key)
-		{
-			if (TargetMetadata.ContainsKey(Key))
-			{
-				return TargetMetadata[Key];
-			}
-			return null;
-		}
-
-		/// <summary>
-		/// Some helper functions to cast the object to a specific type
-		/// </summary>
-		public string GetMetadataAsString(string Key)
-		{
-			object Value = GetMetadata(Key);
-			if (Value != null && Value is string)
-			{
-				return Value as string;
-			}
-			return null;
-		}
-		public string[] GetMetadataAsStringArray(string Key)
-		{
-			object Value = GetMetadata(Key);
-			if (Value != null && Value is string[])
-			{
-				return Value as string[];
-			}
-			return null;
-		}
 	}
 
 
