#UE4 - UBT will no longer attempt to compile binary-only plugins

diff --git a/Engine/Source/Programs/UnrealBuildTool/Configuration/UEBuildBinary.cs b/Engine/Source/Programs/UnrealBuildTool/Configuration/UEBuildBinary.cs
--- a/Engine/Source/Programs/UnrealBuildTool/Configuration/UEBuildBinary.cs
+++ b/Engine/Source/Programs/UnrealBuildTool/Configuration/UEBuildBinary.cs
@@ -67,6 +67,11 @@ namespace UnrealBuildTool
 		/// </summary>
 		public bool bAllowCompilation = true;
 
+		/// <summary>
+		/// True if this binary has any Build.cs files, if not this is probably a binary-only plugins
+		/// </summary>
+		public bool bHasModuleRules = true;
+
         /// <summary>
         /// For most binaries, this is false. If this is a cross-platform binary build for a specific platform (for example XB1 DLL for a windows editor) this will be true.
         /// </summary>
@@ -115,6 +120,7 @@ namespace UnrealBuildTool
 				bool bInCreateImportLibrarySeparately = false,
 				bool bInIncludeDependentLibrariesInLibrary = false,
 				bool bInAllowCompilation = true,
+				bool bInHasModuleRules = true,
                 bool bInIsCrossTarget = false,
                 bool bInCompileMonolithic = false,
 				UnrealTargetConfiguration InTargetConfiguration = UnrealTargetConfiguration.Development,
@@ -130,6 +136,7 @@ namespace UnrealBuildTool
 			bCreateImportLibrarySeparately = bInCreateImportLibrarySeparately;
 			bIncludeDependentLibrariesInLibrary = bInIncludeDependentLibrariesInLibrary;
 			bAllowCompilation = bInAllowCompilation;
+			bHasModuleRules = bInHasModuleRules;
             bIsCrossTarget = bInIsCrossTarget;
             bCompileMonolithic = bInCompileMonolithic;
 			TargetConfiguration = InTargetConfiguration;
@@ -314,13 +321,17 @@ namespace UnrealBuildTool
 		{
 			foreach(var ModuleName in ModuleNames)
 			{
-				var Module = Target.FindOrCreateModuleByName(ModuleName);
-				if(Module.Binary != null)
+				UEBuildModule Module = null;
+				if (Config.bHasModuleRules)
 				{
-					throw new BuildException("Module \"{0}\" linked into both {1} and {2}, which creates ambiguous linkage for dependents.", ModuleName, Module.Binary.Config.OutputFilePath, Config.OutputFilePath);
+					Module = Target.FindOrCreateModuleByName(ModuleName);
+					if (Module.Binary != null)
+					{
+						throw new BuildException("Module \"{0}\" linked into both {1} and {2}, which creates ambiguous linkage for dependents.", ModuleName, Module.Binary.Config.OutputFilePath, Config.OutputFilePath);
+					}
+					Module.Binary = this;
+					Module.bIncludedInTarget = true;
 				}
-				Module.Binary = this;
-				Module.bIncludedInTarget = true;
 
 				// We set whether the binary is being compiled monolithic here to know later - specifically
 				// when we are determining whether to use SharedPCHs or not for static lib builds of plugins.
@@ -332,7 +343,7 @@ namespace UnrealBuildTool
 				Config.TargetConfiguration = Target.Configuration;
 				Config.TargetName = Target.GetAppName();
 
-				if ((Target.Rules == null) || (Target.Rules.bOutputToEngineBinaries == false))
+				if (Module != null && (Target.Rules == null || Target.Rules.bOutputToEngineBinaries == false))
 				{
 					// Fix up the binary path if this is module specifies an alternate output directory
 					Config.OutputFilePath = Module.FixupOutputPath(Config.OutputFilePath);
@@ -377,10 +388,21 @@ namespace UnrealBuildTool
 		public override List<UEBuildBinary> ProcessUnboundModules()
 		{
 			var Binaries = new Dictionary<string, UEBuildBinary>( StringComparer.InvariantCultureIgnoreCase );
-			foreach( var ModuleName in ModuleNames )
+			if (Config.bHasModuleRules)
 			{
-				var Module = Target.FindOrCreateModuleByName(ModuleName);
-				Module.RecursivelyProcessUnboundModules( Target, ref Binaries, this );
+				foreach (var ModuleName in ModuleNames)
+				{
+					var Module = Target.FindOrCreateModuleByName(ModuleName);
+					Module.RecursivelyProcessUnboundModules(Target, ref Binaries, this);
+				}
+			}
+			else
+			{
+				// There's only one module in this case, so just bind it to this binary
+				foreach (var ModuleName in ModuleNames)
+				{
+					Binaries.Add(ModuleName, this);
+				}
 			}
 
 			// Now build a final list of newly-created binaries that were bound to.  The hash may contain duplicates, so
diff --git a/Engine/Source/Programs/UnrealBuildTool/Configuration/UEBuildTarget.cs b/Engine/Source/Programs/UnrealBuildTool/Configuration/UEBuildTarget.cs
--- a/Engine/Source/Programs/UnrealBuildTool/Configuration/UEBuildTarget.cs
+++ b/Engine/Source/Programs/UnrealBuildTool/Configuration/UEBuildTarget.cs
@@ -1719,9 +1719,12 @@ namespace UnrealBuildTool
 				OutputFilePath = MakeBinaryPath(Module.Name, GetAppName() + "-" + Module.Name, BinaryType, Rules.Type, bIsRocketModule, Plugin, AppName);
 			}
 
+			// Try to determine if we have the rules file
+			var ModuleFilename = RulesCompiler.GetModuleFilename(Module.Name);
+			var bHasModuleRules = String.IsNullOrEmpty(ModuleFilename) == false;
 			// Figure out whether we should build it from source
-			string ModuleSourceFolder = Path.GetDirectoryName(RulesCompiler.GetModuleFilename(Module.Name));
-			bool bShouldBeBuiltFromSource = Directory.GetFiles(ModuleSourceFolder, "*.cpp", SearchOption.AllDirectories).Length > 0;
+			var ModuleSourceFolder = bHasModuleRules ? Path.GetDirectoryName(RulesCompiler.GetModuleFilename(Module.Name)) : ModuleFilename;
+			bool bShouldBeBuiltFromSource = bHasModuleRules && Directory.GetFiles(ModuleSourceFolder, "*.cpp", SearchOption.AllDirectories).Length > 0;
 
 			// Create the binary
 			UEBuildBinaryConfiguration Config = new UEBuildBinaryConfiguration( InType: BinaryType, 
@@ -1729,6 +1732,7 @@ namespace UnrealBuildTool
 																				InIntermediateDirectory: IntermediateDirectory,
 																				bInAllowExports: true,
 																				bInAllowCompilation: bShouldBeBuiltFromSource,
+																				bInHasModuleRules : bHasModuleRules,
 																				InModuleNames: new List<string> { Module.Name } );
 			AppBinaries.Add(new UEBuildBinaryCPP(this, Config));
             return SpecialRocketLibFilesThatAreBuildProducts;
@@ -1809,11 +1813,15 @@ namespace UnrealBuildTool
 			if ((BinaryType == UEBuildBinaryType.DynamicLinkLibrary) && (string.IsNullOrEmpty(ModuleName) == false))
 			{
 				// Allow for modules to specify sub-folders in the Binaries folder
-				string IgnoredRulesFilename;
-				ModuleRules ModuleRulesObj = RulesCompiler.CreateModuleRules(ModuleName, new TargetInfo(Platform, Configuration, TargetType), out IgnoredRulesFilename);
-				if (ModuleRulesObj != null)
+				var RulesFilename = RulesCompiler.GetModuleFilename(ModuleName);
+				// Plugins can be binary-only and can have no rules object
+				if (PluginInfo == null || !String.IsNullOrEmpty(RulesFilename))
 				{
-					ModuleBinariesSubDir = ModuleRulesObj.BinariesSubFolder;
+					ModuleRules ModuleRulesObj = RulesCompiler.CreateModuleRules(ModuleName, new TargetInfo(Platform, Configuration, TargetType), out RulesFilename);
+					if (ModuleRulesObj != null)
+					{
+						ModuleBinariesSubDir = ModuleRulesObj.BinariesSubFolder;
+					}
 				}
 			}
 
