#ue4
#ios
* Mac builds will now also create the stub file for PC

TTP331698

diff --git a/Engine/Source/Programs/UnrealBuildTool/IOS/UEDeployIOS.cs b/Engine/Source/Programs/UnrealBuildTool/IOS/UEDeployIOS.cs
--- a/Engine/Source/Programs/UnrealBuildTool/IOS/UEDeployIOS.cs
+++ b/Engine/Source/Programs/UnrealBuildTool/IOS/UEDeployIOS.cs
@@ -6,6 +6,9 @@ using System.Linq;
 using System.Text;
 using System.Xml;
 using System.IO;
+using System.Diagnostics;
+using Ionic.Zip;
+using Ionic.Zlib;
 
 namespace UnrealBuildTool.IOS
 {
@@ -65,6 +68,32 @@ namespace UnrealBuildTool.IOS
 			}
 		}
 
+		/// <summary>
+		/// Process.OutputDataReceived event handler.
+		/// </summary>
+		/// <param name="sender">Sender</param>
+		/// <param name="e">Event args</param>
+		private void StdOut(object sender, DataReceivedEventArgs e)
+		{
+			if (e.Data != null)
+			{
+				Console.WriteLine (e.Data);
+			}
+		}
+
+		/// <summary>
+		/// Process.ErrorDataReceived event handler.
+		/// </summary>
+		/// <param name="sender">Sender</param>
+		/// <param name="e">Event args</param>
+		private void StdErr(object sender, DataReceivedEventArgs e)
+		{
+			if (e.Data != null)
+			{
+				Console.WriteLine ("Error: " + e.Data);
+			}
+		}
+
 		public override bool PrepForUATPackageOrDeploy(string InProjectName, string InProjectDirectory, string InExecutablePath, string InEngineDir, bool bForDistribution)
 		{
 			bool bIsUE4Game = InExecutablePath.Contains ("UE4Game");
@@ -122,6 +151,11 @@ namespace UnrealBuildTool.IOS
 				File.Copy (ProvisionWithPrefix, Environment.GetEnvironmentVariable ("HOME") + "/Library/MobileDevice/Provisioning Profiles/" + InProjectName + ".mobileprovision", true);
 				FileInfo DestFileInfo = new FileInfo (Environment.GetEnvironmentVariable ("HOME") + "/Library/MobileDevice/Provisioning Profiles/" + InProjectName + ".mobileprovision");
 				DestFileInfo.Attributes = DestFileInfo.Attributes & ~FileAttributes.ReadOnly;
+
+				if (BuildConfiguration.bCreateStubIPA)
+				{
+					File.Copy (ProvisionWithPrefix, AppDirectory + "/embedded.mobileprovision");
+				}
 			}
 
 			// install the distribution provision
@@ -144,12 +178,82 @@ namespace UnrealBuildTool.IOS
 				PListFile = BuildDirectory + "/" + InProjectName + "-Info.plist";
 			}
 
+			string Entitlements = InEngineDir + "/Build/IOS/UE4Game.entitlements";
+			if (File.Exists(BuildDirectory + "/" + InProjectName + ".entitlements"))
+			{
+				Entitlements = BuildDirectory + "/" + InProjectName + ".entitlements";
+			}
+
 			Dictionary<string, string> Replacements = new Dictionary<string, string>();
 			Replacements.Add("${EXECUTABLE_NAME}", GameName);
 			Replacements.Add("${BUNDLE_IDENTIFIER}", InProjectName.Replace("_", ""));
 			CopyFileWithReplacements(PListFile, AppDirectory + "/Info.plist", Replacements);
 			CopyFileWithReplacements(PListFile, IntermediateDirectory + "/" + GameName + "-Info.plist", Replacements);
 
+			// ensure the destination is writable
+			if (File.Exists(AppDirectory + "/" + GameName))
+			{
+				FileInfo GameFileInfo = new FileInfo(AppDirectory + "/" + GameName);
+				GameFileInfo.Attributes = GameFileInfo.Attributes & ~FileAttributes.ReadOnly;
+			}
+
+			// copy the GameName binary
+			File.Copy(BinaryPath + "/" + GameExeName, AppDirectory + "/" + GameName, true);
+
+			if (BuildConfiguration.bCreateStubIPA)
+			{
+				// copy the custom rules
+				File.Copy (InEngineDir + "/Build/IOS/XcodeSupportFiles/CustomResourceRules.plist", AppDirectory + "/CustomResourceRules.plist", true);
+
+				// create the PkgInfo file
+
+				// code sign
+				string Arguments = "CODESIGN_ALLOCATE=/Applications/Xcode.app/Contents/Developer/Toolchains/XcodeDefault.xctoolchain/usr/bin/codesign_allocate";
+				Arguments += " /usr/bin/codesign --force --sign \"iPhone Developer\"";
+				Arguments += " --resource-rules=\"";
+				Arguments += AppDirectory + "/CustomResourceRules.plist\"";
+				Arguments += " --entitlements \"" + Entitlements + "\" ";
+				Arguments += "\"" + AppDirectory + "\"";
+				using(var CodeSignProcess = new Process())
+				{
+					CodeSignProcess.StartInfo.FileName = "/usr/bin/env";
+					CodeSignProcess.StartInfo.Arguments = Arguments;
+					CodeSignProcess.StartInfo.CreateNoWindow = true;
+					CodeSignProcess.StartInfo.UseShellExecute = false;
+					CodeSignProcess.StartInfo.RedirectStandardOutput = true;
+					CodeSignProcess.OutputDataReceived += StdOut;
+					CodeSignProcess.StartInfo.RedirectStandardError = true;
+					CodeSignProcess.ErrorDataReceived += StdErr;
+					CodeSignProcess.StartInfo.RedirectStandardInput = false;
+					CodeSignProcess.StartInfo.UseShellExecute = false;
+
+					Log.TraceVerbose( "Launching {0} to harvest Visual Studio environment settings...", CodeSignProcess.StartInfo.FileName );
+
+					// Try to launch the process, and produce a friendly error message if it fails.
+					try
+					{
+						// Start the process up and then wait for it to finish
+						CodeSignProcess.Start();
+						CodeSignProcess.BeginOutputReadLine();
+						CodeSignProcess.BeginErrorReadLine();
+						while (!CodeSignProcess.HasExited)
+						{
+							System.Threading.Thread.Sleep(5);
+						}
+						CodeSignProcess.WaitForExit();
+					}
+					catch(Exception ex)
+					{
+						throw new BuildException( ex, "Failed to start local process for action (\"{0}\"): {1} {2}", ex.Message, CodeSignProcess.StartInfo.FileName, CodeSignProcess.StartInfo.Arguments );
+					}
+
+					Log.TraceVerbose( "Finished launching {0}.", CodeSignProcess.StartInfo.FileName );
+				}
+
+				// zip it up into the stub file
+				PackageStub (BinaryPath, GameName);
+			}
+
 			// copy engine assets in
 			CopyFiles(InEngineDir + "/Build/IOS/Resources/Graphics", AppDirectory, "*.png", true);
 			// merge game assets on top
@@ -161,17 +265,121 @@ namespace UnrealBuildTool.IOS
 			//CopyFiles(BuildDirectory, PayloadDirectory, null, "iTunesArtwork", null);
 			CopyFolder(InEngineDir + "/Content/Stats", AppDirectory + "/cookeddata/engine/content/stats", true);
 
-			// ensure the destination is writable
-			if (File.Exists(AppDirectory + "/" + GameName))
+			return true;
+		}
+
+		private void PackageStub(string BinaryPath, string GameName)
+		{
+			// create the ipa
+			string IPAName = BinaryPath + "/" + GameName + ".stub";
+			// delete the old one
+			if (File.Exists(IPAName))
 			{
-				FileInfo GameFileInfo = new FileInfo(AppDirectory + "/" + GameName);
-				GameFileInfo.Attributes = GameFileInfo.Attributes & ~FileAttributes.ReadOnly;
+				File.Delete(IPAName);
 			}
 
-			// copy the GameName binary
-			File.Copy(BinaryPath + "/" + GameExeName, AppDirectory + "/" + GameName, true);
+			// make the subdirectory if needed
+			string DestSubdir = Path.GetDirectoryName(IPAName);
+			if (!Directory.Exists(DestSubdir))
+			{
+				Directory.CreateDirectory(DestSubdir);
+			}
 
-			return true;
+			// set up the directories
+			string ZipWorkingDir = String.Format("Payload/{0}.app/", GameName);
+			string ZipSourceDir = string.Format("{0}/Payload/{1}.app", BinaryPath, GameName);
+
+			// create the file
+			using (ZipFile Zip = new ZipFile())
+			{
+				// add the entire directory
+				Zip.AddDirectory(ZipSourceDir, ZipWorkingDir);
+
+				// Update permissions to be UNIX-style
+				// Modify the file attributes of any added file to unix format
+				foreach (ZipEntry E in Zip.Entries)
+				{
+					const byte FileAttributePlatform_NTFS = 0x0A;
+					const byte FileAttributePlatform_UNIX = 0x03;
+					const byte FileAttributePlatform_FAT = 0x00;
+
+					const int UNIX_FILETYPE_NORMAL_FILE = 0x8000;
+					//const int UNIX_FILETYPE_SOCKET = 0xC000;
+					//const int UNIX_FILETYPE_SYMLINK = 0xA000;
+					//const int UNIX_FILETYPE_BLOCKSPECIAL = 0x6000;
+					const int UNIX_FILETYPE_DIRECTORY = 0x4000;
+					//const int UNIX_FILETYPE_CHARSPECIAL = 0x2000;
+					//const int UNIX_FILETYPE_FIFO = 0x1000;
+
+					const int UNIX_EXEC = 1;
+					const int UNIX_WRITE = 2;
+					const int UNIX_READ = 4;
+
+
+					int MyPermissions = UNIX_READ | UNIX_WRITE;
+					int OtherPermissions = UNIX_READ;
+
+					int PlatformEncodedBy = (E.VersionMadeBy >> 8) & 0xFF;
+					int LowerBits = 0;
+
+					// Try to preserve read-only if it was set
+					bool bIsDirectory = E.IsDirectory;
+
+					// Check to see if this 
+					bool bIsExecutable = false;
+					if (Path.GetFileNameWithoutExtension(E.FileName).Equals(GameName, StringComparison.InvariantCultureIgnoreCase))
+					{
+						bIsExecutable = true;
+					}
+
+					if (bIsExecutable)
+					{
+						// The executable will be encrypted in the final distribution IPA and will compress very poorly, so keeping it
+						// uncompressed gives a better indicator of IPA size for our distro builds
+						E.CompressionLevel = CompressionLevel.None;
+					}
+
+					if ((PlatformEncodedBy == FileAttributePlatform_NTFS) || (PlatformEncodedBy == FileAttributePlatform_FAT))
+					{
+						FileAttributes OldAttributes = E.Attributes;
+						//LowerBits = ((int)E.Attributes) & 0xFFFF;
+
+						if ((OldAttributes & FileAttributes.Directory) != 0)
+						{
+							bIsDirectory = true;
+						}
+
+						// Permissions
+						if ((OldAttributes & FileAttributes.ReadOnly) != 0)
+						{
+							MyPermissions &= ~UNIX_WRITE;
+							OtherPermissions &= ~UNIX_WRITE;
+						}
+					}
+
+					if (bIsDirectory || bIsExecutable)
+					{
+						MyPermissions |= UNIX_EXEC;
+						OtherPermissions |= UNIX_EXEC;
+					}
+
+					// Re-jigger the external file attributes to UNIX style if they're not already that way
+					if (PlatformEncodedBy != FileAttributePlatform_UNIX)
+					{
+						int NewAttributes = bIsDirectory ? UNIX_FILETYPE_DIRECTORY : UNIX_FILETYPE_NORMAL_FILE;
+
+						NewAttributes |= (MyPermissions << 6);
+						NewAttributes |= (OtherPermissions << 3);
+						NewAttributes |= (OtherPermissions << 0);
+
+						// Now modify the properties
+						E.AdjustExternalFileAttributes(FileAttributePlatform_UNIX, (NewAttributes << 16) | LowerBits);
+					}
+				}
+
+				// Save it out
+				Zip.Save(IPAName);
+			}
 		}
 
 		public override bool PrepTargetForDeployment(UEBuildTarget InTarget)
