Allow ICU to be disabled by programs; move code which forces it on into ResetBuildConfiguration() rather than ValidateBuildConfiguration().

diff --git a/Engine/Source/Programs/UnrealBuildTool/Windows/UEBuildWindows.cs b/Engine/Source/Programs/UnrealBuildTool/Windows/UEBuildWindows.cs
--- a/Engine/Source/Programs/UnrealBuildTool/Windows/UEBuildWindows.cs
+++ b/Engine/Source/Programs/UnrealBuildTool/Windows/UEBuildWindows.cs
@@ -203,6 +203,7 @@ namespace UnrealBuildTool
 
 		public override void ResetBuildConfiguration(UnrealTargetPlatform InPlatform, UnrealTargetConfiguration InConfiguration)
 		{
+			UEBuildConfiguration.bCompileICU = true;
 		}
 
         public override void ValidateBuildConfiguration(CPPTargetConfiguration Configuration, CPPTargetPlatform Platform, bool bCreateDebugInfo)
@@ -222,8 +223,8 @@ namespace UnrealBuildTool
          */
         public override void ValidateUEBuildConfiguration()
         {
-            UEBuildConfiguration.bCompileICU = !bCompileWithClang;	// @todo clang: ICU causes STL link errors when using Clang on Windows.  Needs debugging.
-        }
+			UEBuildConfiguration.bCompileICU &= !bCompileWithClang;	// @todo clang: ICU causes STL link errors when using Clang on Windows.  Needs debugging.
+		}
 
 		public override void ModifyNewlyLoadedModule(UEBuildModule InModule, TargetInfo Target)
 		{
