#ue4
#ios
* added stub to manifest
* fixed name of stub for non-development builds

diff --git a/Engine/Source/Programs/UnrealBuildTool/Configuration/UEBuildTarget.cs b/Engine/Source/Programs/UnrealBuildTool/Configuration/UEBuildTarget.cs
--- a/Engine/Source/Programs/UnrealBuildTool/Configuration/UEBuildTarget.cs
+++ b/Engine/Source/Programs/UnrealBuildTool/Configuration/UEBuildTarget.cs
@@ -1061,15 +1061,19 @@ namespace UnrealBuildTool
 					// Add all the resources and third party dylibs stored in app bundle
 					MacToolChain.AddAppBundleContentsToManifest(ref Manifest, Binary);
 				}
-                // ok, this is pretty awful, we want the import libraries that go with the editor, only on the PC
-                else if (UnrealBuildTool.BuildingRocket() && 
-                    Path.GetFileNameWithoutExtension(Binary.Config.OutputFilePath).StartsWith("UE4Editor-", StringComparison.InvariantCultureIgnoreCase) &&
-                    Path.GetExtension(Binary.Config.OutputFilePath).EndsWith("dll", StringComparison.InvariantCultureIgnoreCase) && 
-                    Binary.Config.Type == UEBuildBinaryType.DynamicLinkLibrary)
-                {
-                    // ok, this is pretty awful, we want the import libraries that go with the editor, only on the PC
-                    Manifest.AddBinaryNames(Path.Combine(Binary.Config.IntermediateDirectory, Path.GetFileNameWithoutExtension(Binary.Config.OutputFilePath) + ".lib"), "");
-                }
+				else if (TargetPlatform == UnrealTargetPlatform.IOS)
+				{
+					IOSToolChain.AddStubToManifest(ref Manifest, Binary);
+				}
+				// ok, this is pretty awful, we want the import libraries that go with the editor, only on the PC
+				else if (UnrealBuildTool.BuildingRocket() &&
+					Path.GetFileNameWithoutExtension(Binary.Config.OutputFilePath).StartsWith("UE4Editor-", StringComparison.InvariantCultureIgnoreCase) &&
+					Path.GetExtension(Binary.Config.OutputFilePath).EndsWith("dll", StringComparison.InvariantCultureIgnoreCase) &&
+					Binary.Config.Type == UEBuildBinaryType.DynamicLinkLibrary)
+				{
+					// ok, this is pretty awful, we want the import libraries that go with the editor, only on the PC
+					Manifest.AddBinaryNames(Path.Combine(Binary.Config.IntermediateDirectory, Path.GetFileNameWithoutExtension(Binary.Config.OutputFilePath) + ".lib"), "");
+				}
 			}
             {
                 string DebugInfoExtension = BuildPlatform.GetDebugInfoExtension(UEBuildBinaryType.StaticLibrary);
diff --git a/Engine/Source/Programs/UnrealBuildTool/IOS/IOSToolChain.cs b/Engine/Source/Programs/UnrealBuildTool/IOS/IOSToolChain.cs
--- a/Engine/Source/Programs/UnrealBuildTool/IOS/IOSToolChain.cs
+++ b/Engine/Source/Programs/UnrealBuildTool/IOS/IOSToolChain.cs
@@ -143,6 +143,12 @@ namespace UnrealBuildTool
 			}
 		}
 
+		static public void AddStubToManifest(ref FileManifest Manifest, UEBuildBinary Binary)
+		{
+			string StubFile = Path.Combine (Path.GetDirectoryName (Binary.Config.OutputFilePath), Path.GetFileNameWithoutExtension (Binary.Config.OutputFilePath) + ".stub");
+			Manifest.AddFileName(StubFile);
+		}
+
 		static bool bHasPrinted = false;
 		static string GetArchitectureArgument(CPPTargetConfiguration Configuration, string UBTArchitecture)
 		{
diff --git a/Engine/Source/Programs/UnrealBuildTool/IOS/UEDeployIOS.cs b/Engine/Source/Programs/UnrealBuildTool/IOS/UEDeployIOS.cs
--- a/Engine/Source/Programs/UnrealBuildTool/IOS/UEDeployIOS.cs
+++ b/Engine/Source/Programs/UnrealBuildTool/IOS/UEDeployIOS.cs
@@ -251,7 +251,7 @@ namespace UnrealBuildTool.IOS
 				}
 
 				// zip it up into the stub file
-				PackageStub (BinaryPath, GameName);
+				PackageStub (BinaryPath, GameName, Path.GetFileNameWithoutExtension(InExecutablePath));
 			}
 
 			// copy engine assets in
@@ -268,10 +268,10 @@ namespace UnrealBuildTool.IOS
 			return true;
 		}
 
-		private void PackageStub(string BinaryPath, string GameName)
+		private void PackageStub(string BinaryPath, string GameName, string ExeName)
 		{
 			// create the ipa
-			string IPAName = BinaryPath + "/" + GameName + ".stub";
+			string IPAName = BinaryPath + "/" + ExeName + ".stub";
 			// delete the old one
 			if (File.Exists(IPAName))
 			{
