Switched Linux tolchain to use compiler for linking and added action descriptions.

diff --git a/Engine/Source/Programs/UnrealBuildTool/Linux/LinuxToolChain.cs b/Engine/Source/Programs/UnrealBuildTool/Linux/LinuxToolChain.cs
--- a/Engine/Source/Programs/UnrealBuildTool/Linux/LinuxToolChain.cs
+++ b/Engine/Source/Programs/UnrealBuildTool/Linux/LinuxToolChain.cs
@@ -40,18 +40,6 @@ namespace UnrealBuildTool.Linux
                 ClangPath = Path.Combine(BaseLinuxPath, @"bin\Clang++.exe");
                 ArPath = Path.Combine(BaseLinuxPath, @"bin\x86_64-unknown-linux-gnu-ar.exe");
                 RanlibPath = Path.Combine(BaseLinuxPath, @"bin\x86_64-unknown-linux-gnu-ranlib.exe");
-                LdPath = Path.Combine(BaseLinuxPath, @"bin\x86_64-unknown-linux-gnu-ld.exe");
-
-                CRTPath = Path.Combine(BaseLinuxPath, @"lib\gcc\x86_64-unknown-linux-gnu\4.8.1\");
-                if (!Directory.Exists(CRTPath))
-                {
-                    CRTPath = Path.Combine(BaseLinuxPath, @"lib\gcc\x86_64-unknown-linux-gnu\4.5.2\");
-                    if (!Directory.Exists(CRTPath))
-                    {
-                        Console.WriteLine("        Could not find suitable Linux cross-toolchain! (LINUX_ROOT={0})", Environment.GetEnvironmentVariable("LINUX_ROOT"));
-                        return;
-                    }
-                }
             }
 
             // Register this tool chain for both Linux
@@ -207,7 +195,7 @@ namespace UnrealBuildTool.Linux
             // debugging symbols
             if (LinkEnvironment.Config.TargetConfiguration < CPPTargetConfiguration.Shipping)
             {
-                Result += " -export-dynamic";   // needed for backtrace_symbols()...
+                Result += " -rdynamic";   // needed for backtrace_symbols()...
             }
             else
             {
@@ -223,7 +211,7 @@ namespace UnrealBuildTool.Linux
             if (CrossCompiling())
             {
                 // ignore unresolved symbols in shared libs
-                Result += " --unresolved-symbols=ignore-in-shared-libs";
+                Result += " -Wl,--unresolved-symbols=ignore-in-shared-libs";
             }
             else
             {
@@ -231,18 +219,15 @@ namespace UnrealBuildTool.Linux
             }
 
             // RPATH for third party libs
-            Result += " -rpath=${ORIGIN}/../../../engine/binaries/linux";
+            Result += " -Wl,-rpath=${ORIGIN}/../../../engine/binaries/linux";
 
             if (CrossCompiling())
             {
+                Result += " -target x86_64-unknown-linux-gnu";        // Set target triple
                 string SysRootPath = BaseLinuxPath.TrimEnd(new char[] { '\\', '/' });
                 Result += String.Format(" \"--sysroot={0}\"", SysRootPath);
             }
 
-            Result += " --eh-frame-hdr";
-            Result += " -m elf_x86_64";
-            Result += " -dynamic-linker /lib64/ld-linux-x86-64.so.2";
-
             return Result;
         }
 
@@ -297,8 +282,6 @@ namespace UnrealBuildTool.Linux
         static string ClangPath;
         static string ArPath;
         static string RanlibPath;
-        static string LdPath;
-        static string CRTPath;
 
         public override CPPOutput CompileCPPFiles(CPPEnvironment CompileEnvironment, List<FileItem> SourceFiles, string ModuleName)
         {
@@ -419,7 +402,8 @@ namespace UnrealBuildTool.Linux
                 CompileAction.WorkingDirectory = Path.GetFullPath(".");
                 CompileAction.CommandPath = ClangPath;
                 CompileAction.CommandArguments = Arguments + FileArguments + CompileEnvironment.Config.AdditionalArguments;
-                CompileAction.StatusDescription = string.Format("{0}", Path.GetFileName(SourceFile.AbsolutePath));
+                CompileAction.CommandDescription = "Compile";
+                CompileAction.StatusDescription = Path.GetFileName(SourceFile.AbsolutePath);
                 CompileAction.StatusDetailedDescription = SourceFile.Description;
                 CompileAction.bIsGCCCompiler = true;
 
@@ -458,7 +442,8 @@ namespace UnrealBuildTool.Linux
             // Add the output file as a production of the link action.
             FileItem OutputFile = FileItem.GetItemByPath(LinkEnvironment.Config.OutputFilePath);
             ArchiveAction.ProducedItems.Add(OutputFile);
-            ArchiveAction.StatusDescription = string.Format("{0}", Path.GetFileName(OutputFile.AbsolutePath));
+            ArchiveAction.CommandDescription = "Archive";
+            ArchiveAction.StatusDescription = Path.GetFileName(OutputFile.AbsolutePath);
             ArchiveAction.CommandArguments += string.Format("{0} {1} \"{2}\"",  ArPath, GetArchiveArguments(LinkEnvironment), OutputFile.AbsolutePath);
 
             // Add the input files to a response file, and pass the response file on the command-line.
@@ -499,7 +484,7 @@ namespace UnrealBuildTool.Linux
             // Create an action that invokes the linker.
             Action LinkAction = new Action(ActionType.Link);
             LinkAction.WorkingDirectory = Path.GetFullPath(".");
-            LinkAction.CommandPath = LdPath;
+            LinkAction.CommandPath = ClangPath;
 
             // Get link arguments.
             LinkAction.CommandArguments = GetLinkArguments(LinkEnvironment);
@@ -511,22 +496,12 @@ namespace UnrealBuildTool.Linux
             // Add the output file as a production of the link action.
             FileItem OutputFile = FileItem.GetItemByPath(LinkEnvironment.Config.OutputFilePath);
             LinkAction.ProducedItems.Add(OutputFile);
-            LinkAction.StatusDescription = string.Format("{0}", Path.GetFileName(OutputFile.AbsolutePath));
+            LinkAction.CommandDescription = "Link";
+            LinkAction.StatusDescription = Path.GetFileName(OutputFile.AbsolutePath);
 
             // Add the output file to the command-line.
             LinkAction.CommandArguments += string.Format(" -o \"{0}\"", OutputFile.AbsolutePath);
 
-            // CRT-Begin
-			if (CrossCompiling()) 
-			{
-				LinkAction.CommandArguments += string.Format (" \"{0}\"", Path.Combine (BaseLinuxPath, @"usr\lib64\Scrt1.o"));
-				LinkAction.CommandArguments += string.Format (" \"{0}\"", Path.Combine (BaseLinuxPath, @"usr\lib64\crti.o"));
-				LinkAction.CommandArguments += string.Format (" \"{0}\"", Path.Combine (CRTPath, @"crtbeginS.o"));
-				LinkAction.CommandArguments += string.Format (" \"-L{0}\"", Path.Combine (BaseLinuxPath, @"lib"));
-				LinkAction.CommandArguments += string.Format (" \"-L{0}\"", Path.Combine (BaseLinuxPath, @"lib64"));
-				LinkAction.CommandArguments += string.Format (" \"-L{0}\"", Path.Combine (BaseLinuxPath, @"usr\lib64"));
-			}
-
             // Add the input files to a response file, and pass the response file on the command-line.
             List<string> InputFileNames = new List<string>();
             foreach (FileItem InputFile in LinkEnvironment.InputFiles)
@@ -545,7 +520,7 @@ namespace UnrealBuildTool.Linux
 			}
 
 			// add libraries in a library group
-            LinkAction.CommandArguments += string.Format(" --start-group");
+            LinkAction.CommandArguments += string.Format(" -Wl,--start-group");
             foreach (string AdditionalLibrary in LinkEnvironment.Config.AdditionalLibraries)
 			{
 				if (String.IsNullOrEmpty(Path.GetDirectoryName(AdditionalLibrary)))
@@ -561,16 +536,7 @@ namespace UnrealBuildTool.Linux
 			}
             LinkAction.CommandArguments += " -lrt"; // needed for clock_gettime()
             LinkAction.CommandArguments += " -lm"; // math
-            LinkAction.CommandArguments += string.Format(" --end-group");
-
-            // CRT-End
-            LinkAction.CommandArguments += " -lstdc++";
-            LinkAction.CommandArguments += " -lc";
-			if (CrossCompiling()) 
-			{
-				LinkAction.CommandArguments += string.Format (" \"{0}\"", Path.Combine (CRTPath, @"crtend.o"));
-				LinkAction.CommandArguments += string.Format (" \"{0}\"", Path.Combine (BaseLinuxPath, @"usr\lib64\crtn.o"));
-			}
+            LinkAction.CommandArguments += string.Format(" -Wl,--end-group");
 
             // Add the additional arguments specified by the environment.
             LinkAction.CommandArguments += LinkEnvironment.Config.AdditionalArguments;
@@ -579,8 +545,6 @@ namespace UnrealBuildTool.Linux
             // Only execute linking on the local PC.
             LinkAction.bCanExecuteRemotely = false;
 
-            Console.WriteLine(LinkAction.CommandPath + LinkAction.CommandArguments);
-
             return OutputFile;
         }
 
