Show progress (action number/total actions) when executing actions locally.

Also, do not halve number of processors when running on mono (unnecessary generalization).
Part of pull request 49 by sbc, see https://github.com/EpicGames/UnrealEngine/pull/49 for details.

diff --git a/Engine/Source/Programs/UnrealBuildTool/System/LocalExecutor.cs b/Engine/Source/Programs/UnrealBuildTool/System/LocalExecutor.cs
--- a/Engine/Source/Programs/UnrealBuildTool/System/LocalExecutor.cs
+++ b/Engine/Source/Programs/UnrealBuildTool/System/LocalExecutor.cs
@@ -20,6 +20,9 @@ namespace UnrealBuildTool
 		public bool bComplete = false;
 		/** Cache the action that this thread is managing */
 		Action Action;
+        /** For reporting status to the user */
+        int JobNumber;
+        int TotalJobs;
 
 		/** Regex that matches environment variables in $(Variable) format. */
 		static Regex EnvironmentVariableRegex = new Regex("\\$\\(([\\d\\w]+)\\)");
@@ -38,9 +41,11 @@ namespace UnrealBuildTool
 		/**
 		 * Constructor, takes the action to process
 		 */
-		public ActionThread(Action InAction)
+        public ActionThread(Action InAction, int InJobNumber, int InTotalJobs)
 		{
 			Action = InAction;
+            JobNumber = InJobNumber;
+            TotalJobs = InTotalJobs;
 		}
 
 		/**
@@ -86,8 +91,8 @@ namespace UnrealBuildTool
 		}
 
 		/**
-			* The actual function to run in a thread. This is potentially long and blocking
-			*/
+		 * The actual function to run in a thread. This is potentially long and blocking
+		 */
 		private void ThreadFunc()
 		{
 			// thread start time
@@ -148,7 +153,7 @@ namespace UnrealBuildTool
 					}
 					else
 					{
-						Log.TraceInformation("{0} {1}", CommandDescription, Action.StatusDescription);
+                        Log.TraceInformation("[{0}/{1}] {2} {3}", JobNumber, TotalJobs, CommandDescription, Action.StatusDescription);
 					}
 				}
 
@@ -316,14 +321,17 @@ namespace UnrealBuildTool
 				MaxActionsToExecuteInParallel = NumCores;
 			}
 
-			if( Utils.IsRunningOnMono )
+            if (ExternalExecution.GetRuntimePlatform() == UnrealTargetPlatform.Mac)
 			{
 				// @todo iosmerge: this should be looking at available memory as well since some of our
 				// Macs have a poor memory/CPU ratio
 				MaxActionsToExecuteInParallel /= 2;
 			}
 
+            Log.TraceInformation("Performing {0} actions (max {1} parallel jobs)", Actions.Count, MaxActionsToExecuteInParallel);
+
 			Dictionary<Action, ActionThread> ActionThreadDictionary = new Dictionary<Action, ActionThread>();
+            int JobNumber = 1;
 			while(true)
 			{
 				// Count the number of pending and still executing actions.
@@ -403,7 +411,8 @@ namespace UnrealBuildTool
 							// If there aren't any outdated prerequisites of this action, execute it.
                             else if (!bHasOutdatedPrerequisites)
                             {
-								ActionThread ActionThread = new ActionThread(Action);
+                                ActionThread ActionThread = new ActionThread(Action, JobNumber, Actions.Count);
+                                JobNumber++;
 								ActionThread.Run();
 
 								ActionThreadDictionary.Add(Action, ActionThread);
