Delete "Run (iOS)" Xcode schemes that we don't use any more when regenerating Xcode project

diff --git a/Engine/Source/Programs/UnrealBuildTool/System/XcodeProjectFileGenerator.cs b/Engine/Source/Programs/UnrealBuildTool/System/XcodeProjectFileGenerator.cs
--- a/Engine/Source/Programs/UnrealBuildTool/System/XcodeProjectFileGenerator.cs
+++ b/Engine/Source/Programs/UnrealBuildTool/System/XcodeProjectFileGenerator.cs
@@ -1465,6 +1465,17 @@ namespace UnrealBuildTool
 				Log.TraceInformation("Skipping project file write, as it didn't change...");
 			}
 
+			// Delete obsolete schemes
+			string SchemesDir = XcodeProjectPath + "/xcuserdata/" + Environment.UserName + ".xcuserdatad/xcschemes";
+			if (Directory.Exists(SchemesDir))
+			{
+				var ObsoleteSchemes = Directory.GetFiles(SchemesDir, "*iOS (*.xcscheme", SearchOption.AllDirectories);
+				foreach (string SchemeFile in ObsoleteSchemes)
+				{
+					File.Delete(SchemeFile);
+				}
+			}
+
 			// write scheme files for targets
 			foreach (var Target in ProjectTargets) 
 			{
