Do not fail the whole build if a toolchain cannot be switched (as it might be not necessary for the build in question).

diff --git a/Engine/Source/Programs/UnrealBuildTool/Configuration/UEBuildPlatform.cs b/Engine/Source/Programs/UnrealBuildTool/Configuration/UEBuildPlatform.cs
--- a/Engine/Source/Programs/UnrealBuildTool/Configuration/UEBuildPlatform.cs
+++ b/Engine/Source/Programs/UnrealBuildTool/Configuration/UEBuildPlatform.cs
@@ -400,14 +400,16 @@ namespace UnrealBuildTool
                         // switch over (note that version string can be empty)
                         if (!RunSDKHooks(PlatformSDKRoot, CurrentSDKString, SDKHookType.Uninstall))
                         {
-                            throw new BuildException("Failed to uninstall currently installed SDK {0}", CurrentSDKString);
+                            Console.WriteLine("Failed to uninstall currently installed SDK {0}", CurrentSDKString);
+                            return false;
                         }
                         // delete Manifest file to avoid multiple uninstalls
                         SetCurrentlyInstalledSDKString(PlatformSDKRoot, "");
 
                         if (!RunSDKHooks(PlatformSDKRoot, GetRequiredSDKString(), SDKHookType.Install, false))
                         {
-                            throw new BuildException("Failed to install required SDK {0}", GetRequiredSDKString());
+                            Console.WriteLine("Failed to install required SDK {0}", GetRequiredSDKString());
+                            return false;
                         }
                         SetCurrentlyInstalledSDKString(PlatformSDKRoot, GetRequiredSDKString());
                     }
@@ -415,7 +417,8 @@ namespace UnrealBuildTool
                     // load environment variables from current SDK
                     if (!SetupEnvironmentFromSDK(PlatformSDKRoot, GetRequiredSDKString()))
                     {
-                        throw new BuildException("Failed to load environment from required SDK {0}", GetRequiredSDKString());
+                        Console.WriteLine("Failed to load environment from required SDK {0}", GetRequiredSDKString());
+                        return false;
                     }
                 }
             }
