#ttp 331927 - [GitHub] 47 : Fixed game plugins using shared PCH files
#proj UE4
#branch UE4
#summary Merged.

diff --git a/Engine/Source/Programs/UnrealBuildTool/System/RulesCompiler.cs b/Engine/Source/Programs/UnrealBuildTool/System/RulesCompiler.cs
--- a/Engine/Source/Programs/UnrealBuildTool/System/RulesCompiler.cs
+++ b/Engine/Source/Programs/UnrealBuildTool/System/RulesCompiler.cs
@@ -1226,9 +1226,9 @@ namespace UnrealBuildTool
 					// Disable shared PCHs for game modules by default
 					if (RulesObject.PCHUsage == ModuleRules.PCHUsageMode.Default)
 					{
-						// Engine/Source and Engine/Plugins are considered 'Engine' code...
+						// Note that bIsEngineModule includes Engine/Plugins, so Engine/Plugins will use shared PCHs.
 						var IsProgramTarget = Target.Type != null && Target.Type == TargetRules.TargetType.Program;
-						if (bIsEngineModule || Plugins.IsPluginModule(ModuleName) || IsProgramTarget )
+						if (bIsEngineModule || IsProgramTarget)
 						{
 							// Engine module or plugin module -- allow shared PCHs
 							RulesObject.PCHUsage = ModuleRules.PCHUsageMode.UseSharedPCHs;
