#ttp 331937 - [GitHub] 44 : Prevent duplication of Target Arguments passed to RunUBT
#proj UE4
#branch UE4
#summary Merged.

diff --git a/Engine/Source/Programs/UnrealBuildTool/System/ProjectFileGenerator.cs b/Engine/Source/Programs/UnrealBuildTool/System/ProjectFileGenerator.cs
--- a/Engine/Source/Programs/UnrealBuildTool/System/ProjectFileGenerator.cs
+++ b/Engine/Source/Programs/UnrealBuildTool/System/ProjectFileGenerator.cs
@@ -1204,10 +1204,7 @@ namespace UnrealBuildTool
 
 						var ArgumentsCopy = new string[ Arguments.Length + 1 ];
 						ArgumentsCopy[ 0 ] = TargetName;
-						for( var CurEntryIndex = 0; CurEntryIndex < Arguments.Length; ++CurEntryIndex )
-						{
-							ArgumentsCopy[ CurEntryIndex + 1 ] = Arguments[ CurEntryIndex ];
-						}
+						Array.Copy(Arguments, 0, ArgumentsCopy, 1, Arguments.Length);
 
 						bSuccess = UnrealBuildTool.RunUBT( ArgumentsCopy ) == ECompilationResult.Succeeded;
 						if( !bSuccess )
