#ttp 331894 - [GitHub] 37 : Assert for possible undefined behavior where UProjectFile is empty.
#proj UE4
#branch UE4
#summary Merged.

diff --git a/Engine/Source/Programs/UnrealBuildTool/System/UnrealBuildTool.cs b/Engine/Source/Programs/UnrealBuildTool/System/UnrealBuildTool.cs
--- a/Engine/Source/Programs/UnrealBuildTool/System/UnrealBuildTool.cs
+++ b/Engine/Source/Programs/UnrealBuildTool/System/UnrealBuildTool.cs
@@ -491,6 +491,7 @@ namespace UnrealBuildTool
 							{
 								UProjectPath = AlternativeProjectpath;
 								UProjectFile = UProjectFile.Substring(3);
+								Debug.Assert(UProjectFile.Length > 0);
 							}
 						}
 					}
