Xbox One deploy/packaging/path fixes
- Executable path modified to Game\\Binaries\\XboxOne
- Packaging reworked to function with March XDK and remove warnings/errors
- Pakcaging images cleaned up to allow localization
- Support for localized strings added
- Support for resource.pri generation added
- Separate Package.appxmanifest and Deploy.appxmanifest for packaging/deployment (major difference is number of included configs)
- Switched from TargetName to AumidOverride to set configuration to run when debugging in VS (TargetName was incompatible with the executable path change)
- Cleaned up some files that are no longer in use for packaging.

These changes *require* project file regeneration and updated UAT/UBT to function (which will in most cases be taken care of for you automatically).

diff --git a/Engine/Source/Programs/UnrealBuildTool/System/UEPlatformProjectGenerator.cs b/Engine/Source/Programs/UnrealBuildTool/System/UEPlatformProjectGenerator.cs
--- a/Engine/Source/Programs/UnrealBuildTool/System/UEPlatformProjectGenerator.cs
+++ b/Engine/Source/Programs/UnrealBuildTool/System/UEPlatformProjectGenerator.cs
@@ -82,6 +82,20 @@ namespace UnrealBuildTool
 			return true;
 		}
 
+		/// <summary>
+		/// Allow various platform project generators to generate any special project properties if required
+		/// </summary>
+		/// <param name="InPlatform"></param>
+		/// <returns></returns>
+		public static bool GenerateGamePlatformSpecificProperties(UnrealTargetPlatform InPlatform, UnrealTargetConfiguration Configuration, TargetRules.TargetType TargetType, StringBuilder VCProjectFileContent, string RootDirectory, string TargetFilePath)
+		{
+			if (ProjectGeneratorDictionary.ContainsKey(InPlatform) == true)
+			{
+				ProjectGeneratorDictionary[InPlatform].GenerateGameProperties(Configuration, VCProjectFileContent, TargetType, RootDirectory, TargetFilePath); ;
+			}
+			return true;
+		}
+
 		public static bool PlatformRequiresVSUserFileGeneration(List<UnrealTargetPlatform> InPlatforms, List<UnrealTargetConfiguration> InConfigurations)
 		{
 			bool bRequiresVSUserFileGeneration = false;
@@ -107,6 +121,11 @@ namespace UnrealBuildTool
 			// Do nothing
 		}
 
+		public virtual void GenerateGameProperties(UnrealTargetConfiguration Configuration, StringBuilder VCProjectFileContent, TargetRules.TargetType TargetType, string RootDirectory, string TargetFilePath)
+		{
+			// Do nothing
+		}
+
 		public virtual bool RequiresVSUserFileGeneration()
 		{
 			return false;
diff --git a/Engine/Source/Programs/UnrealBuildTool/System/VCProject.cs b/Engine/Source/Programs/UnrealBuildTool/System/VCProject.cs
--- a/Engine/Source/Programs/UnrealBuildTool/System/VCProject.cs
+++ b/Engine/Source/Programs/UnrealBuildTool/System/VCProject.cs
@@ -983,15 +983,11 @@ namespace UnrealBuildTool
 						VCProjectFileContent.Append(PathStrings);
 					}
 
-					// Force specification of TargetName on XboxOne so that the manifest can identify the correct executable to the debugger.
-					if (Platform == UnrealTargetPlatform.XboxOne)
+					if (TargetRules.IsGameType(TargetRulesObject.Type) &&
+						(TargetRules.IsEditorType(TargetRulesObject.Type) == false))
 					{
-						VCProjectFileContent.Append("		<TargetName>" + Utils.GetFilenameWithoutAnyExtensions( TargetFilePath ));
-						if (Configuration != UnrealTargetConfiguration.Development)
-						{
-							VCProjectFileContent.Append(UBTConfigurationName);
-						}
-						VCProjectFileContent.Append("</TargetName>" + ProjectFileGenerator.NewLine);
+						// Allow platforms to add any special properties they require... like aumid override for Xbox One
+						UEPlatformProjectGenerator.GenerateGamePlatformSpecificProperties(Platform, Configuration, TargetRulesObject.Type, VCProjectFileContent, RootDirectory, TargetFilePath);
 					}
 
 					// This is the standard UE4 based project NMake build line:
