#ue4
#ios
* better fix for creating stubs on Mac
* creates stub files when building from UFE on Mac, need to fix that next, but wanted this in so PC workflow was back to normal

diff --git a/Engine/Build/BatchFiles/Mac/Build.sh b/Engine/Build/BatchFiles/Mac/Build.sh
--- a/Engine/Build/BatchFiles/Mac/Build.sh
+++ b/Engine/Build/BatchFiles/Mac/Build.sh
@@ -17,11 +17,13 @@ case $ACTION in
                 if [ $2 = "iphoneos" ]
                 then
 	                Platform="IOS"
+			AdditionalFlags+=" -nocreatestub"
                 else
 	                if [ $2 = "iphonesimulator" ]
 	                then
 		                Platform="IOS"
 		                AdditionalFlags+=" -simulator"
+				AdditionalFlags+=" -nocreatestub"
 	                else
 		                Platform="Mac"
 	                fi
@@ -39,11 +41,13 @@ case $ACTION in
                 if [ $3 = "iphoneos" ]
                 then
 	                Platform="IOS"
+			AdditionalFlags+=" -nocreatestub"
                 else
 	                if [ $3 = "iphonesimulator" ]
 	                then
 		                Platform="IOS"
 		                AdditionalFlags+=" -simulator"
+				AdditionalFlags+=" -nocreatestub"
 	                else
 		                Platform="Mac"
 	                fi
diff --git a/Engine/Build/BatchFiles/Mac/RocketBuild.sh b/Engine/Build/BatchFiles/Mac/RocketBuild.sh
--- a/Engine/Build/BatchFiles/Mac/RocketBuild.sh
+++ b/Engine/Build/BatchFiles/Mac/RocketBuild.sh
@@ -16,11 +16,13 @@ case $ACTION in
                 if [ $2 = "iphoneos" ]
                 then
 	                Platform="IOS"
+			AdditionalFlags+=" -nocreatestub"
                 else
 	                if [ $2 = "iphonesimulator" ]
 	                then
 		                Platform="IOS"
 		                AdditionalFlags+=" -simulator"
+				AdditionalFlags+=" -nocreatestub"
 	                else
 		                Platform="Mac"
 	                fi
@@ -37,11 +39,13 @@ case $ACTION in
                 if [ $3 = "iphoneos" ]
                 then
 	                Platform="IOS"
+			AdditionalFlags+=" -nocreatestub"
                 else
 	                if [ $3 = "iphonesimulator" ]
 	                then
 		                Platform="IOS"
 		                AdditionalFlags+=" -simulator"
+				AdditionalFlags+=" -nocreatestub"
 	                else
 		                Platform="Mac"
 	                fi
diff --git a/Engine/Source/Programs/UnrealBuildTool/IOS/IOSToolChain.cs b/Engine/Source/Programs/UnrealBuildTool/IOS/IOSToolChain.cs
--- a/Engine/Source/Programs/UnrealBuildTool/IOS/IOSToolChain.cs
+++ b/Engine/Source/Programs/UnrealBuildTool/IOS/IOSToolChain.cs
@@ -9,6 +9,8 @@ using System.Diagnostics;
 using System.Security.AccessControl;
 using System.Xml;
 using System.Text;
+using Ionic.Zip;
+using Ionic.Zlib;
 
 namespace UnrealBuildTool
 {
@@ -145,8 +147,11 @@ namespace UnrealBuildTool
 
 		static public void AddStubToManifest(ref FileManifest Manifest, UEBuildBinary Binary)
 		{
-			string StubFile = Path.Combine (Path.GetDirectoryName (Binary.Config.OutputFilePath), Path.GetFileNameWithoutExtension (Binary.Config.OutputFilePath) + ".stub");
-			Manifest.AddFileName(StubFile);
+			if (BuildConfiguration.bCreateStubIPA)
+			{
+				string StubFile = Path.Combine (Path.GetDirectoryName (Binary.Config.OutputFilePath), Path.GetFileNameWithoutExtension (Binary.Config.OutputFilePath) + ".stub");
+				Manifest.AddFileName (StubFile);
+			}
 		}
 
 		static bool bHasPrinted = false;
@@ -661,6 +666,120 @@ namespace UnrealBuildTool
 			return DestFile;
 		}
 
+		private void PackageStub(string BinaryPath, string GameName, string ExeName)
+		{
+			// create the ipa
+			string IPAName = BinaryPath + "/" + ExeName + ".stub";
+			// delete the old one
+			if (File.Exists(IPAName))
+			{
+				File.Delete(IPAName);
+			}
+
+			// make the subdirectory if needed
+			string DestSubdir = Path.GetDirectoryName(IPAName);
+			if (!Directory.Exists(DestSubdir))
+			{
+				Directory.CreateDirectory(DestSubdir);
+			}
+
+			// set up the directories
+			string ZipWorkingDir = String.Format("Payload/{0}.app/", GameName);
+			string ZipSourceDir = string.Format("{0}/Payload/{1}.app", BinaryPath, GameName);
+
+			// create the file
+			using (ZipFile Zip = new ZipFile())
+			{
+				// add the entire directory
+				Zip.AddDirectory(ZipSourceDir, ZipWorkingDir);
+
+				// Update permissions to be UNIX-style
+				// Modify the file attributes of any added file to unix format
+				foreach (ZipEntry E in Zip.Entries)
+				{
+					const byte FileAttributePlatform_NTFS = 0x0A;
+					const byte FileAttributePlatform_UNIX = 0x03;
+					const byte FileAttributePlatform_FAT = 0x00;
+
+					const int UNIX_FILETYPE_NORMAL_FILE = 0x8000;
+					//const int UNIX_FILETYPE_SOCKET = 0xC000;
+					//const int UNIX_FILETYPE_SYMLINK = 0xA000;
+					//const int UNIX_FILETYPE_BLOCKSPECIAL = 0x6000;
+					const int UNIX_FILETYPE_DIRECTORY = 0x4000;
+					//const int UNIX_FILETYPE_CHARSPECIAL = 0x2000;
+					//const int UNIX_FILETYPE_FIFO = 0x1000;
+
+					const int UNIX_EXEC = 1;
+					const int UNIX_WRITE = 2;
+					const int UNIX_READ = 4;
+
+
+					int MyPermissions = UNIX_READ | UNIX_WRITE;
+					int OtherPermissions = UNIX_READ;
+
+					int PlatformEncodedBy = (E.VersionMadeBy >> 8) & 0xFF;
+					int LowerBits = 0;
+
+					// Try to preserve read-only if it was set
+					bool bIsDirectory = E.IsDirectory;
+
+					// Check to see if this 
+					bool bIsExecutable = false;
+					if (Path.GetFileNameWithoutExtension(E.FileName).Equals(GameName, StringComparison.InvariantCultureIgnoreCase))
+					{
+						bIsExecutable = true;
+					}
+
+					if (bIsExecutable)
+					{
+						// The executable will be encrypted in the final distribution IPA and will compress very poorly, so keeping it
+						// uncompressed gives a better indicator of IPA size for our distro builds
+						E.CompressionLevel = CompressionLevel.None;
+					}
+
+					if ((PlatformEncodedBy == FileAttributePlatform_NTFS) || (PlatformEncodedBy == FileAttributePlatform_FAT))
+					{
+						FileAttributes OldAttributes = E.Attributes;
+						//LowerBits = ((int)E.Attributes) & 0xFFFF;
+
+						if ((OldAttributes & FileAttributes.Directory) != 0)
+						{
+							bIsDirectory = true;
+						}
+
+						// Permissions
+						if ((OldAttributes & FileAttributes.ReadOnly) != 0)
+						{
+							MyPermissions &= ~UNIX_WRITE;
+							OtherPermissions &= ~UNIX_WRITE;
+						}
+					}
+
+					if (bIsDirectory || bIsExecutable)
+					{
+						MyPermissions |= UNIX_EXEC;
+						OtherPermissions |= UNIX_EXEC;
+					}
+
+					// Re-jigger the external file attributes to UNIX style if they're not already that way
+					if (PlatformEncodedBy != FileAttributePlatform_UNIX)
+					{
+						int NewAttributes = bIsDirectory ? UNIX_FILETYPE_DIRECTORY : UNIX_FILETYPE_NORMAL_FILE;
+
+						NewAttributes |= (MyPermissions << 6);
+						NewAttributes |= (OtherPermissions << 3);
+						NewAttributes |= (OtherPermissions << 0);
+
+						// Now modify the properties
+						E.AdjustExternalFileAttributes(FileAttributePlatform_UNIX, (NewAttributes << 16) | LowerBits);
+					}
+				}
+
+				// Save it out
+				Zip.Save(IPAName);
+			}
+		}
+
 		public override void PostBuildSync(UEBuildTarget Target)
 		{
 			base.PostBuildSync(Target);
@@ -697,6 +816,58 @@ namespace UnrealBuildTool
 					Directory.CreateDirectory(String.Format("{0}/Payload/{1}.app", RemoteShadowDirectoryMac, Target.GameName));
 				}
 				File.Copy(Target.OutputPath, FinalRemoteExecutablePath, true);
+
+				if (BuildConfiguration.bCreateStubIPA)
+				{
+					string Project = Target.ProjectDirectory + "/" + Target.GameName + ".uproject";
+
+					// generate the dummy project so signing works
+					if (Target.GameName == "UE4Game" || Target.GameName == "UE4Client" || Utils.IsFileUnderDirectory(Target.ProjectDirectory + "/" + Target.GameName + ".uproject", Path.GetFullPath("../..")))
+					{
+						UnrealBuildTool.GenerateProjectFiles (new XcodeProjectFileGenerator (), new string[] {"-platforms=IOS", "-NoIntellIsense"});
+						Project = Path.GetFullPath("../..") + "/UE4_IOS.xcodeproj";
+					}
+					else
+					{
+						Project = Target.ProjectDirectory + "/" + Target.GameName + ".xcodeproj";
+					}
+
+					if (Directory.Exists (Project))
+					{
+						// code sign the project
+						string CmdLine = XcodeDeveloperDir + "usr/bin/xcodebuild" +
+						                " -project \"" + Project + "\"" +
+						                " -configuration " + Target.Configuration +
+						                " -scheme '" + Target.GameName + " - iOS'" +
+						                " -sdk iphoneos" +
+						                " CODE_SIGN_IDENTITY=\"iPhone Developer\"";
+
+						Process SignProcess = new Process ();
+						SignProcess.StartInfo.WorkingDirectory = RemoteShadowDirectoryMac;
+						SignProcess.StartInfo.FileName = "/usr/bin/xcrun";
+						SignProcess.StartInfo.Arguments = CmdLine;
+						SignProcess.OutputDataReceived += new DataReceivedEventHandler (OutputReceivedDataEventHandler);
+						SignProcess.ErrorDataReceived += new DataReceivedEventHandler (OutputReceivedDataEventHandler);
+
+						OutputReceivedDataEventHandlerEncounteredError = false;
+						OutputReceivedDataEventHandlerEncounteredErrorMessage = "";
+						Utils.RunLocalProcess (SignProcess);
+
+						// delete the temp project
+						if (Project.Contains ("_IOS.xcodeproj"))
+						{
+							Directory.Delete (Project, true);
+						}
+
+						if (OutputReceivedDataEventHandlerEncounteredError)
+						{
+							throw new Exception (OutputReceivedDataEventHandlerEncounteredErrorMessage);
+						}
+
+						// Package the stub
+						PackageStub (RemoteShadowDirectoryMac, Target.GameName, Path.GetFileNameWithoutExtension (Target.OutputPath));
+					}
+				}
 			}
 			else
 			{
diff --git a/Engine/Source/Programs/UnrealBuildTool/IOS/UEDeployIOS.cs b/Engine/Source/Programs/UnrealBuildTool/IOS/UEDeployIOS.cs
--- a/Engine/Source/Programs/UnrealBuildTool/IOS/UEDeployIOS.cs
+++ b/Engine/Source/Programs/UnrealBuildTool/IOS/UEDeployIOS.cs
@@ -200,60 +200,6 @@ namespace UnrealBuildTool.IOS
 			// copy the GameName binary
 			File.Copy(BinaryPath + "/" + GameExeName, AppDirectory + "/" + GameName, true);
 
-			if (BuildConfiguration.bCreateStubIPA)
-			{
-				// copy the custom rules
-				File.Copy (InEngineDir + "/Build/IOS/XcodeSupportFiles/CustomResourceRules.plist", AppDirectory + "/CustomResourceRules.plist", true);
-
-				// create the PkgInfo file
-
-				// code sign
-				string Arguments = "CODESIGN_ALLOCATE=/Applications/Xcode.app/Contents/Developer/Toolchains/XcodeDefault.xctoolchain/usr/bin/codesign_allocate";
-				Arguments += " /usr/bin/codesign --force --sign \"iPhone Developer\"";
-				Arguments += " --resource-rules=\"";
-				Arguments += AppDirectory + "/CustomResourceRules.plist\"";
-				Arguments += " --entitlements \"" + Entitlements + "\" ";
-				Arguments += "\"" + AppDirectory + "\"";
-				using(var CodeSignProcess = new Process())
-				{
-					CodeSignProcess.StartInfo.FileName = "/usr/bin/env";
-					CodeSignProcess.StartInfo.Arguments = Arguments;
-					CodeSignProcess.StartInfo.CreateNoWindow = true;
-					CodeSignProcess.StartInfo.UseShellExecute = false;
-					CodeSignProcess.StartInfo.RedirectStandardOutput = true;
-					CodeSignProcess.OutputDataReceived += StdOut;
-					CodeSignProcess.StartInfo.RedirectStandardError = true;
-					CodeSignProcess.ErrorDataReceived += StdErr;
-					CodeSignProcess.StartInfo.RedirectStandardInput = false;
-					CodeSignProcess.StartInfo.UseShellExecute = false;
-
-					Log.TraceVerbose( "Launching {0} to harvest Visual Studio environment settings...", CodeSignProcess.StartInfo.FileName );
-
-					// Try to launch the process, and produce a friendly error message if it fails.
-					try
-					{
-						// Start the process up and then wait for it to finish
-						CodeSignProcess.Start();
-						CodeSignProcess.BeginOutputReadLine();
-						CodeSignProcess.BeginErrorReadLine();
-						while (!CodeSignProcess.HasExited)
-						{
-							System.Threading.Thread.Sleep(5);
-						}
-						CodeSignProcess.WaitForExit();
-					}
-					catch(Exception ex)
-					{
-						throw new BuildException( ex, "Failed to start local process for action (\"{0}\"): {1} {2}", ex.Message, CodeSignProcess.StartInfo.FileName, CodeSignProcess.StartInfo.Arguments );
-					}
-
-					Log.TraceVerbose( "Finished launching {0}.", CodeSignProcess.StartInfo.FileName );
-				}
-
-				// zip it up into the stub file
-				PackageStub (BinaryPath, GameName, Path.GetFileNameWithoutExtension(InExecutablePath));
-			}
-
 			// copy engine assets in
 			CopyFiles(InEngineDir + "/Build/IOS/Resources/Graphics", AppDirectory, "*.png", true);
 			// merge game assets on top
@@ -268,120 +214,6 @@ namespace UnrealBuildTool.IOS
 			return true;
 		}
 
-		private void PackageStub(string BinaryPath, string GameName, string ExeName)
-		{
-			// create the ipa
-			string IPAName = BinaryPath + "/" + ExeName + ".stub";
-			// delete the old one
-			if (File.Exists(IPAName))
-			{
-				File.Delete(IPAName);
-			}
-
-			// make the subdirectory if needed
-			string DestSubdir = Path.GetDirectoryName(IPAName);
-			if (!Directory.Exists(DestSubdir))
-			{
-				Directory.CreateDirectory(DestSubdir);
-			}
-
-			// set up the directories
-			string ZipWorkingDir = String.Format("Payload/{0}.app/", GameName);
-			string ZipSourceDir = string.Format("{0}/Payload/{1}.app", BinaryPath, GameName);
-
-			// create the file
-			using (ZipFile Zip = new ZipFile())
-			{
-				// add the entire directory
-				Zip.AddDirectory(ZipSourceDir, ZipWorkingDir);
-
-				// Update permissions to be UNIX-style
-				// Modify the file attributes of any added file to unix format
-				foreach (ZipEntry E in Zip.Entries)
-				{
-					const byte FileAttributePlatform_NTFS = 0x0A;
-					const byte FileAttributePlatform_UNIX = 0x03;
-					const byte FileAttributePlatform_FAT = 0x00;
-
-					const int UNIX_FILETYPE_NORMAL_FILE = 0x8000;
-					//const int UNIX_FILETYPE_SOCKET = 0xC000;
-					//const int UNIX_FILETYPE_SYMLINK = 0xA000;
-					//const int UNIX_FILETYPE_BLOCKSPECIAL = 0x6000;
-					const int UNIX_FILETYPE_DIRECTORY = 0x4000;
-					//const int UNIX_FILETYPE_CHARSPECIAL = 0x2000;
-					//const int UNIX_FILETYPE_FIFO = 0x1000;
-
-					const int UNIX_EXEC = 1;
-					const int UNIX_WRITE = 2;
-					const int UNIX_READ = 4;
-
-
-					int MyPermissions = UNIX_READ | UNIX_WRITE;
-					int OtherPermissions = UNIX_READ;
-
-					int PlatformEncodedBy = (E.VersionMadeBy >> 8) & 0xFF;
-					int LowerBits = 0;
-
-					// Try to preserve read-only if it was set
-					bool bIsDirectory = E.IsDirectory;
-
-					// Check to see if this 
-					bool bIsExecutable = false;
-					if (Path.GetFileNameWithoutExtension(E.FileName).Equals(GameName, StringComparison.InvariantCultureIgnoreCase))
-					{
-						bIsExecutable = true;
-					}
-
-					if (bIsExecutable)
-					{
-						// The executable will be encrypted in the final distribution IPA and will compress very poorly, so keeping it
-						// uncompressed gives a better indicator of IPA size for our distro builds
-						E.CompressionLevel = CompressionLevel.None;
-					}
-
-					if ((PlatformEncodedBy == FileAttributePlatform_NTFS) || (PlatformEncodedBy == FileAttributePlatform_FAT))
-					{
-						FileAttributes OldAttributes = E.Attributes;
-						//LowerBits = ((int)E.Attributes) & 0xFFFF;
-
-						if ((OldAttributes & FileAttributes.Directory) != 0)
-						{
-							bIsDirectory = true;
-						}
-
-						// Permissions
-						if ((OldAttributes & FileAttributes.ReadOnly) != 0)
-						{
-							MyPermissions &= ~UNIX_WRITE;
-							OtherPermissions &= ~UNIX_WRITE;
-						}
-					}
-
-					if (bIsDirectory || bIsExecutable)
-					{
-						MyPermissions |= UNIX_EXEC;
-						OtherPermissions |= UNIX_EXEC;
-					}
-
-					// Re-jigger the external file attributes to UNIX style if they're not already that way
-					if (PlatformEncodedBy != FileAttributePlatform_UNIX)
-					{
-						int NewAttributes = bIsDirectory ? UNIX_FILETYPE_DIRECTORY : UNIX_FILETYPE_NORMAL_FILE;
-
-						NewAttributes |= (MyPermissions << 6);
-						NewAttributes |= (OtherPermissions << 3);
-						NewAttributes |= (OtherPermissions << 0);
-
-						// Now modify the properties
-						E.AdjustExternalFileAttributes(FileAttributePlatform_UNIX, (NewAttributes << 16) | LowerBits);
-					}
-				}
-
-				// Save it out
-				Zip.Save(IPAName);
-			}
-		}
-
 		public override bool PrepTargetForDeployment(UEBuildTarget InTarget)
 		{
 			if (ExternalExecution.GetRuntimePlatform() == UnrealTargetPlatform.Mac && Environment.GetEnvironmentVariable("UBT_NO_POST_DEPLOY") != "true")
diff --git a/Engine/Source/Programs/UnrealBuildTool/System/UnrealBuildTool.cs b/Engine/Source/Programs/UnrealBuildTool/System/UnrealBuildTool.cs
--- a/Engine/Source/Programs/UnrealBuildTool/System/UnrealBuildTool.cs
+++ b/Engine/Source/Programs/UnrealBuildTool/System/UnrealBuildTool.cs
@@ -1041,6 +1041,12 @@ namespace UnrealBuildTool
 					BuildConfiguration.bUseIncludeDependencyResolveCache = true;
 				}
 
+				bool CreateStub = Utils.ParseCommandLineFlag(Arguments, "-nocreatestub", out ArgumentIndex);
+				if (CreateStub)
+				{
+					BuildConfiguration.bCreateStubIPA = false;
+				}
+
 				// Build action lists for all passed in targets.
 				var OutputItemsForAllTargets = new List<FileItem>();
 				foreach (string[] TargetSetting in TargetSettings)
