#ttp 331903 - [GitHub] Improve platform exclusion
#proj UE4
#branch UE4
#summary Removed PlatformsToExclude setting as its no longer used.

diff --git a/Engine/Source/Programs/UnrealBuildTool/System/UnrealBuildTool.cs b/Engine/Source/Programs/UnrealBuildTool/System/UnrealBuildTool.cs
--- a/Engine/Source/Programs/UnrealBuildTool/System/UnrealBuildTool.cs
+++ b/Engine/Source/Programs/UnrealBuildTool/System/UnrealBuildTool.cs
@@ -586,9 +586,6 @@ namespace UnrealBuildTool
 			bool bRunCopyrightVerification = false;
 			bool bDumpToFile = false;
 			bool bCheckThirdPartyHeaders = false;
-			// List of all platforms to skip when bBuildAll is true.
-			List<UnrealTargetPlatform> PlatformsToExclude = new List<UnrealTargetPlatform>();
-
 
 			// @todo: Ideally we never need to Mutex unless we are invoked with the same target project,
 			// in the same branch/path!  This would allow two clientspecs to build at the same time (even though we need
@@ -772,34 +769,18 @@ namespace UnrealBuildTool
 						}
 						else
 						{
-							bool bPlatformExclusionArg = false;
-							// Check if this is one of the platforms to exclude params: -nowin32, -nomac, etc.
-							foreach (UnrealTargetPlatform platform in Enum.GetValues(typeof(UnrealTargetPlatform)))
+							// This arg may be a game name. Check for the existence of a game folder with this name.
+							// "Engine" is not a valid game name.
+							if (LowercaseArg != "engine" && Arg.IndexOfAny(Path.GetInvalidPathChars()) == -1 &&
+								Directory.Exists(Path.Combine(ProjectFileGenerator.RootRelativePath, Arg, "Config")))
 							{
-								string ExcludePlatformArg = "-no" + platform.ToString().ToLowerInvariant();
-								if (LowercaseArg == ExcludePlatformArg)
-								{
-									PlatformsToExclude.Add(platform);
-									bPlatformExclusionArg = true;
-									break;
-								}
+								GameName = Arg;
+								Log.TraceVerbose("CommandLine: Found game name '{0}'", GameName);
 							}
-
-							if (!bPlatformExclusionArg)
+							else if (LowercaseArg == "rocket")
 							{
-								// This arg may be a game name. Check for the existence of a game folder with this name.
-								// "Engine" is not a valid game name.
-								if (LowercaseArg != "engine" && Arg.IndexOfAny(Path.GetInvalidPathChars()) == -1 &&
-									Directory.Exists(Path.Combine(ProjectFileGenerator.RootRelativePath, Arg, "Config")))
-								{
-									GameName = Arg;
-									Log.TraceVerbose("CommandLine: Found game name '{0}'", GameName);
-								}
-								else if (LowercaseArg == "rocket")
-								{
-									GameName = Arg;
-									Log.TraceVerbose("CommandLine: Found game name '{0}'", GameName);
-								}
+								GameName = Arg;
+								Log.TraceVerbose("CommandLine: Found game name '{0}'", GameName);
 							}
 						}
 					}
