#ue4
#ios
* fix for xcode generation for UE4Client

diff --git a/Engine/Source/Programs/UnrealBuildTool/System/XcodeProjectFileGenerator.cs b/Engine/Source/Programs/UnrealBuildTool/System/XcodeProjectFileGenerator.cs
--- a/Engine/Source/Programs/UnrealBuildTool/System/XcodeProjectFileGenerator.cs
+++ b/Engine/Source/Programs/UnrealBuildTool/System/XcodeProjectFileGenerator.cs
@@ -484,7 +484,7 @@ namespace UnrealBuildTool
 				"\t\t};" + ProjectFileGenerator.NewLine);
 		}
 
-		private void AppendIOSRunConfig(ref StringBuilder Contents, string ConfigName, string ConfigGuid, string TargetName, string EngineRelative, string GamePath, bool bIsUE4Game, bool IsAGame)
+		private void AppendIOSRunConfig(ref StringBuilder Contents, string ConfigName, string ConfigGuid, string TargetName, string EngineRelative, string GamePath, bool bIsUE4Game, bool IsAGame, bool bIsUE4Client)
 		{
 			Contents.Append(
 				"\t\t" + ConfigGuid + " /* " + ConfigName + " */ = {" + ProjectFileGenerator.NewLine +
@@ -500,13 +500,22 @@ namespace UnrealBuildTool
 			{
 				if (bIsUE4Game)
 				{
-					Contents.Append(
+					Contents.Append (
 						"\t\t\t\tCODE_SIGN_RESOURCE_RULES_PATH = \"" + EngineRelative + "Engine/Build/iOS/XcodeSupportFiles/CustomResourceRules.plist\";" + ProjectFileGenerator.NewLine +
 						"\t\t\t\tINFOPLIST_FILE = \"" + EngineRelative + "Engine/Intermediate/IOS/" + TargetName + "-Info.plist\";" + ProjectFileGenerator.NewLine +
 						"\t\t\t\tSYMROOT = \"" + EngineRelative + "Engine/Binaries/IOS/Payload\";" + ProjectFileGenerator.NewLine +
 						"\t\t\t\tOBJROOT = \"" + EngineRelative + "Engine/Intermediate/IOS/build\";" + ProjectFileGenerator.NewLine +
 						"\t\t\t\tCONFIGURATION_BUILD_DIR = \"" + EngineRelative + "Engine/Binaries/IOS/Payload\";" + ProjectFileGenerator.NewLine);
 				}
+				else if (bIsUE4Client)
+				{
+					Contents.Append (
+						"\t\t\t\tCODE_SIGN_RESOURCE_RULES_PATH = \"" + EngineRelative + "Engine/Build/iOS/XcodeSupportFiles/CustomResourceRules.plist\";" + ProjectFileGenerator.NewLine +
+						"\t\t\t\tINFOPLIST_FILE = \"" + EngineRelative + "Engine/Intermediate/IOS/UE4Game-Info.plist\";" + ProjectFileGenerator.NewLine +
+						"\t\t\t\tSYMROOT = \"" + EngineRelative + "Engine/Binaries/IOS/Payload\";" + ProjectFileGenerator.NewLine +
+						"\t\t\t\tOBJROOT = \"" + EngineRelative + "Engine/Intermediate/IOS/build\";" + ProjectFileGenerator.NewLine +
+						"\t\t\t\tCONFIGURATION_BUILD_DIR = \"" + EngineRelative + "Engine/Binaries/IOS/Payload\";" + ProjectFileGenerator.NewLine);
+				}
 				else if (IsAGame)
 				{
 					Contents.Append(
@@ -534,7 +543,7 @@ namespace UnrealBuildTool
 					"\t\t\t\tOBJROOT = \"XcodeSupportFiles/build\";" + ProjectFileGenerator.NewLine);
 
 				// if we're building for PC, set up Xcode to accommodate IphonePackager
-				if (bIsUE4Game)
+				if (bIsUE4Game || bIsUE4Client)
 				{
 					Contents.Append("\t\t\t\tCONFIGURATION_BUILD_DIR = \"Payload\";" + ProjectFileGenerator.NewLine);
 				}
@@ -603,7 +612,7 @@ namespace UnrealBuildTool
 		}
 
 		private void AppendSingleConfig(ref StringBuilder Contents, XcodeProjectTarget Target, string ConfigName, string ConfigGuid, string PreprocessorDefinitions, string HeaderSearchPaths,
-										string EngineRelative, string GamePath, bool bIsUE4Game, bool IsAGame)
+			string EngineRelative, string GamePath, bool bIsUE4Game, bool IsAGame, bool bIsUE4Client)
 		{
 			if (Target.Type == "Project")
 			{
@@ -625,12 +634,12 @@ namespace UnrealBuildTool
 					{
 						if (Target.Type != "XCTest")
 						{
-							AppendIOSRunConfig(ref Contents, ConfigName, ConfigGuid, Target.TargetName, EngineRelative, GamePath, bIsUE4Game, IsAGame);
+							AppendIOSRunConfig(ref Contents, ConfigName, ConfigGuid, Target.TargetName, EngineRelative, GamePath, bIsUE4Game, IsAGame, bIsUE4Client);
 						}
 						else
 						{
 							string EngineTarget = "";
-							if (bIsUE4Game)
+							if (bIsUE4Game || bIsUE4Client)
 							{
 								if (!bGeneratingGameProjectFiles && !bGeneratingRocketProjectFiles)
 								{
@@ -659,6 +668,7 @@ namespace UnrealBuildTool
 			bool IsAGame = false;
 			string GamePath = null;
 			bool bIsUE4Game = Target.TargetName.Equals("UE4Game", StringComparison.InvariantCultureIgnoreCase);
+			bool bIsUE4Client = Target.TargetName.Equals("UE4Client", StringComparison.InvariantCultureIgnoreCase);
 			string EngineRelative = "";
 
 			foreach (string GameFolder in GameFolders )
@@ -681,12 +691,12 @@ namespace UnrealBuildTool
 
 			if (!bGeneratingRocketProjectFiles)
 			{
-				AppendSingleConfig(ref Contents, Target, "Debug", Target.DebugConfigGuid, PreprocessorDefinitions, HeaderSearchPaths, EngineRelative, GamePath, bIsUE4Game, IsAGame);
-				AppendSingleConfig(ref Contents, Target, "Test", Target.TestConfigGuid, PreprocessorDefinitions, HeaderSearchPaths, EngineRelative, GamePath, bIsUE4Game, IsAGame);
+				AppendSingleConfig(ref Contents, Target, "Debug", Target.DebugConfigGuid, PreprocessorDefinitions, HeaderSearchPaths, EngineRelative, GamePath, bIsUE4Game, IsAGame, bIsUE4Client);
+				AppendSingleConfig(ref Contents, Target, "Test", Target.TestConfigGuid, PreprocessorDefinitions, HeaderSearchPaths, EngineRelative, GamePath, bIsUE4Game, IsAGame, bIsUE4Client);
 			}
-			AppendSingleConfig(ref Contents, Target, "DebugGame", Target.DebugGameConfigGuid, PreprocessorDefinitions, HeaderSearchPaths, EngineRelative, GamePath, bIsUE4Game, IsAGame);
-			AppendSingleConfig(ref Contents, Target, "Development", Target.DevelopmentConfigGuid, PreprocessorDefinitions, HeaderSearchPaths, EngineRelative, GamePath, bIsUE4Game, IsAGame);
-			AppendSingleConfig(ref Contents, Target, "Shipping", Target.ShippingConfigGuid, PreprocessorDefinitions, HeaderSearchPaths, EngineRelative, GamePath, bIsUE4Game, IsAGame);
+			AppendSingleConfig(ref Contents, Target, "DebugGame", Target.DebugGameConfigGuid, PreprocessorDefinitions, HeaderSearchPaths, EngineRelative, GamePath, bIsUE4Game, IsAGame, bIsUE4Client);
+			AppendSingleConfig(ref Contents, Target, "Development", Target.DevelopmentConfigGuid, PreprocessorDefinitions, HeaderSearchPaths, EngineRelative, GamePath, bIsUE4Game, IsAGame, bIsUE4Client);
+			AppendSingleConfig(ref Contents, Target, "Shipping", Target.ShippingConfigGuid, PreprocessorDefinitions, HeaderSearchPaths, EngineRelative, GamePath, bIsUE4Game, IsAGame, bIsUE4Client);
 		}
 
 		/// <summary>
