#UE4
- Add Google Play support on Android
- Add editor settings for Google Play - enabled/disabled, and the app id for the games service
- Hack in Google Play leaderboards, including blueprint support, proper online subsystem implementation will come in the future
- These leaderboards work in TappyChicken

diff --git a/Engine/Build/Android/Java/AndroidManifest.xml b/Engine/Build/Android/Java/AndroidManifest.xml
--- a/Engine/Build/Android/Java/AndroidManifest.xml
+++ b/Engine/Build/Android/Java/AndroidManifest.xml
@@ -21,6 +21,10 @@
 				<category android:name="android.intent.category.LAUNCHER" />
 			</intent-filter>
 		</activity>
+		<meta-data android:name="com.google.android.gms.games.APP_ID"
+			android:value="@string/app_id" />
+		<meta-data android:name="com.google.android.gms.version"
+           android:value="@integer/google_play_services_version" />
 	</application>
 	
 	<!-- Requirements -->
diff --git a/Engine/Build/Android/Java/proguard-project.txt b/Engine/Build/Android/Java/proguard-project.txt
--- a/Engine/Build/Android/Java/proguard-project.txt
+++ b/Engine/Build/Android/Java/proguard-project.txt
@@ -23,3 +23,23 @@
 #-keepclassmembers class fqcn.of.javascript.interface.for.webview {
 #   public *;
 #}
+
+# Required for Google Play services: http://developer.android.com/google/play-services/setup.html
+-keep class * extends java.util.ListResourceBundle {
+    protected Object[][] getContents();
+}
+
+-keep public class com.google.android.gms.common.internal.safeparcel.SafeParcelable {
+    public static final *** NULL;
+}
+
+-keepnames @com.google.android.gms.common.annotation.KeepName class *
+-keepclassmembernames class * {
+    @com.google.android.gms.common.annotation.KeepName *;
+}
+
+-keepnames class * implements android.os.Parcelable {
+    public static final ** CREATOR;
+}
+
+# End Google Play services requirements
diff --git a/Engine/Build/Android/Java/res/values/GooglePlayAppID.xml b/Engine/Build/Android/Java/res/values/GooglePlayAppID.xml
new file mode 100644
--- /dev/null
+++ b/Engine/Build/Android/Java/res/values/GooglePlayAppID.xml
@@ -0,0 +1,4 @@
+<?xml version="1.0" encoding="utf-8"?>
+<resources>
+	<string name="app_id"></string>
+</resources>
diff --git a/Engine/Build/Android/Java/src/com/epicgames/ue4/GameActivity.java b/Engine/Build/Android/Java/src/com/epicgames/ue4/GameActivity.java
--- a/Engine/Build/Android/Java/src/com/epicgames/ue4/GameActivity.java
+++ b/Engine/Build/Android/Java/src/com/epicgames/ue4/GameActivity.java
@@ -7,10 +7,17 @@ import android.os.Bundle;
 import android.util.Log;
 
 import android.app.AlertDialog;
+import android.app.Dialog;
 import android.widget.EditText;
 import android.content.DialogInterface;
 import android.content.Intent;
 import android.content.res.Configuration;
+import android.content.IntentSender.SendIntentException;
+
+import com.google.android.gms.common.api.GoogleApiClient;
+import com.google.android.gms.common.ConnectionResult;
+import com.google.android.gms.games.Games;
+import com.google.android.gms.common.GooglePlayServicesUtil;
 
 // TODO: use the resources from the UE4 lib project once we've got the packager up and running
 //import com.epicgames.ue4.R;
@@ -26,20 +33,37 @@ import android.content.res.Configuration;
 //  Java libraries at the startup of the program and store references 
 //  to them in this class.
 
-public class GameActivity extends NativeActivity
+public class GameActivity extends NativeActivity implements GoogleApiClient.ConnectionCallbacks, GoogleApiClient.OnConnectionFailedListener
 {
 	public static Logger Log = new Logger("UE4");
 	
 	GameActivity _activity;	
 	AlertDialog alert;
 	
+	private GoogleApiClient googleClient;
+	private boolean bResolvingGoogleServicesError = false;
+	private int dialogError = 0;
+
+	/** Request code to use when launching the Google Services resolution activity */
+    private static final int GOOGLE_SERVICES_REQUEST_RESOLVE_ERROR = 1001;
+
+	/** Unique tag for the error dialog fragment */
+    private static final String DIALOG_ERROR = "dialog_error";
+
+	/** Unique ID to identify Google Play Services error dialog */
+	private static final int PLAY_SERVICES_DIALOG_ID = 1;
+
+	@Override
 	public void onStart()
 	{
 		super.onStart();
 		
 		Log.debug("==================================> Inside onStart function in GameActivity");
+		
+		AndroidThunkJava_GooglePlayConnect();
 	}
 
+	@Override
 	public void onCreate(Bundle savedInstanceState)
 	{
 		super.onCreate(savedInstanceState);
@@ -94,8 +118,104 @@ public class GameActivity extends NativeActivity
 		alertBuilder.setTitle("Console Window - Enter Command");
 		
 		alert = alertBuilder.create();
+
+		// Connect to Google Play Services
+		googleClient = new GoogleApiClient.Builder(this)
+         .addApi(Games.API)
+         .addScope(Games.SCOPE_GAMES)
+         .addConnectionCallbacks(this)
+         .addOnConnectionFailedListener(this)
+         .build();
 	}
+
+	@Override
+	public void onStop()
+	{
+		super.onStop();
+
+		googleClient.disconnect();
+	}
+
+	// Callbacks to handle connections with Google Play
+	 @Override
+    public void onConnected(Bundle connectionHint)
+	{
+        Log.debug("Connected to Google Play Services.");
+    }
+
+    @Override
+    public void onConnectionSuspended(int cause)
+	{
+        // The connection has been interrupted.
+        // TODO: Disable any UI components that depend on Google APIs
+        // until onConnected() is called.
+		Log.debug("Google Play Services connection suspended.");
+    }
 	
+    @Override
+    public void onConnectionFailed(ConnectionResult result)
+	{
+		Log.debug("Google Play Services connection failed: " + result.toString());
+
+		if (bResolvingGoogleServicesError)
+		{
+			// Already attempting to resolve an error.
+			return;
+		}
+		else if (result.hasResolution())
+		{
+            try
+			{
+				Log.debug("Starting Google Play Services connection resolution");
+                bResolvingGoogleServicesError = true;
+                result.startResolutionForResult(this, GOOGLE_SERVICES_REQUEST_RESOLVE_ERROR);
+            }
+			catch (SendIntentException e)
+			{
+                // There was an error with the resolution intent. Try again.
+                googleClient.connect();
+            }
+        }
+		else
+		{
+            // Show dialog using GooglePlayServicesUtil.getErrorDialog()
+			dialogError = result.getErrorCode();
+			showDialog(PLAY_SERVICES_DIALOG_ID);
+
+            bResolvingGoogleServicesError = true;
+        }
+    }
+
+	@Override
+	protected Dialog onCreateDialog(int id)
+	{
+		if(id == PLAY_SERVICES_DIALOG_ID)
+		{
+			Dialog dialog = GooglePlayServicesUtil.getErrorDialog(dialogError, this, GOOGLE_SERVICES_REQUEST_RESOLVE_ERROR);
+			dialog.show();
+		}
+
+		return super.onCreateDialog(id);
+	}
+
+	@Override
+	protected void onActivityResult(int requestCode, int resultCode, Intent data)
+	{
+		if (requestCode == GOOGLE_SERVICES_REQUEST_RESOLVE_ERROR)
+		{
+			Log.debug("Google Play Services connection resolution finished with resultCode " + resultCode);
+			bResolvingGoogleServicesError = false;
+			if (resultCode == RESULT_OK)
+			{
+				// Make sure the app is not already connected or attempting to connect
+				if (!googleClient.isConnecting() &&	!googleClient.isConnected())
+				{
+					googleClient.connect();
+				}
+			}
+		}
+	}
+
 	// Called from event thread in NativeActivity	
 	public void AndroidThunkJava_ShowConsoleWindow(String Formats)
 	{
@@ -132,6 +252,32 @@ public class GameActivity extends NativeActivity
 		}
 	}
 
+	public void AndroidThunkJava_GooglePlayConnect()
+	{
+		if(nativeIsGooglePlayEnabled())
+		{
+			googleClient.connect();
+		}
+	}
+
+	public void AndroidThunkJava_ShowLeaderboard(String LeaderboardID)
+	{
+		Log.debug("In AndroidThunkJava_ShowLeaderboard, ID is " + LeaderboardID);
+		if(googleClient.isConnected())
+		{
+			startActivityForResult(Games.Leaderboards.getLeaderboardIntent(googleClient, LeaderboardID), 0);
+		}
+	}
+
+	public void AndroidThunkJava_WriteLeaderboardValue(String LeaderboardID, long Value)
+	{
+		Log.debug("In AndroidThunkJava_WriteLeaderboardValue, ID is " + LeaderboardID + ", value is " + Value);
+		if(googleClient.isConnected())
+		{
+			Games.Leaderboards.submitScore(googleClient, LeaderboardID, Value);
+		}
+	}
+
 	public native boolean nativeIsShippingBuild();
 	public native void nativeSetGlobalActivity();
 	public native void nativeSetWindowInfo(boolean bIsPortrait);
@@ -139,6 +285,8 @@ public class GameActivity extends NativeActivity
 
 	public native void nativeConsoleCommand(String commandString);
 	
+	public native boolean nativeIsGooglePlayEnabled();
+
 	static
 	{
 		System.loadLibrary("UE4");
diff --git a/Engine/Source/Developer/Android/AndroidPlatformEditor/Private/AndroidTargetSettingsCustomization.cpp b/Engine/Source/Developer/Android/AndroidPlatformEditor/Private/AndroidTargetSettingsCustomization.cpp
--- a/Engine/Source/Developer/Android/AndroidPlatformEditor/Private/AndroidTargetSettingsCustomization.cpp
+++ b/Engine/Source/Developer/Android/AndroidPlatformEditor/Private/AndroidTargetSettingsCustomization.cpp
@@ -29,6 +29,8 @@ FAndroidTargetSettingsCustomization::FAndroidTargetSettingsCustomization()
 	, GameAndroidPath(FPaths::GameDir() + TEXT("Build/Android"))
 	, EngineManifestPath(EngineAndroidPath / TEXT("AndroidManifest.xml"))
 	, GameManifestPath(GameAndroidPath / TEXT("AndroidManifest.xml"))
+	, EngineGooglePlayAppIDPath(EngineAndroidPath / TEXT("res") / TEXT("values") / TEXT("GooglePlayAppID.xml"))
+	, GameGooglePlayAppIDPath(GameAndroidPath / TEXT("res") / TEXT("values") / TEXT("GooglePlayAppID.xml"))
 	, EngineSigningConfigPath(EngineAndroidPath / TEXT("SigningConfig.xml"))
 	, GameSigningConfigPath(GameAndroidPath / TEXT("SigningConfig.xml"))
 	, EngineProguardPath(EngineAndroidPath / TEXT("proguard-project.txt"))
@@ -112,6 +114,43 @@ void FAndroidTargetSettingsCustomization::BuildAppManifestSection(IDetailLayoutB
 	OrientationProperty->SetOnPropertyValueChanged(FSimpleDelegate::CreateRaw(this, &FAndroidTargetSettingsCustomization::OnOrientationModified));
 	AppManifestCategory.AddProperty(OrientationProperty)
 		.EditCondition(SetupForPlatformAttribute, NULL);
+
+
+	// Google Play category
+	IDetailCategoryBuilder& GooglePlayCategory = DetailLayout.EditCategory(TEXT("GooglePlayServices"));
+	
+	TSharedRef<SPlatformSetupMessage> GooglePlaySetupMessage = SNew(SPlatformSetupMessage, GameGooglePlayAppIDPath)
+		.PlatformName(LOCTEXT("GooglePlayPlatformName", "Google Play services"))
+		.OnSetupClicked(this, &FAndroidTargetSettingsCustomization::CopyGooglePlayAppIDFileIntoProject);
+
+	SetupForGooglePlayAttribute = GooglePlaySetupMessage->GetReadyToGoAttribute();
+
+	GooglePlayCategory.AddCustomRow(TEXT("Warning"), false)
+		.WholeRowWidget
+		[
+			GooglePlaySetupMessage
+		];
+
+	GooglePlayCategory.AddCustomRow(TEXT("App ID Hyperlink"), false)
+		.WholeRowWidget
+		[
+			SNew(SBox)
+			.HAlign(HAlign_Center)
+			[
+				SNew(SHyperlinkLaunchURL, TEXT("http://developer.android.com/google/index.html"))
+				.Text(LOCTEXT("GooglePlayDeveloperPage", "Android Developer Page on Google Play services"))
+				.ToolTipText(LOCTEXT("GooglePlayDeveloperPageTooltip", "Opens a page that discusses Google Play services"))
+			]
+		];
+
+	TSharedRef<IPropertyHandle> EnabledProperty = DetailLayout.GetProperty(GET_MEMBER_NAME_CHECKED(UAndroidRuntimeSettings, bEnableGooglePlaySupport));
+	GooglePlayCategory.AddProperty(EnabledProperty)
+		.EditCondition(SetupForGooglePlayAttribute, NULL);
+
+	TSharedRef<IPropertyHandle> AppIDProperty = DetailLayout.GetProperty(GET_MEMBER_NAME_CHECKED(UAndroidRuntimeSettings, GamesAppID));
+	AppIDProperty->SetOnPropertyValueChanged(FSimpleDelegate::CreateRaw(this, &FAndroidTargetSettingsCustomization::OnAppIDModified));
+	GooglePlayCategory.AddProperty(AppIDProperty)
+		.EditCondition(SetupForGooglePlayAttribute, NULL);
 }
 
 void FAndroidTargetSettingsCustomization::BuildIconSection(IDetailLayoutBuilder& DetailLayout)
@@ -208,6 +247,19 @@ void FAndroidTargetSettingsCustomization::CopySetupFilesIntoProject()
 	SavedLayoutBuilder->ForceRefreshDetails();
 }
 
+void FAndroidTargetSettingsCustomization::CopyGooglePlayAppIDFileIntoProject()
+{
+	FText ErrorMessage;
+	if (!SourceControlHelpers::CopyFileUnderSourceControl(GameGooglePlayAppIDPath, EngineGooglePlayAppIDPath, LOCTEXT("GooglePlayAppID", "GooglePlayAppID.xml"), /*out*/ ErrorMessage))
+	{
+		FNotificationInfo Info(ErrorMessage);
+		Info.ExpireDuration = 3.0f;
+		FSlateNotificationManager::Get().AddNotification(Info);
+	}
+
+	SavedLayoutBuilder->ForceRefreshDetails();
+}
+
 void FAndroidTargetSettingsCustomization::OnOrientationModified()
 {
 	check(SetupForPlatformAttribute.Get() == true);
@@ -233,6 +285,21 @@ FString FAndroidTargetSettingsCustomization::OrientationToString(const EAndroidS
 	return Enum->GetMetaData(TEXT("ManifestValue"), (int32)Orientation);
 }
 
+void FAndroidTargetSettingsCustomization::OnAppIDModified()
+{
+	check(SetupForPlatformAttribute.Get() == true);
+
+
+	FManifestUpdateHelper Updater(GameGooglePlayAppIDPath);
+
+	const FString AppIDTag(TEXT("name=\"app_id\">"));
+	const FString ClosingTag(TEXT("</string>"));
+	const FString NewIDString = GetDefault<UAndroidRuntimeSettings>()->GamesAppID;
+	Updater.ReplaceKey(AppIDTag, ClosingTag, NewIDString);
+
+	Updater.Finalize(GameGooglePlayAppIDPath);
+}
+
 //////////////////////////////////////////////////////////////////////////
 
 #undef LOCTEXT_NAMESPACE
\ No newline at end of file
diff --git a/Engine/Source/Developer/Android/AndroidPlatformEditor/Private/AndroidTargetSettingsCustomization.h b/Engine/Source/Developer/Android/AndroidPlatformEditor/Private/AndroidTargetSettingsCustomization.h
--- a/Engine/Source/Developer/Android/AndroidPlatformEditor/Private/AndroidTargetSettingsCustomization.h
+++ b/Engine/Source/Developer/Android/AndroidPlatformEditor/Private/AndroidTargetSettingsCustomization.h
@@ -31,8 +31,15 @@ private:
 	// Copies the setup files for the platform into the project
 	void CopySetupFilesIntoProject();
 
+	// Copies the strings.xml file for the platform into the project
+	void CopyGooglePlayAppIDFileIntoProject();
+
 	// Called when the orientation is modified
 	void OnOrientationModified();
+
+	// Called when the app id is modified
+	void OnAppIDModified();
+
 private:
 	const FString AndroidRelativePath;
 
@@ -41,6 +48,9 @@ private:
 
 	const FString EngineManifestPath;
 	const FString GameManifestPath;
+	
+	const FString EngineGooglePlayAppIDPath;
+	const FString GameGooglePlayAppIDPath;
 
 	const FString EngineSigningConfigPath;
 	const FString GameSigningConfigPath;
@@ -56,6 +66,9 @@ private:
 	// Is the manifest writable?
 	TAttribute<bool> SetupForPlatformAttribute;
 
+	// Is the App ID string writable?
+	TAttribute<bool> SetupForGooglePlayAttribute;
+
 	// Converts an orientation enum to the associated string value
 	static FString OrientationToString(const EAndroidScreenOrientation::Type Orientation);
 
diff --git a/Engine/Source/Programs/UnrealBuildTool/Android/UEDeployAndroid.cs b/Engine/Source/Programs/UnrealBuildTool/Android/UEDeployAndroid.cs
--- a/Engine/Source/Programs/UnrealBuildTool/Android/UEDeployAndroid.cs
+++ b/Engine/Source/Programs/UnrealBuildTool/Android/UEDeployAndroid.cs
@@ -195,11 +195,37 @@ namespace UnrealBuildTool.Android
 			}
 		}
 
+		private void DeleteDirectory(string Path)
+		{
+			Console.WriteLine("\nDeleting: " + Path);
+			if (Directory.Exists(Path))
+			{
+				try
+				{
+					// first make sure everything it writable
+					string[] Files = Directory.GetFiles(Path, "*.*", SearchOption.AllDirectories);
+					foreach (string Filename in Files)
+					{
+						// remove any read only flags
+						FileInfo FileInfo = new FileInfo(Filename);
+						FileInfo.Attributes = FileInfo.Attributes & ~FileAttributes.ReadOnly;
+					}
+
+					// now deleting will work better
+					Directory.Delete(Path, true);
+				}
+				catch (Exception)
+				{
+					Log.TraceInformation("Failed to delete intermediate cdirectory {0}. Continuing on...", Path);
+				}
+			}
+		}
+
 		private void MakeAPK(string ProjectName, string ProjectDirectory, string OutputPath, string EngineDirectory, bool bForDistribution)
 		{
 			// cache some build product paths
 			string SourceSOName = OutputPath;
-			string DestApkName = ProjectDirectory + "/Binaries/Android/" + Path.GetFileNameWithoutExtension(SourceSOName) + ".apk";
+			string DestApkName = Path.Combine(ProjectDirectory, "Binaries/Android/") + Path.GetFileNameWithoutExtension(SourceSOName) + ".apk";
 
 			// if the source binary was UE4Game, replace it with the new project name, when re-packaging a binary only build
 			DestApkName = DestApkName.Replace("UE4Game-", ProjectName + "-");
@@ -210,10 +236,10 @@ namespace UnrealBuildTool.Android
 			string AntBuildPath = GetAntPath();
 
 			// set up some directory info
-			string UE4BuildPath = ProjectDirectory + "/Intermediate/Android/APK";
+			string IntermediateAndroidPath = Path.Combine(ProjectDirectory, "Intermediate/Android/");
+			string UE4BuildPath = IntermediateAndroidPath + "APK";
 			string UE4BuildFilesPath = Path.GetFullPath(Path.Combine(EngineDirectory, "Build/Android/Java"));
-			string GameBuildFilesPath = ProjectDirectory + "/Build/Android";
-
+			string GameBuildFilesPath = Path.Combine(ProjectDirectory, "Build/Android");
 
 			// check to see if it's out of date before trying the slow make apk process (look at .so and all Engine and Project build files to be safe)
 			List<String> InputFiles = new List<string>();
@@ -246,28 +272,7 @@ namespace UnrealBuildTool.Android
 
 			
 			//Wipe the Intermediate/Build/APK directory first
-			Console.WriteLine("\nDeleting: " + UE4BuildPath);
-			if (Directory.Exists(UE4BuildPath))
-			{
-				try
-				{
-					// first make sure everything it writable
-					string[] Files = Directory.GetFiles(UE4BuildPath, "*.*", SearchOption.AllDirectories);
-					foreach (string Filename in Files)
-					{
-						// remove any read only flags
-						FileInfo FileInfo = new FileInfo(Filename);
-						FileInfo.Attributes = FileInfo.Attributes & ~FileAttributes.ReadOnly;
-					}
-
-					// now deleting will work better
-					Directory.Delete(UE4BuildPath, true);
-				}
-				catch (Exception)
-				{
-					Log.TraceInformation("Failed to delete intermediate cdirectory {0}. Continuing on...", UE4BuildPath);
-				}
-			}
+			DeleteDirectory(UE4BuildPath);
 
 			Dictionary<string, string> Replacements = new Dictionary<string, string>();
 			Replacements.Add("${EXECUTABLE_NAME}", ProjectName);
@@ -316,6 +321,53 @@ namespace UnrealBuildTool.Android
 				throw new BuildException("android.bat failed [{0}]", AndroidBatStartInfoGame.Arguments);
 			}
 
+            // Update the Google Play services lib with the target platform version currently in use.
+            // This appears to be required for the build to work without errors when Play services are referenced by the game.
+            // This will try to modify existing files, like project.properties, so we copy the entire library into
+            // an intermediate directory and work from there.
+            string GooglePlayServicesSourcePath = Path.GetFullPath(Path.Combine(EngineDirectory, "Source/ThirdParty/Android/google-play-services_lib/"));
+			string GooglePlayServicesIntermediatePath = Path.GetFullPath(Path.Combine(IntermediateAndroidPath, "google-play-services_lib/"));
+
+			DeleteDirectory(GooglePlayServicesIntermediatePath);
+			CopyFileDirectory(GooglePlayServicesSourcePath, GooglePlayServicesIntermediatePath, new Dictionary<string, string>());
+		    
+			ProcessStartInfo AndroidBatStartInfoPlayServicesLib = new ProcessStartInfo();
+			AndroidBatStartInfoPlayServicesLib.WorkingDirectory = GooglePlayServicesIntermediatePath;
+			AndroidBatStartInfoPlayServicesLib.FileName = AndroidCommandPath;
+			AndroidBatStartInfoPlayServicesLib.Arguments = "update project " + " --path . --target " + GetSdkApiLevel();
+			AndroidBatStartInfoPlayServicesLib.UseShellExecute = false;
+			Console.WriteLine("\nRunning: " + AndroidBatStartInfoPlayServicesLib.FileName + " " + AndroidBatStartInfoPlayServicesLib.Arguments);
+			Process AndroidBatPlayServicesLib = new Process();
+			AndroidBatPlayServicesLib.StartInfo = AndroidBatStartInfoPlayServicesLib;
+			AndroidBatPlayServicesLib.Start();
+			AndroidBatPlayServicesLib.WaitForExit();
+
+            // android bat failure
+            if (AndroidBatPlayServicesLib.ExitCode != 0)
+            {
+                throw new BuildException("android.bat failed [{0}]", AndroidBatStartInfoPlayServicesLib.Arguments);
+            }
+
+            //need to create separate run for each lib. Will be added to project.properties in order in which they are added 
+            //the order is important.
+            //as android.library.reference.X=libpath where X = 1 - N
+            //for e.g this one will be added as android.library.reference.1=<EngineDirectory>/Source/ThirdParty/Android/google_play_services_lib
+            
+            // Ant seems to need a relative path to work
+			Uri ServicesBuildUri = new Uri(GooglePlayServicesIntermediatePath);
+            Uri ProjectUri = new Uri(UE4BuildPath + "/");
+			string RelativeServicesUri = ProjectUri.MakeRelativeUri(ServicesBuildUri).ToString();
+
+            AndroidBatStartInfoGame.Arguments = " update project --name " + ProjectName + " --path .  --target " + GetSdkApiLevel() + " --library " + RelativeServicesUri;
+            Console.WriteLine("\nRunning: " + AndroidBatStartInfoGame.FileName + " " + AndroidBatStartInfoGame.Arguments);
+            AndroidBatGame.Start();
+            AndroidBatGame.WaitForExit();
+
+            if (AndroidBatGame.ExitCode != 0)
+            {
+                throw new BuildException("android.bat failed [{0}]", AndroidBatStartInfoGame.Arguments);
+            }
+
 			// Use ndk-build to do stuff and move the .so file to the lib folder (only if NDK is installed)
 
 			string FinalSOName = "";
diff --git a/Engine/Source/Runtime/Android/AndroidRuntimeSettings/Classes/AndroidRuntimeSettings.h b/Engine/Source/Runtime/Android/AndroidRuntimeSettings/Classes/AndroidRuntimeSettings.h
--- a/Engine/Source/Runtime/Android/AndroidRuntimeSettings/Classes/AndroidRuntimeSettings.h
+++ b/Engine/Source/Runtime/Android/AndroidRuntimeSettings/Classes/AndroidRuntimeSettings.h
@@ -48,4 +48,11 @@ public:
 	UPROPERTY(GlobalConfig, EditAnywhere, Category = AppManifest)
 	TEnumAsByte<EAndroidScreenOrientation::Type> Orientation;
 
+	// Should Google Play support be enabled?
+	UPROPERTY(GlobalConfig, EditAnywhere, Category = GooglePlayServices)
+	bool bEnableGooglePlaySupport;
+
+	// The app id obtained from the Google Play Developer Console
+	UPROPERTY(GlobalConfig, EditAnywhere, Category = GooglePlayServices)
+	FString GamesAppID;
 };
diff --git a/Engine/Source/Runtime/Android/AndroidRuntimeSettings/Private/AndroidRuntimeSettings.cpp b/Engine/Source/Runtime/Android/AndroidRuntimeSettings/Private/AndroidRuntimeSettings.cpp
--- a/Engine/Source/Runtime/Android/AndroidRuntimeSettings/Private/AndroidRuntimeSettings.cpp
+++ b/Engine/Source/Runtime/Android/AndroidRuntimeSettings/Private/AndroidRuntimeSettings.cpp
@@ -7,6 +7,7 @@
 UAndroidRuntimeSettings::UAndroidRuntimeSettings(const class FPostConstructInitializeProperties& PCIP)
 	: Super(PCIP)
 	, Orientation(EAndroidScreenOrientation::Landscape)
+	, bEnableGooglePlaySupport(false)
 {
 
 }
diff --git a/Engine/Source/Runtime/Engine/Private/KismetSystemLibrary.cpp b/Engine/Source/Runtime/Engine/Private/KismetSystemLibrary.cpp
--- a/Engine/Source/Runtime/Engine/Private/KismetSystemLibrary.cpp
+++ b/Engine/Source/Runtime/Engine/Private/KismetSystemLibrary.cpp
@@ -2255,5 +2255,11 @@ void UKismetSystemLibrary::EXPERIMENTAL_ShowGameCenterLeaderboard(const FString&
 {
 #if PLATFORM_IOS
 	IOSShowLeaderboardUI(CategoryName);
+#elif PLATFORM_ANDROID
+	UE_LOG(LogTemp, Log, TEXT("Showing leaderboard!"));
+
+	extern void AndroidThunkCpp_ShowLeaderboard(const FString&);
+	extern FString GetLeaderboardID(const FString&);
+	AndroidThunkCpp_ShowLeaderboard(GetLeaderboardID(CategoryName));
 #endif
 }
\ No newline at end of file
diff --git a/Engine/Source/Runtime/Launch/Private/Android/AndroidJNI.cpp b/Engine/Source/Runtime/Launch/Private/Android/AndroidJNI.cpp
--- a/Engine/Source/Runtime/Launch/Private/Android/AndroidJNI.cpp
+++ b/Engine/Source/Runtime/Launch/Private/Android/AndroidJNI.cpp
@@ -94,6 +94,9 @@ JNIEnv* GetJavaEnv(bool bRequireGlobalThis)
 jclass JDef_GameActivity::ClassID;
 jmethodID JDef_GameActivity::AndroidThunkJava_ShowConsoleWindow;
 jmethodID JDef_GameActivity::AndroidThunkJava_LaunchURL;
+jmethodID JDef_GameActivity::AndroidThunkJava_ShowLeaderboard;
+jmethodID JDef_GameActivity::AndroidThunkJava_WriteLeaderboardValue;
+jmethodID JDef_GameActivity::AndroidThunkJava_GooglePlayConnect;
 
 DEFINE_LOG_CATEGORY_STATIC(LogEngine, Log, All);
 
@@ -161,6 +164,27 @@ void AndroidThunkCpp_LaunchURL(const FString& URL)
 	}
 }
 
+void AndroidThunkCpp_ShowLeaderboard(const FString& CategoryName)
+{
+	if (JNIEnv* Env = GetJavaEnv())
+	{
+		jstring Argument = Env->NewStringUTF(TCHAR_TO_UTF8(*CategoryName));
+		Env->CallVoidMethod(GJavaGlobalThis, JDef_GameActivity::AndroidThunkJava_ShowLeaderboard, Argument);
+		Env->DeleteLocalRef(Argument);
+	}
+}
+
+void AndroidThunkCpp_WriteLeaderboardValue(const FString& LeaderboardName, int64_t Value)
+{
+	if (JNIEnv* Env = GetJavaEnv())
+	{
+		jstring LeaderboardNameArg = Env->NewStringUTF(TCHAR_TO_UTF8(*LeaderboardName));
+		Env->CallVoidMethod(GJavaGlobalThis, JDef_GameActivity::AndroidThunkJava_WriteLeaderboardValue, LeaderboardNameArg, Value);
+		Env->DeleteLocalRef(LeaderboardNameArg);
+	}
+}
+
+
 //The JNI_OnLoad function is triggered by loading the game library from 
 //the Java source file.
 //	static
@@ -184,6 +208,9 @@ JNIEXPORT jint JNI_OnLoad(JavaVM* InJavaVM, void* InReserved)
 
 	JDef_GameActivity::AndroidThunkJava_ShowConsoleWindow = env->GetMethodID(JDef_GameActivity::ClassID, "AndroidThunkJava_ShowConsoleWindow", "(Ljava/lang/String;)V");
 	JDef_GameActivity::AndroidThunkJava_LaunchURL = env->GetMethodID(JDef_GameActivity::ClassID, "AndroidThunkJava_LaunchURL", "(Ljava/lang/String;)V");
+	JDef_GameActivity::AndroidThunkJava_ShowLeaderboard = env->GetMethodID(JDef_GameActivity::ClassID, "AndroidThunkJava_ShowLeaderboard", "(Ljava/lang/String;)V");
+	JDef_GameActivity::AndroidThunkJava_WriteLeaderboardValue = env->GetMethodID(JDef_GameActivity::ClassID, "AndroidThunkJava_WriteLeaderboardValue", "(Ljava/lang/String;J)V");
+	JDef_GameActivity::AndroidThunkJava_GooglePlayConnect = env->GetMethodID(JDef_GameActivity::ClassID, "AndroidThunkJava_GooglePlayConnect", "()V");
 
 	// hook signals
 #if UE_BUILD_DEBUG
diff --git a/Engine/Source/Runtime/Launch/Private/Android/AndroidJNI.h b/Engine/Source/Runtime/Launch/Private/Android/AndroidJNI.h
--- a/Engine/Source/Runtime/Launch/Private/Android/AndroidJNI.h
+++ b/Engine/Source/Runtime/Launch/Private/Android/AndroidJNI.h
@@ -21,6 +21,9 @@ public:
 	// Nonstatic methods
 	static jmethodID AndroidThunkJava_ShowConsoleWindow;
 	static jmethodID AndroidThunkJava_LaunchURL;
+	static jmethodID AndroidThunkJava_ShowLeaderboard;
+	static jmethodID AndroidThunkJava_WriteLeaderboardValue;
+	static jmethodID AndroidThunkJava_GooglePlayConnect;
 };
 
 // Returns the java environment
diff --git a/Engine/Source/Runtime/Launch/Private/Android/LaunchAndroid.cpp b/Engine/Source/Runtime/Launch/Private/Android/LaunchAndroid.cpp
--- a/Engine/Source/Runtime/Launch/Private/Android/LaunchAndroid.cpp
+++ b/Engine/Source/Runtime/Launch/Private/Android/LaunchAndroid.cpp
@@ -223,6 +223,13 @@ int32 AndroidMain(struct android_app* state)
 
 	UE_LOG(LogAndroid, Log, TEXT("Passed GEngineLoop.Init()"));
 
+	// Hack to initialize Google Play if enabled until we get the full subsystem online.
+	// GConfig should be valid here.
+	if (JNIEnv* Env = GetJavaEnv())
+	{
+		Env->CallVoidMethod(GJavaGlobalThis, JDef_GameActivity::AndroidThunkJava_GooglePlayConnect);
+	}
+
 	// tick until done
 	while (!GIsRequestingExit)
 	{
@@ -660,3 +667,12 @@ extern "C" void Java_com_epicgames_ue4_GameActivity_nativeConsoleCommand(JNIEnv*
 	//Release the string
 	jenv->ReleaseStringUTFChars(commandString, javaChars);
 }
+
+//This function is declared in the Java-defined class, GameActivity.java: "public native void nativeIsGooglePlayEnabled(String commandString);"
+extern "C" jboolean Java_com_epicgames_ue4_GameActivity_nativeIsGooglePlayEnabled(JNIEnv* jenv, jobject thiz)
+{
+	bool bEnabled = true;
+	GConfig->GetBool(TEXT("/Script/AndroidRuntimeSettings.AndroidRuntimeSettings"), TEXT("bEnableGooglePlaySupport"), bEnabled, GEngineIni);
+	UE_LOG(LogOnline, Log, TEXT("Checking whether Google Play is enabled. bEnableGooglePlaySupport = %d"), bEnabled);
+	return bEnabled;
+}
diff --git a/Engine/Source/Runtime/Online/OnlineSubsystemUtils/Private/LeaderboardBlueprintLibrary.cpp b/Engine/Source/Runtime/Online/OnlineSubsystemUtils/Private/LeaderboardBlueprintLibrary.cpp
--- a/Engine/Source/Runtime/Online/OnlineSubsystemUtils/Private/LeaderboardBlueprintLibrary.cpp
+++ b/Engine/Source/Runtime/Online/OnlineSubsystemUtils/Private/LeaderboardBlueprintLibrary.cpp
@@ -4,6 +4,65 @@
 #include "LeaderboardBlueprintLibrary.h"
 #include "Runtime/Online/OnlineSubsystem/Public/Interfaces/OnlineLeaderboardInterface.h"
 
+
+#if PLATFORM_ANDROID
+/**
+ * Class to store the leaderboard config.
+ * This really belongs in the proper OSS, move it there when we bring the whole OSS online.
+ */
+class FLeaderboardsConfig : public FOnlineJsonSerializable
+{
+public:
+	FJsonSerializableKeyValueMap LeaderboardMap;
+
+	BEGIN_ONLINE_JSON_SERIALIZER
+	ONLINE_JSON_SERIALIZE_MAP( "LeaderboardMap", LeaderboardMap );
+	END_ONLINE_JSON_SERIALIZER
+};
+
+/**
+ * Helper function to get the platform- and game-specific leaderboard ID from the JSON config file.
+ * This really belongs in the proper OSS class, move it there when we bring the whole OSS online.
+ */
+FString GetLeaderboardID(const FString& LeaderboardName)
+{
+	// Hacks to initialize the config one
+	static bool bHasReadConfig = false;
+	static FLeaderboardsConfig LeaderboardConfig;
+
+	if(!bHasReadConfig)
+	{
+		const FString JSonConfigFilename = FPaths::GameDir() + TEXT( "Config/OSS/GooglePlay/Leaderboards.json" );
+
+		FString JSonText;
+
+		if ( !FFileHelper::LoadFileToString( JSonText, *JSonConfigFilename ) )
+		{
+			UE_LOG(LogOnline, Warning, TEXT( "GetLeaderboardID: Failed to find json OSS leaderboard config: %s"), *JSonConfigFilename );
+			return LeaderboardName;
+		}
+
+		if ( !LeaderboardConfig.FromJson( JSonText ) )
+		{
+			UE_LOG(LogOnline, Warning, TEXT( "GetLeaderboardID: Failed to parse json OSS leaderboard config: %s"), *JSonConfigFilename );
+			return LeaderboardName;
+		}
+
+		bHasReadConfig = true;
+	}
+
+	FString * LeaderboardID = LeaderboardConfig.LeaderboardMap.Find( LeaderboardName );
+
+	if(LeaderboardID == NULL)
+	{
+		UE_LOG(LogOnline, Warning, TEXT( "GetLeaderboardID: No mapping for leaderboard %s"), *LeaderboardName );
+		return LeaderboardName;
+	}
+
+	return *LeaderboardID;
+}
+#endif
+
 //////////////////////////////////////////////////////////////////////////
 // ULeaderboardBlueprintLibrary
 
@@ -14,6 +73,41 @@ ULeaderboardBlueprintLibrary::ULeaderboardBlueprintLibrary(const FPostConstructI
 
 bool ULeaderboardBlueprintLibrary::WriteLeaderboardObject(APlayerController* PlayerController, class FOnlineLeaderboardWrite& WriteObject)
 {
+#if PLATFORM_ANDROID
+	// Hack in Android leaderboard write since we don't have the full OSS yet.
+	UE_LOG(LogOnline, Log, TEXT("Leaderboard stat name is %s"), *WriteObject.RatedStat.ToString());
+
+	for (FStatPropertyArray::TConstIterator It(WriteObject.Properties); It; ++It)
+	{
+		// Access the stat and the value.
+		const FVariantData& Stat = It.Value();
+
+		FString LeaderboardName(It.Key().ToString());
+		
+		int32 Value = 0;
+
+		switch (Stat.GetType())
+		{
+		case EOnlineKeyValuePairDataType::Int32:
+			{
+				Stat.GetValue(Value);
+				UE_LOG(LogOnline, Log, TEXT("Writing Int32 leaderboard stat %s with value %d."), *LeaderboardName, Value);
+				extern void AndroidThunkCpp_WriteLeaderboardValue(const FString&, int64_t);
+				AndroidThunkCpp_WriteLeaderboardValue(GetLeaderboardID(LeaderboardName), Value);
+				return true;
+			}
+
+		default:
+			{
+				UE_LOG(LogOnline, Warning, TEXT("ULeaderboardBlueprintLibrary::WriteLeaderboardObject(Leaderboard: %s) Invalid data type (only Int32 is currently supported)"), *LeaderboardName);
+				break;
+			}
+		}
+	}
+
+	return false;
+
+#else
 	if (APlayerState* PlayerState = (PlayerController != NULL) ? PlayerController->PlayerState : NULL)
 	{
 		TSharedPtr<FUniqueNetId> UserId = PlayerState->UniqueId.GetUniqueNetId();
@@ -53,6 +147,7 @@ bool ULeaderboardBlueprintLibrary::WriteLeaderboardObject(APlayerController* Pla
 	}
 
 	return false;
+#endif
 }
 
 bool ULeaderboardBlueprintLibrary::WriteLeaderboardInteger(APlayerController* PlayerController, FName StatName, int32 StatValue)
