#ue4
#ios
* fix for entitlements for stub files

TTP333156

diff --git a/Engine/Build/IOS/UE4Game.entitlements b/Engine/Build/IOS/UE4Game.entitlements
deleted file mode 100644
diff --git a/Engine/Source/Programs/UnrealBuildTool/IOS/UEDeployIOS.cs b/Engine/Source/Programs/UnrealBuildTool/IOS/UEDeployIOS.cs
--- a/Engine/Source/Programs/UnrealBuildTool/IOS/UEDeployIOS.cs
+++ b/Engine/Source/Programs/UnrealBuildTool/IOS/UEDeployIOS.cs
@@ -6,9 +6,6 @@ using System.Linq;
 using System.Text;
 using System.Xml;
 using System.IO;
-using System.Diagnostics;
-using Ionic.Zip;
-using Ionic.Zlib;
 
 namespace UnrealBuildTool.IOS
 {
@@ -68,32 +65,6 @@ namespace UnrealBuildTool.IOS
 			}
 		}
 
-		/// <summary>
-		/// Process.OutputDataReceived event handler.
-		/// </summary>
-		/// <param name="sender">Sender</param>
-		/// <param name="e">Event args</param>
-		private void StdOut(object sender, DataReceivedEventArgs e)
-		{
-			if (e.Data != null)
-			{
-				Console.WriteLine (e.Data);
-			}
-		}
-
-		/// <summary>
-		/// Process.ErrorDataReceived event handler.
-		/// </summary>
-		/// <param name="sender">Sender</param>
-		/// <param name="e">Event args</param>
-		private void StdErr(object sender, DataReceivedEventArgs e)
-		{
-			if (e.Data != null)
-			{
-				Console.WriteLine ("Error: " + e.Data);
-			}
-		}
-
 		public override bool PrepForUATPackageOrDeploy(string InProjectName, string InProjectDirectory, string InExecutablePath, string InEngineDir, bool bForDistribution)
 		{
 			bool bIsUE4Game = InExecutablePath.Contains ("UE4Game");
@@ -151,11 +122,6 @@ namespace UnrealBuildTool.IOS
 				File.Copy (ProvisionWithPrefix, Environment.GetEnvironmentVariable ("HOME") + "/Library/MobileDevice/Provisioning Profiles/" + InProjectName + ".mobileprovision", true);
 				FileInfo DestFileInfo = new FileInfo (Environment.GetEnvironmentVariable ("HOME") + "/Library/MobileDevice/Provisioning Profiles/" + InProjectName + ".mobileprovision");
 				DestFileInfo.Attributes = DestFileInfo.Attributes & ~FileAttributes.ReadOnly;
-
-				if (BuildConfiguration.bCreateStubIPA)
-				{
-					File.Copy (ProvisionWithPrefix, AppDirectory + "/embedded.mobileprovision");
-				}
 			}
 
 			// install the distribution provision
@@ -178,12 +144,6 @@ namespace UnrealBuildTool.IOS
 				PListFile = BuildDirectory + "/" + InProjectName + "-Info.plist";
 			}
 
-			string Entitlements = InEngineDir + "/Build/IOS/UE4Game.entitlements";
-			if (File.Exists(BuildDirectory + "/" + InProjectName + ".entitlements"))
-			{
-				Entitlements = BuildDirectory + "/" + InProjectName + ".entitlements";
-			}
-
 			Dictionary<string, string> Replacements = new Dictionary<string, string>();
 			Replacements.Add("${EXECUTABLE_NAME}", GameName);
 			Replacements.Add("${BUNDLE_IDENTIFIER}", InProjectName.Replace("_", ""));
