#ttp 326117 - Output an error message in UBT when trying to build an invalid configuration
#proj UE4
#branch UE4
#summary Done.

diff --git a/Engine/Source/Programs/UnrealBuildTool/System/RulesCompiler.cs b/Engine/Source/Programs/UnrealBuildTool/System/RulesCompiler.cs
--- a/Engine/Source/Programs/UnrealBuildTool/System/RulesCompiler.cs
+++ b/Engine/Source/Programs/UnrealBuildTool/System/RulesCompiler.cs
@@ -453,8 +453,8 @@ namespace UnrealBuildTool
 
 		public bool SupportsPlatform(UnrealTargetPlatform InPlatform)
 		{
-			// Win32 is not supported by the editor.
-			if (IsEditorType(Type) && InPlatform == UnrealTargetPlatform.Win32)
+			// Only Win64 and Mac are supported by the editor.
+			if (IsEditorType(Type) && InPlatform != UnrealTargetPlatform.Win64 && InPlatform != UnrealTargetPlatform.Mac)
 			{
 				return false;
 			}
