Merging UE4-Pretest @ 2042161 to UE4

diff --git a/Engine/Source/Editor/UnrealEd/Private/Kismet2/StructureEditorUtils.cpp b/Engine/Source/Editor/UnrealEd/Private/Kismet2/StructureEditorUtils.cpp
--- a/Engine/Source/Editor/UnrealEd/Private/Kismet2/StructureEditorUtils.cpp
+++ b/Engine/Source/Editor/UnrealEd/Private/Kismet2/StructureEditorUtils.cpp
@@ -10,142 +10,38 @@
 #include "Editor/UnrealEd/Public/Kismet2/CompilerResultsLog.h"
 #include "Editor/KismetCompiler/Public/KismetCompilerModule.h"
 
-
 #define LOCTEXT_NAMESPACE "Structure"
 
-FName FStructureEditorUtils::MakeUniqueStructName(UBlueprint* Blueprint, const FString& BaseName)
-{
-	FString StructNameString(BaseName);
-	FName StructName(*StructNameString);
-	FKismetNameValidator NameValidator(Blueprint);
-	int32 Index = 0;
-	while ((NameValidator.IsValid(StructName) != EValidatorResult::Ok) 
-		|| !IsUniqueObjectName(StructName, Blueprint->GetOutermost()))
-	{
-		StructName = FName(*FString::Printf(TEXT("%s%i"), *StructNameString, Index));
-		++Index;
-	}
+//////////////////////////////////////////////////////////////////////////
+// FEnumEditorManager
+FStructureEditorUtils::FStructEditorManager::FStructEditorManager() : FListenerManager<UUserDefinedStruct>()
+{}
 
-	return StructName;
-}
-
-bool FStructureEditorUtils::AddStructure(UBlueprint* Blueprint)
+FStructureEditorUtils::FStructEditorManager& FStructureEditorUtils::FStructEditorManager::Get()
 {
-	if(NULL != Blueprint)
-	{
-		const FScopedTransaction Transaction( LOCTEXT("AddStructure", "Add Structure") );
-		Blueprint->Modify();
-
-		const FName StructName = MakeUniqueStructName(Blueprint, TEXT("NewStruct"));
-		FBPStructureDescription StructureDesc;
-		StructureDesc.Name = StructName;
-
-		Blueprint->UserDefinedStructures.Add(StructureDesc);
-		OnStructureChanged(StructureDesc, Blueprint);
-
-		return true;
-	}
-	return false;
+	static TSharedRef< FStructEditorManager > EditorManager( new FStructEditorManager() );
+	return *EditorManager;
 }
 
-bool FStructureEditorUtils::RemoveStructure(UBlueprint* Blueprint, FName StructName)
+//////////////////////////////////////////////////////////////////////////
+// FStructureEditorUtils
+UUserDefinedStruct* FStructureEditorUtils::CreateUserDefinedStruct(UObject* InParent, FName Name, EObjectFlags Flags)
 {
-	if(NULL != Blueprint)
+	UUserDefinedStruct* Struct = NULL;
+	
+	if (UserDefinedStructEnabled())
 	{
-		const FScopedTransaction Transaction( LOCTEXT("RemoveStructure", "Remove Structure") );
-		Blueprint->Modify();
-
-		const int32 DescIdx = Blueprint->UserDefinedStructures.IndexOfByPredicate(FFindByNameHelper(StructName));
-		if(INDEX_NONE != DescIdx)
-		{
-			if(auto Struct = Blueprint->UserDefinedStructures[DescIdx].CompiledStruct)
-			{
-				FWeakObjectPtr StructWeakPtr(Struct);
-				Blueprint->UserDefinedStructures[DescIdx].CompiledStruct = NULL;
-
-				TArray< UObject* > ObjectsToDelete;
-				ObjectsToDelete.Add(Struct);
-				if(!ObjectTools::DeleteObjects(ObjectsToDelete, false))
-				{
-					Blueprint->UserDefinedStructures[DescIdx].CompiledStruct = Struct;
-					return false;
-				}
-
-				if (auto ZombieStruct = StructWeakPtr.Get())
-				{
-					ZombieStruct->Rename(NULL, GetTransientPackage(), REN_DontCreateRedirectors);
-					ULinkerLoad::InvalidateExport(ZombieStruct);
-					ZombieStruct->MarkPendingKill();
-				}
-			}
-
-			Blueprint->UserDefinedStructures.RemoveAt(DescIdx);
-			FBlueprintEditorUtils::MarkBlueprintAsModified(Blueprint);
-			return true;
-		}
+		Struct = NewNamedObject<UUserDefinedStruct>(InParent, Name, Flags);
+		check(Struct);
+		Struct->EditorData = NewNamedObject<UUserDefinedStructEditorData>(Struct, NAME_None, RF_Transactional);
+		check(Struct->EditorData);
+		Struct->SetMetaData(TEXT("BlueprintType"), TEXT("true"));
+		Struct->Bind();
+		Struct->StaticLink(true);
+		Struct->Status = UDSS_Error;
 	}
-	return false;
-}
-
-bool FStructureEditorUtils::RenameStructure(UBlueprint* Blueprint, FName OldStructName, const FString& NewNameStr)
-{
-	if(NULL != Blueprint)
-	{
-		const FScopedTransaction Transaction( LOCTEXT("RenameStructure", "Rename Structure") );
-		Blueprint->Modify();
 
-		const FName NewName(*NewNameStr);
-		FKismetNameValidator NameValidator(Blueprint);
-		if(NewName.IsValidXName(INVALID_OBJECTNAME_CHARACTERS) && (NameValidator.IsValid(NewName) == EValidatorResult::Ok))
-		{
-			FBPStructureDescription* StructureDesc = Blueprint->UserDefinedStructures.FindByPredicate(FFindByNameHelper(OldStructName));
-			if(StructureDesc)
-			{
-				if(StructureDesc->CompiledStruct)
-				{
-					if(!StructureDesc->CompiledStruct->Rename(*NewNameStr, NULL, REN_Test))
-					{
-						return false;
-					}
-
-					if(!StructureDesc->CompiledStruct->Rename(*NewNameStr, NULL))
-					{
-						return false;
-					}
-				}
-
-				StructureDesc->Name = NewName;
-
-				FBlueprintEditorUtils::MarkBlueprintAsModified(Blueprint);
-				return true;
-			}
-		}
-	}
-	return false;
-}
-
-bool FStructureEditorUtils::DuplicateStructure(UBlueprint* Blueprint, FName StructName)
-{
-	if(NULL != Blueprint)
-	{
-		const FScopedTransaction Transaction( LOCTEXT("RemoveStructure", "Remove Structure") );
-		Blueprint->Modify();
-
-		const FBPStructureDescription* StructureDesc = Blueprint->UserDefinedStructures.FindByPredicate(FFindByNameHelper(StructName));
-		if (StructureDesc)
-		{
-			FBPStructureDescription NewStructureDesc(*StructureDesc);
-			NewStructureDesc.CompiledStruct = NULL;
-			NewStructureDesc.Name = MakeUniqueStructName(Blueprint, StructureDesc->Name.ToString());
-
-			Blueprint->UserDefinedStructures.Add(NewStructureDesc);
-			OnStructureChanged(NewStructureDesc, Blueprint);
-
-			return true;
-		}
-	}
-
-	return false;
+	return Struct;
 }
 
 FStructureEditorUtils::EStructureError FStructureEditorUtils::IsStructureValid(const UScriptStruct* Struct, const UStruct* RecursionParent, FString* OutMsg)
@@ -171,19 +67,18 @@ FStructureEditorUtils::EStructureError FStructureEditorUtils::IsStructureValid(c
 		return EStructureError::FallbackStruct;
 	}
 
-	if (const UBlueprintGeneratedStruct* BPGStruct = Cast<const UBlueprintGeneratedStruct>(Struct))
+	if (Struct->GetStructureSize() <= 0)
 	{
-		static const FName BlueprintTypeName(TEXT("BlueprintType"));
-		if (!Struct->GetBoolMetaData(BlueprintTypeName))
+		if (OutMsg)
 		{
-			if (OutMsg)
-			{
-				*OutMsg = FString::Printf(*LOCTEXT("StructureNotBlueprintType", "Struct '%s' is not BlueprintType").ToString(), *Struct->GetFullName());
-			}
-			return EStructureError::NotBlueprintType;
+			*OutMsg = FString::Printf(*LOCTEXT("StructureSizeIsZero", "Struct '%s' is empty").ToString(), *Struct->GetFullName());
 		}
+		return EStructureError::EmptyStructure;
+	}
 
-		if (BPGStruct->Status != EBlueprintStructureStatus::BSS_UpToDate)
+	if (const UUserDefinedStruct* UDStruct = Cast<const UUserDefinedStruct>(Struct))
+	{
+		if (UDStruct->Status != EUserDefinedStructureStatus::UDSS_UpToDate)
 		{
 			if (OutMsg)
 			{
@@ -194,31 +89,42 @@ FStructureEditorUtils::EStructureError FStructureEditorUtils::IsStructureValid(c
 
 		for (const UProperty * P = Struct->PropertyLink; P; P = P->PropertyLinkNext)
 		{
-			const UStructProperty* StructProp = Cast<const UStructProperty>(P);
-			if (NULL == StructProp)
 			{
-				if (const UArrayProperty* ArrayProp = Cast<const UArrayProperty>(P))
+				const UStructProperty* StructProp = Cast<const UStructProperty>(P);
+				if (NULL == StructProp)
 				{
-					StructProp = Cast<const UStructProperty>(ArrayProp->Inner);
+					if (const UArrayProperty* ArrayProp = Cast<const UArrayProperty>(P))
+					{
+						StructProp = Cast<const UStructProperty>(ArrayProp->Inner);
+					}
 				}
-			}
 
-			if (StructProp)
-			{
-				if ((NULL == StructProp->Struct) || (FallbackStruct == StructProp->Struct))
+				if (StructProp)
 				{
-					if (OutMsg)
+					if ((NULL == StructProp->Struct) || (FallbackStruct == StructProp->Struct))
 					{
-						*OutMsg = FString::Printf(*LOCTEXT("StructureUnknownProperty", "Struct unknown (deleted?). Parent '%s' Property: '%s").ToString(), 
-							*Struct->GetFullName(), *StructProp->GetName());
+						if (OutMsg)
+						{
+							*OutMsg = FString::Printf(*LOCTEXT("StructureUnknownProperty", "Struct unknown (deleted?). Parent '%s' Property: '%s'").ToString(), 
+								*Struct->GetFullName(), *StructProp->GetName());
+						}
+						return EStructureError::FallbackStruct;
 					}
-					return EStructureError::FallbackStruct;
-				}
 
-				const EStructureError Result = IsStructureValid(StructProp->Struct, RecursionParent ? RecursionParent : Struct, OutMsg);
-				if (EStructureError::Ok != Result) 
-				{
-					return Result;
+					FString OutMsgInner;
+					const EStructureError Result = IsStructureValid(
+						StructProp->Struct, 
+						RecursionParent ? RecursionParent : Struct, 
+						OutMsg ? &OutMsgInner : NULL);
+					if (EStructureError::Ok != Result) 
+					{
+						if (OutMsg)
+						{
+							*OutMsg = FString::Printf(*LOCTEXT("StructurePropertyErrorTemplate", "Struct '%s' Property '%s' Error ( %s )").ToString(), 
+								*Struct->GetFullName(), *StructProp->GetName(), *OutMsgInner);
+						}
+						return Result;
+					}
 				}
 			}
 		}
@@ -227,7 +133,7 @@ FStructureEditorUtils::EStructureError FStructureEditorUtils::IsStructureValid(c
 	return EStructureError::Ok;
 }
 
-bool FStructureEditorUtils::CanHaveAMemberVariableOfType(const UBlueprintGeneratedStruct* Struct, const FEdGraphPinType& VarType, FString* OutMsg)
+bool FStructureEditorUtils::CanHaveAMemberVariableOfType(const UUserDefinedStruct* Struct, const FEdGraphPinType& VarType, FString* OutMsg)
 {
 	const UEdGraphSchema_K2* K2Schema = GetDefault<UEdGraphSchema_K2>();
 	if ((VarType.PinCategory == K2Schema->PC_Struct) && Struct)
@@ -249,7 +155,11 @@ bool FStructureEditorUtils::CanHaveAMemberVariableOfType(const UBlueprintGenerat
 			return false;
 		}
 	}
-	else if ((VarType.PinCategory == K2Schema->PC_Exec) || (VarType.PinCategory == K2Schema->PC_Meta) || (VarType.PinCategory == K2Schema->PC_Wildcard))
+	else if ((VarType.PinCategory == K2Schema->PC_Exec) 
+		|| (VarType.PinCategory == K2Schema->PC_Meta) 
+		|| (VarType.PinCategory == K2Schema->PC_Wildcard)
+		|| (VarType.PinCategory == K2Schema->PC_MCDelegate)
+		|| (VarType.PinCategory == K2Schema->PC_Delegate))
 	{
 		if (OutMsg)
 		{
@@ -262,7 +172,7 @@ bool FStructureEditorUtils::CanHaveAMemberVariableOfType(const UBlueprintGenerat
 
 struct FMemberVariableNameHelper
 {
-	static FName Generate(UBlueprintGeneratedStruct* Struct, const FString& NameBase)
+	static FName Generate(UUserDefinedStruct* Struct, const FString& NameBase)
 	{
 		FString Result;
 		if (!NameBase.IsEmpty())
@@ -281,7 +191,7 @@ struct FMemberVariableNameHelper
 
 		if(Struct)
 		{
-			const uint32 UniqueNameId = Struct->GenerateUniqueNameIdForMemberVariable();
+			const uint32 UniqueNameId = CastChecked<UUserDefinedStructEditorData>(Struct->EditorData)->GenerateUniqueNameIdForMemberVariable();
 			Result = FString::Printf(TEXT("%s_%u"), *Result, UniqueNameId);
 		}
 		FName NameResult = FName(*Result);
@@ -290,89 +200,85 @@ struct FMemberVariableNameHelper
 	}
 };
 
-bool FStructureEditorUtils::AddVariable(UBlueprint* Blueprint, FName StructName, const FEdGraphPinType& VarType)
+bool FStructureEditorUtils::AddVariable(UUserDefinedStruct* Struct, const FEdGraphPinType& VarType)
 {
-	if (NULL != Blueprint)
+	if (Struct)
 	{
 		const FScopedTransaction Transaction( LOCTEXT("AddVariable", "Add Variable") );
-		Blueprint->Modify();
+		ModifyStructData(Struct);
 
-		FBPStructureDescription* StructureDesc = Blueprint->UserDefinedStructures.FindByPredicate(FFindByNameHelper(StructName));
-		if (StructureDesc)
+		FString ErrorMessage;
+		if (!CanHaveAMemberVariableOfType(Struct, VarType, &ErrorMessage))
 		{
-			FString ErrorMessage;
-			if (!CanHaveAMemberVariableOfType(StructureDesc->CompiledStruct, VarType, &ErrorMessage))
-			{
-				UE_LOG(LogBlueprint, Warning, TEXT("%s"), *ErrorMessage);
-				return false;
-			}
+			UE_LOG(LogBlueprint, Warning, TEXT("%s"), *ErrorMessage);
+			return false;
+		}
 
-			const FName VarName = FMemberVariableNameHelper::Generate(StructureDesc->CompiledStruct, FString());
-			check(NULL == StructureDesc->Fields.FindByPredicate(FFindByNameHelper(VarName)));
-			const FString DisplayName = VarName.ToString();
-			check(IsUniqueVariableDisplayName(Blueprint, StructName, DisplayName));
+		const FName VarName = FMemberVariableNameHelper::Generate(Struct, FString());
+		check(NULL == GetVarDesc(Struct).FindByPredicate(FStructureEditorUtils::FFindByNameHelper<FStructVariableDescription>(VarName)));
+		const FString DisplayName = VarName.ToString();
+		check(IsUniqueVariableDisplayName(Struct, DisplayName));
 
-			FBPVariableDescription NewVar;
-			NewVar.VarName = VarName;
-			NewVar.FriendlyName = DisplayName;
-			NewVar.VarType = VarType;
-			NewVar.VarGuid = FGuid::NewGuid();
-			StructureDesc->Fields.Add(NewVar);
+		FStructVariableDescription NewVar;
+		NewVar.VarName = VarName;
+		NewVar.FriendlyName = DisplayName;
+		NewVar.SetPinType(VarType);
+		NewVar.VarGuid = FGuid::NewGuid();
+		GetVarDesc(Struct).Add(NewVar);
 
-			OnStructureChanged(*StructureDesc, Blueprint);
-			return true;
-		}
+		OnStructureChanged(Struct);
+		return true;
 	}
 	return false;
 }
 
-bool FStructureEditorUtils::RemoveVariable(UBlueprint* Blueprint, FName StructName, FGuid VarGuid)
+bool FStructureEditorUtils::RemoveVariable(UUserDefinedStruct* Struct, FGuid VarGuid)
 {
-	if(NULL != Blueprint)
+	if(Struct)
 	{
 		const FScopedTransaction Transaction( LOCTEXT("RemoveVariable", "Remove Variable") );
-		Blueprint->Modify();
+		ModifyStructData(Struct);
 
-		FBPStructureDescription* StructureDesc = Blueprint->UserDefinedStructures.FindByPredicate(FFindByNameHelper(StructName));
-		if(StructureDesc)
+		const auto OldNum = GetVarDesc(Struct).Num();
+		const bool bAllowToMakeEmpty = false;
+		if (bAllowToMakeEmpty || (OldNum > 1))
 		{
-			const auto OldNum = StructureDesc->Fields.Num();
-			StructureDesc->Fields.RemoveAll(FFindByGuidHelper(VarGuid));
-			if(OldNum != StructureDesc->Fields.Num())
+			GetVarDesc(Struct).RemoveAll(FFindByGuidHelper<FStructVariableDescription>(VarGuid));
+			if (OldNum != GetVarDesc(Struct).Num())
 			{
-				OnStructureChanged(*StructureDesc, Blueprint);
+				OnStructureChanged(Struct);
 				return true;
 			}
 		}
+		else
+		{
+			UE_LOG(LogBlueprint, Log, TEXT("Member variable cannot be removed. User Defined Structure cannot be empty"));
+		}
 	}
 	return false;
 }
 
-bool FStructureEditorUtils::RenameVariable(UBlueprint* Blueprint, FName StructName, FGuid VarGuid, const FString& NewDisplayNameStr)
+bool FStructureEditorUtils::RenameVariable(UUserDefinedStruct* Struct, FGuid VarGuid, const FString& NewDisplayNameStr)
 {
-	if (NULL != Blueprint)
+	if (Struct)
 	{
 		const FScopedTransaction Transaction( LOCTEXT("RenameVariable", "Rename Variable") );
-		Blueprint->Modify();
+		ModifyStructData(Struct);
 
-		FBPStructureDescription* StructureDesc = Blueprint->UserDefinedStructures.FindByPredicate(FFindByNameHelper(StructName));
-		if (StructureDesc)
+		const FName NewNameBase(*NewDisplayNameStr);
+		if (NewNameBase.IsValidXName(INVALID_OBJECTNAME_CHARACTERS) &&
+			IsUniqueVariableDisplayName(Struct, NewDisplayNameStr))
 		{
-			const FName NewNameBase(*NewDisplayNameStr);
-			if (NewNameBase.IsValidXName(INVALID_OBJECTNAME_CHARACTERS) &&
-				IsUniqueVariableDisplayName(Blueprint, StructName, NewDisplayNameStr))
+			const FName NewName = FMemberVariableNameHelper::Generate(Struct, NewDisplayNameStr);
+			if (NULL == GetVarDesc(Struct).FindByPredicate(FFindByNameHelper<FStructVariableDescription>(NewName)))
 			{
-				const FName NewName = FMemberVariableNameHelper::Generate(StructureDesc->CompiledStruct, NewDisplayNameStr);
-				if (NULL == StructureDesc->Fields.FindByPredicate(FFindByNameHelper(NewName)))
+				FStructVariableDescription* VarDesc = GetVarDesc(Struct).FindByPredicate(FFindByGuidHelper<FStructVariableDescription>(VarGuid));
+				if (VarDesc)
 				{
-					FBPVariableDescription* VarDesc = StructureDesc->Fields.FindByPredicate(FFindByGuidHelper(VarGuid));
-					if (VarDesc)
-					{
-						VarDesc->FriendlyName = NewDisplayNameStr;
-						VarDesc->VarName = NewName;
-						OnStructureChanged(*StructureDesc, Blueprint);
-						return true;
-					}
+					VarDesc->FriendlyName = NewDisplayNameStr;
+					VarDesc->VarName = NewName;
+					OnStructureChanged(Struct);
+					return true;
 				}
 			}
 		}
@@ -380,31 +286,31 @@ bool FStructureEditorUtils::RenameVariable(UBlueprint* Blueprint, FName StructNa
 	return false;
 }
 
-bool FStructureEditorUtils::ChangeVariableType(UBlueprint* Blueprint, FName StructName, FGuid VarGuid, const FEdGraphPinType& NewType)
+bool FStructureEditorUtils::ChangeVariableType(UUserDefinedStruct* Struct, FGuid VarGuid, const FEdGraphPinType& NewType)
 {
-	if(NULL != Blueprint)
+	if (Struct)
 	{
 		const FScopedTransaction Transaction( LOCTEXT("ChangeVariableType", "Change Variable Type") );
-		Blueprint->Modify();
+		ModifyStructData(Struct);
 
-		FBPStructureDescription* StructureDesc = Blueprint->UserDefinedStructures.FindByPredicate(FFindByNameHelper(StructName));
-		if(StructureDesc)
+		FString ErrorMessage;
+		if(!CanHaveAMemberVariableOfType(Struct, NewType, &ErrorMessage))
 		{
-			FString ErrorMessage;
-			if(!CanHaveAMemberVariableOfType(StructureDesc->CompiledStruct, NewType, &ErrorMessage))
-			{
-				UE_LOG(LogBlueprint, Warning, TEXT("%s"), *ErrorMessage);
-				return false;
-			}
+			UE_LOG(LogBlueprint, Warning, TEXT("%s"), *ErrorMessage);
+			return false;
+		}
 
-			FBPVariableDescription* VarDesc = StructureDesc->Fields.FindByPredicate(FFindByGuidHelper(VarGuid));
-			if(VarDesc && (NewType != VarDesc->VarType))
+		FStructVariableDescription* VarDesc = GetVarDesc(Struct).FindByPredicate(FFindByGuidHelper<FStructVariableDescription>(VarGuid));
+		if(VarDesc)
+		{
+			const bool bChangedType = (VarDesc->ToPinType() != NewType);
+			if (bChangedType)
 			{
-				VarDesc->VarName = FMemberVariableNameHelper::Generate(StructureDesc->CompiledStruct, VarDesc->FriendlyName);
+				VarDesc->VarName = FMemberVariableNameHelper::Generate(Struct, VarDesc->FriendlyName);
 				VarDesc->DefaultValue = FString();
-				VarDesc->VarType = NewType;
+				VarDesc->SetPinType(NewType);
 
-				OnStructureChanged(*StructureDesc, Blueprint);
+				OnStructureChanged(Struct);
 				return true;
 			}
 		}
@@ -412,112 +318,75 @@ bool FStructureEditorUtils::ChangeVariableType(UBlueprint* Blueprint, FName Stru
 	return false;
 }
 
-bool FStructureEditorUtils::ChangeVariableDefaultValue(UBlueprint* Blueprint, FName StructName, FGuid VarGuid, const FString& NewDefaultValue)
+bool FStructureEditorUtils::ChangeVariableDefaultValue(UUserDefinedStruct* Struct, FGuid VarGuid, const FString& NewDefaultValue)
 {
-	if(NULL != Blueprint)
+	if(Struct)
 	{
 		const FScopedTransaction Transaction( LOCTEXT("ChangeVariableDefaultValue", "Change Variable Default Value") );
-		Blueprint->Modify();
+		ModifyStructData(Struct);
 
-		FBPStructureDescription* StructureDesc = Blueprint->UserDefinedStructures.FindByPredicate(FFindByNameHelper(StructName));
-		if (StructureDesc)
+		const UEdGraphSchema_K2* K2Schema = GetDefault<UEdGraphSchema_K2>();
+		FStructVariableDescription* VarDesc = GetVarDesc(Struct).FindByPredicate(FFindByGuidHelper<FStructVariableDescription>(VarGuid));
+		if (VarDesc 
+			&& (NewDefaultValue != VarDesc->DefaultValue) 
+			&& K2Schema->DefaultValueSimpleValidation(VarDesc->ToPinType(), FString(), NewDefaultValue, NULL, FText::GetEmpty()))
 		{
-			const UEdGraphSchema_K2* K2Schema = GetDefault<UEdGraphSchema_K2>();
-			FBPVariableDescription* VarDesc = StructureDesc->Fields.FindByPredicate(FFindByGuidHelper(VarGuid));
-			if (VarDesc 
-				&& (NewDefaultValue != VarDesc->DefaultValue) 
-				&& K2Schema->DefaultValueSimpleValidation(VarDesc->VarType, FString(), NewDefaultValue, NULL, FText::GetEmpty()))
-			{
-				VarDesc->DefaultValue = NewDefaultValue;
-
-				OnStructureChanged(*StructureDesc, Blueprint);
-				return true;
-			}
+			VarDesc->DefaultValue = NewDefaultValue;
+			OnStructureChanged(Struct);
+			return true;
 		}
 	}
 	return false;
 }
 
-bool FStructureEditorUtils::IsUniqueVariableDisplayName(const UBlueprint* Blueprint, FName StructName, const FString& DisplayName)
+bool FStructureEditorUtils::IsUniqueVariableDisplayName(const UUserDefinedStruct* Struct, const FString& DisplayName)
 {
-	if(NULL != Blueprint)
+	if(Struct)
 	{
-		const FBPStructureDescription* StructureDesc = Blueprint->UserDefinedStructures.FindByPredicate(FFindByNameHelper(StructName));
-		if (StructureDesc)
+		for (auto& VarDesc : GetVarDesc(Struct))
 		{
-			for (auto Iter = StructureDesc->Fields.CreateConstIterator(); Iter; ++Iter)
+			if (VarDesc.FriendlyName == DisplayName)
 			{
-				const FBPVariableDescription& VarDesc = *Iter;
-				if (VarDesc.FriendlyName == DisplayName)
-				{
-					return false;
-				}
+				return false;
 			}
-			return true;
 		}
+		return true;
 	}
 	return false;
 }
 
-FString FStructureEditorUtils::GetDisplayName(const UBlueprint* Blueprint, FName StructName, FGuid VarGuid)
+FString FStructureEditorUtils::GetVariableDisplayName(const UUserDefinedStruct* Struct, FGuid VarGuid)
 {
-	if (Blueprint)
+	if (Struct)
 	{
-		const FBPStructureDescription* StructureDesc = Blueprint->UserDefinedStructures.FindByPredicate(FFindByNameHelper(StructName));
-		if (StructureDesc)
+		const FStructVariableDescription* VarDesc = GetVarDesc(Struct).FindByPredicate(FFindByGuidHelper<FStructVariableDescription>(VarGuid));
+		if (VarDesc)
 		{
-			const FBPVariableDescription* VarDesc = StructureDesc->Fields.FindByPredicate(FFindByGuidHelper(VarGuid));
-			if (VarDesc)
-			{
-				return VarDesc->FriendlyName;
-			}
+			return VarDesc->FriendlyName;
 		}
 	}
 
 	return FString();
 }
 
-void FStructureEditorUtils::GetAllStructureNames(const UBlueprint* Blueprint, TArray<FName>& OutNames)
-{
-	if(NULL == Blueprint)
-	{
-		return;
-	}
-
-	TArray<UBlueprint*> BPStack;
-	UBlueprint::GetBlueprintHierarchyFromClass(Blueprint->SkeletonGeneratedClass, BPStack);
-	for(auto BPIter = BPStack.CreateConstIterator(); BPIter; ++BPIter)
-	{
-		if(const UBlueprint* BP = *BPIter)
-		{
-			for(auto StructIter = BP->UserDefinedStructures.CreateConstIterator(); StructIter; ++StructIter)
-			{
-				OutNames.AddUnique(StructIter->Name);
-			}
-		}
-	}
-}
-
-bool FStructureEditorUtils::StructureEditingEnabled()
+bool FStructureEditorUtils::UserDefinedStructEnabled()
 {
 	bool bUseUserDefinedStructure = false;
 	GConfig->GetBool( TEXT("UserDefinedStructure"), TEXT("bUseUserDefinedStructure"), bUseUserDefinedStructure, GEditorIni );
 	return bUseUserDefinedStructure;
 }
 
-bool FStructureEditorUtils::Fill_MakeStructureDefaultValue(const UBlueprintGeneratedStruct* Struct, uint8* StructData)
+bool FStructureEditorUtils::Fill_MakeStructureDefaultValue(const UUserDefinedStruct* Struct, uint8* StructData)
 {
-	static const FName MakeStructureDefaultValue = "MakeStructureDefaultValue";
 	bool bResult = true;
-
 	for (TFieldIterator<UProperty> It(Struct); It; ++It)
 	{
 		if (const UProperty* Property = *It)
 		{
-			const FString DefaultValAsString = Property->GetMetaData(MakeStructureDefaultValue);
-			if (!DefaultValAsString.IsEmpty())
+			auto VarDesc = GetVarDesc(Struct).FindByPredicate(FFindByNameHelper<FStructVariableDescription>(Property->GetFName()));
+			if (VarDesc && !VarDesc->CurrentDefaultValue.IsEmpty())
 			{
-				bResult &= FBlueprintEditorUtils::PropertyValueFromString(Property, DefaultValAsString, StructData);
+				bResult &= FBlueprintEditorUtils::PropertyValueFromString(Property, VarDesc->CurrentDefaultValue, StructData);
 			}
 			else
 			{
@@ -536,7 +405,7 @@ bool FStructureEditorUtils::Fill_MakeStructureDefaultValue(const UProperty* Prop
 
 	if (const UStructProperty* StructProperty = Cast<const UStructProperty>(Property))
 	{
-		if (const UBlueprintGeneratedStruct* InnerStruct = Cast<const UBlueprintGeneratedStruct>(StructProperty->Struct))
+		if (const UUserDefinedStruct* InnerStruct = Cast<const UUserDefinedStruct>(StructProperty->Struct))
 		{
 			bResult &= Fill_MakeStructureDefaultValue(InnerStruct, PropertyData);
 		}
@@ -544,7 +413,7 @@ bool FStructureEditorUtils::Fill_MakeStructureDefaultValue(const UProperty* Prop
 	else if (const UArrayProperty* ArrayProp = Cast<const UArrayProperty>(Property))
 	{
 		StructProperty = Cast<const UStructProperty>(ArrayProp->Inner);
-		const UBlueprintGeneratedStruct* InnerStruct = StructProperty ? Cast<const UBlueprintGeneratedStruct>(StructProperty->Struct) : NULL;
+		const UUserDefinedStruct* InnerStruct = StructProperty ? Cast<const UUserDefinedStruct>(StructProperty->Struct) : NULL;
 		if(InnerStruct)
 		{
 			FScriptArrayHelper ArrayHelper(ArrayProp, PropertyData);
@@ -559,27 +428,28 @@ bool FStructureEditorUtils::Fill_MakeStructureDefaultValue(const UProperty* Prop
 	return bResult;
 }
 
-void FStructureEditorUtils::CompileStructuresInBlueprint(UBlueprint* Blueprint)
+void FStructureEditorUtils::CompileStructure(UUserDefinedStruct* Struct)
 {
-	if (Blueprint)
+	if (Struct)
 	{
 		IKismetCompilerInterface& Compiler = FModuleManager::LoadModuleChecked<IKismetCompilerInterface>(KISMET_COMPILER_MODULENAME);
 		FCompilerResultsLog Results;
-		FKismetCompilerOptions CompileOptions;
-		CompileOptions.CompileType = EKismetCompileType::StructuresOnly;
-		Compiler.CompileBlueprint(Blueprint, CompileOptions, Results, NULL);
+		Compiler.CompileStructure(Struct, Results);
 	}
 }
 
-void FStructureEditorUtils::OnStructureChanged(FBPStructureDescription& StructureDesc, UBlueprint* Blueprint)
+void FStructureEditorUtils::OnStructureChanged(UUserDefinedStruct* Struct)
 {
-	if (StructureDesc.CompiledStruct)
+	if (Struct)
 	{
-		StructureDesc.CompiledStruct->Status = EBlueprintStructureStatus::BSS_Dirty;
+		Struct->Status = EUserDefinedStructureStatus::UDSS_Dirty;
+		CompileStructure(Struct);
+		FStructEditorManager::Get().OnChanged(Struct);
+		Struct->MarkPackageDirty();
 	}
-	CompileStructuresInBlueprint(Blueprint);
 }
 
+//TODO: Move to blueprint utils
 void FStructureEditorUtils::RemoveInvalidStructureMemberVariableFromBlueprint(UBlueprint* Blueprint)
 {
 	if (Blueprint)
@@ -622,7 +492,7 @@ void FStructureEditorUtils::RemoveInvalidStructureMemberVariableFromBlueprint(UB
 				for (auto NameIter = ZombieMemberNames.CreateConstIterator(); NameIter; ++NameIter)
 				{
 					const FName Name = *NameIter;
-					Blueprint->NewVariables.RemoveAll(FFindByNameHelper(Name)); //TODO: Add RemoveFirst to TArray
+					Blueprint->NewVariables.RemoveAll(FFindByNameHelper<FBPVariableDescription>(Name)); //TODO: Add RemoveFirst to TArray
 					FBlueprintEditorUtils::RemoveVariableNodes(Blueprint, Name);
 				}
 			}
@@ -630,4 +500,26 @@ void FStructureEditorUtils::RemoveInvalidStructureMemberVariableFromBlueprint(UB
 	}
 }
 
+TArray<FStructVariableDescription>& FStructureEditorUtils::GetVarDesc(UUserDefinedStruct* Struct)
+{
+	check(Struct);
+	return CastChecked<UUserDefinedStructEditorData>(Struct->EditorData)->VariablesDescriptions;
+}
+
+const TArray<FStructVariableDescription>& FStructureEditorUtils::GetVarDesc(const UUserDefinedStruct* Struct)
+{
+	check(Struct);
+	return CastChecked<const UUserDefinedStructEditorData>(Struct->EditorData)->VariablesDescriptions;
+}
+
+void FStructureEditorUtils::ModifyStructData(UUserDefinedStruct* Struct)
+{
+	UUserDefinedStructEditorData* EditorData = Struct ? Cast<UUserDefinedStructEditorData>(Struct->EditorData) : NULL;
+	ensure(EditorData);
+	if (EditorData)
+	{
+		EditorData->Modify();
+	}
+}
+
 #undef LOCTEXT_NAMESPACE
diff --git a/Engine/Source/Editor/UnrealEd/Private/UnrealEd.cpp b/Engine/Source/Editor/UnrealEd/Private/UnrealEd.cpp
--- a/Engine/Source/Editor/UnrealEd/Private/UnrealEd.cpp
+++ b/Engine/Source/Editor/UnrealEd/Private/UnrealEd.cpp
@@ -170,5 +170,5 @@ void EditorExit()
 
 IMPLEMENT_MODULE( FDefaultModuleImpl, UnrealEd );
 
-#include "UnrealEd.generated.inl"
+
 
diff --git a/Engine/Source/Editor/UnrealEd/Private/UserDefinedStructEditorData.cpp b/Engine/Source/Editor/UnrealEd/Private/UserDefinedStructEditorData.cpp
new file mode 100644
--- /dev/null
+++ b/Engine/Source/Editor/UnrealEd/Private/UserDefinedStructEditorData.cpp
@@ -0,0 +1,51 @@
+
+#include "UnrealEd.h"
+#include "StructureEditorUtils.h"
+
+bool FStructVariableDescription::SetPinType(const FEdGraphPinType& VarType)
+{
+	Category = VarType.PinCategory;
+	SubCategory = VarType.PinSubCategory;
+	SubCategoryObject = VarType.PinSubCategoryObject.Get();
+	bIsArray = VarType.bIsArray;
+
+	return !VarType.bIsReference && !VarType.bIsWeakPointer;
+}
+
+FEdGraphPinType FStructVariableDescription::ToPinType() const
+{
+	return FEdGraphPinType(Category, SubCategory, SubCategoryObject.Get(), bIsArray, false);
+}
+
+UUserDefinedStructEditorData::UUserDefinedStructEditorData(const class FPostConstructInitializeProperties& PCIP)
+: Super(PCIP)
+{
+}
+
+uint32 UUserDefinedStructEditorData::GenerateUniqueNameIdForMemberVariable()
+{
+	const uint32 Result = UniqueNameId;
+	++UniqueNameId;
+	return Result;
+}
+
+UUserDefinedStruct* UUserDefinedStructEditorData::GetOwnerStruct() const
+{
+	return Cast<UUserDefinedStruct>(GetOuter());
+}
+
+void UUserDefinedStructEditorData::PostEditUndo()
+{
+	Super::PostEditUndo();
+	FStructureEditorUtils::OnStructureChanged(GetOwnerStruct());
+}
+
+void UUserDefinedStructEditorData::PostLoadSubobjects(FObjectInstancingGraph* OuterInstanceGraph)
+{
+	Super::PostLoadSubobjects(OuterInstanceGraph);
+
+	for (auto& VarDesc : VariablesDescriptions)
+	{
+		VarDesc.bInvalidMember = !FStructureEditorUtils::CanHaveAMemberVariableOfType(GetOwnerStruct(), VarDesc.ToPinType());
+	}
+}
\ No newline at end of file
diff --git a/Engine/Source/Editor/UnrealEd/Public/Kismet2/BlueprintEditorUtils.h b/Engine/Source/Editor/UnrealEd/Public/Kismet2/BlueprintEditorUtils.h
--- a/Engine/Source/Editor/UnrealEd/Public/Kismet2/BlueprintEditorUtils.h
+++ b/Engine/Source/Editor/UnrealEd/Public/Kismet2/BlueprintEditorUtils.h
@@ -67,11 +67,6 @@ public:
 	 */
 	static UClass* RegenerateBlueprintClass(UBlueprint* Blueprint, UClass* ClassToRegenerate, UObject* PreviousCDO, TArray<UObject*>& ObjLoaded);
 
-	/**
-	 * Input delegate pin's PinSubCategoryObject may be obsolete. But MCDelegate node has reference to actual function's signature.
-	 */
-	static void RefreshInputDelegatePins(UBlueprint* Blueprint);
-
 	/**
 	 * Copies the default properties of all parent blueprint classes in the chain to the specified blueprint's skeleton CDO
 	 */
@@ -373,6 +368,15 @@ public:
 	//////////////////////////////////////////////////////////////////////////
 	// Variables
 
+	/**
+	 * Checks if pin type stores proper type for a variable or parameter. Especially if the UDStruct is valid.
+	 *
+	 * @param		Type	Checked pin type.
+	 *
+	 * @return				if type is valid
+	 */
+	static bool IsPinTypeValid(const FEdGraphPinType& Type);
+
 	/**
 	 * Gets the visible class variable list.  This includes both variables introduced here and in all superclasses.
 	 *
@@ -797,4 +801,7 @@ protected:
 	 */
 	static bool CheckIfNodeConnectsToSelection(UEdGraphNode* InNode, const TSet<UEdGraphNode*>& InSelectionSet);
 
+public:
+	static FName GetFunctionNameFromClassByGuid(const UClass* InClass, const FGuid FunctionGuid);
+	static bool GetFunctionGuidFromClassByFieldName(const UClass* InClass, const FName FunctionName, FGuid& FunctionGuid);
 };
diff --git a/Engine/Source/Editor/UnrealEd/Public/Kismet2/EnumEditorUtils.h b/Engine/Source/Editor/UnrealEd/Public/Kismet2/EnumEditorUtils.h
--- a/Engine/Source/Editor/UnrealEd/Public/Kismet2/EnumEditorUtils.h
+++ b/Engine/Source/Editor/UnrealEd/Public/Kismet2/EnumEditorUtils.h
@@ -6,6 +6,7 @@
 #pragma once
 
 #include "Engine.h"
+#include "ListenerManager.h"
 
 class UNREALED_API FEnumEditorUtils
 {
@@ -14,26 +15,16 @@ class UNREALED_API FEnumEditorUtils
 	/** copy full enumeratos names from given enum to OutEnumNames, the last '_MAX' enumerator is skipped */
 	static void CopyEnumeratorsWithoutMax(const UEnum* Enum, TArray<FName>& OutEnumNames);
 public:
-	
-	class UNREALED_API INotifyOnEnumChanged
-	{
-	public:
-		virtual void OnEnumChanged(const class UUserDefinedEnum* Enum) = 0;
-	};
 
-	class UNREALED_API FEnumEditorManager
+	class FEnumEditorManager : public FListenerManager<UUserDefinedEnum>
 	{
-		TSet<INotifyOnEnumChanged*> Listeners;
-
 		FEnumEditorManager();
 	public:
-		static FEnumEditorManager& Get();
-
-		void AddEnumListener(INotifyOnEnumChanged* EnumListener);
-		void RemoveEnumListener(INotifyOnEnumChanged* EnumListener);
-		void OnEnumChanged(const UUserDefinedEnum* Enum);
+		UNREALED_API static FEnumEditorManager& Get();
 	};
 
+	typedef FEnumEditorManager::ListenerType INotifyOnEnumChanged;
+
 	//////////////////////////////////////////////////////////////////////////
 	// User defined enumerations
 
diff --git a/Engine/Source/Editor/UnrealEd/Public/Kismet2/ListenerManager.h b/Engine/Source/Editor/UnrealEd/Public/Kismet2/ListenerManager.h
new file mode 100644
--- /dev/null
+++ b/Engine/Source/Editor/UnrealEd/Public/Kismet2/ListenerManager.h
@@ -0,0 +1,42 @@
+// Copyright 1998-2013 Epic Games, Inc. All Rights Reserved.
+
+#pragma once
+
+template<typename TType>
+class INotifyOnChanged
+{
+public:
+	virtual void OnChanged(const TType* Changed) = 0;
+};
+
+template<typename TType>
+class FListenerManager
+{
+public:
+	typedef INotifyOnChanged<TType> ListenerType;
+
+private:
+	TSet<ListenerType*> Listeners;
+
+public:
+	void AddListener(ListenerType* Listener)
+	{
+		if (Listener)
+		{
+			Listeners.Add(Listener);
+		}
+	}
+
+	void RemoveListener(ListenerType* Listener)
+	{
+		Listeners.Remove(Listener);
+	}
+
+	void OnChanged(const TType* Changed)
+	{
+		for (auto Listener : Listeners)
+		{
+			Listener->OnChanged(Changed);
+		}
+	}
+};
\ No newline at end of file
diff --git a/Engine/Source/Editor/UnrealEd/Public/Kismet2/StructureEditorUtils.h b/Engine/Source/Editor/UnrealEd/Public/Kismet2/StructureEditorUtils.h
--- a/Engine/Source/Editor/UnrealEd/Public/Kismet2/StructureEditorUtils.h
+++ b/Engine/Source/Editor/UnrealEd/Public/Kismet2/StructureEditorUtils.h
@@ -2,70 +2,111 @@
 #pragma once
 
 #include "Engine.h"
+#include "ListenerManager.h"
+#include "Editor/UnrealEd/Classes/UserDefinedStructure/UserDefinedStructEditorData.h"
 
 class UNREALED_API FStructureEditorUtils
 {
 public:
 
+	class FStructEditorManager : public FListenerManager<UUserDefinedStruct>
+	{
+		FStructEditorManager();
+	public:
+		UNREALED_API static FStructEditorManager& Get();
+	};
+
+	typedef FStructEditorManager::ListenerType INotifyOnStructChanged;
+
+	template<class TElement>
 	struct FFindByNameHelper
 	{
 		const FName Name;
 
 		FFindByNameHelper(FName InName) : Name(InName) { }
 
-		bool operator() (const FBPStructureDescription& Element) const
-		{
-			return (Name == Element.Name);
-		}
-
-		bool operator() (const FBPVariableDescription& Element) const
+		bool operator() (const TElement& Element) const
 		{
 			return (Name == Element.VarName);
 		}
 	};
 
+	template<class TElement>
 	struct FFindByGuidHelper
 	{
-		FGuid Guid;
+		const FGuid Guid;
 
 		FFindByGuidHelper(FGuid InGuid) : Guid(InGuid) { }
 
-		bool operator() (const FBPVariableDescription& Element) const
+		bool operator() (const TElement& Element) const
 		{
 			return (Guid == Element.VarGuid);
 		}
 	};
 
-	//STRUCTURE
-	static bool AddStructure(UBlueprint* Blueprint);
+	struct FStructOnScope
+	{
+	private:
+		UScriptStruct* ScriptStruct;
+		uint8* SampleStructMemory;
+
+		FStructOnScope(const FStructOnScope&);
+		FStructOnScope& operator=(const FStructOnScope&);
 
-	static bool RemoveStructure(UBlueprint* Blueprint, FName StructName);
+	public:
+		FStructOnScope(UScriptStruct* InScriptStruct) 
+			: ScriptStruct(InScriptStruct)
+			, SampleStructMemory(NULL)
+		{
+			if(ScriptStruct)
+			{
+				SampleStructMemory = (uint8*)FMemory::Malloc(ScriptStruct->GetStructureSize());
+				ScriptStruct->InitializeScriptStruct(SampleStructMemory);
+			}
+		}
+
+		uint8* GetStructMemory() { return SampleStructMemory; }
+		const uint8* GetStructMemory() const { return SampleStructMemory; }
+
+		~FStructOnScope()
+		{
+			if(ScriptStruct)
+			{
+				ScriptStruct->DestroyScriptStruct(SampleStructMemory);
+				FMemory::Free(SampleStructMemory);
+				SampleStructMemory = NULL;
+			}
+		}
+	};
 
-	static bool RenameStructure(UBlueprint* Blueprint, FName OldStructName, const FString& NewName);
+	//STRUCTURE
+	static UUserDefinedStruct* CreateUserDefinedStruct(UObject* InParent, FName Name, EObjectFlags Flags);
 
-	static bool DuplicateStructure(UBlueprint* Blueprint, FName StructName);
+	static void CompileStructure(UUserDefinedStruct* Struct);
 
 	//VARIABLE
-	static bool AddVariable(UBlueprint* Blueprint, FName StructName, const FEdGraphPinType& VarType);
+	static bool AddVariable(UUserDefinedStruct* Struct, const FEdGraphPinType& VarType);
 
-	static bool RemoveVariable(UBlueprint* Blueprint, FName StructName, FGuid VarGuid);
+	static bool RemoveVariable(UUserDefinedStruct* Struct, FGuid VarGuid);
 
-	static bool RenameVariable(UBlueprint* Blueprint, FName StructName, FGuid VarGuid, const FString& NewDisplayNameStr);
+	static bool RenameVariable(UUserDefinedStruct* Struct, FGuid VarGuid, const FString& NewDisplayNameStr);
 
-	static bool ChangeVariableType(UBlueprint* Blueprint, FName StructName, FGuid VarGuid, const FEdGraphPinType& NewType);
+	static bool ChangeVariableType(UUserDefinedStruct* Struct, FGuid VarGuid, const FEdGraphPinType& NewType);
 
-	static bool ChangeVariableDefaultValue(UBlueprint* Blueprint, FName StructName, FGuid VarGuid, const FString& NewDefaultValue);
+	static bool ChangeVariableDefaultValue(UUserDefinedStruct* Struct, FGuid VarGuid, const FString& NewDefaultValue);
 
-	static bool IsUniqueVariableDisplayName(const UBlueprint* Blueprint, FName StructName, const FString& DisplayName);
+	static bool IsUniqueVariableDisplayName(const UUserDefinedStruct* Struct, const FString& DisplayName);
 
-	static FString GetDisplayName(const UBlueprint* Blueprint, FName StructName, FGuid VarGuid);
+	static FString GetVariableDisplayName(const UUserDefinedStruct* Struct, FGuid VarGuid);
 
 	//MISC
-	static void GetAllStructureNames(const UBlueprint* Blueprint, TArray<FName>& OutNames);
+	static TArray<FStructVariableDescription>& GetVarDesc(UUserDefinedStruct* Struct);
 
-	static bool StructureEditingEnabled();
+	static const TArray<FStructVariableDescription>& GetVarDesc(const UUserDefinedStruct* Struct);
 
-	static void CompileStructuresInBlueprint(UBlueprint* Blueprint);
+	static void ModifyStructData(UUserDefinedStruct* Struct);
+
+	static bool UserDefinedStructEnabled();
 
 	static void RemoveInvalidStructureMemberVariableFromBlueprint(UBlueprint* Blueprint);
 
@@ -73,12 +114,12 @@ public:
 	 Default values for member variables in User Defined Structure are stored in meta data "MakeStructureDefaultValue"
 	 The following functions are used to fill an instance of user defined struct with those default values.
 	 */
-	static bool Fill_MakeStructureDefaultValue(const UBlueprintGeneratedStruct* Struct, uint8* StructData);
+	static bool Fill_MakeStructureDefaultValue(const UUserDefinedStruct* Struct, uint8* StructData);
 
 	static bool Fill_MakeStructureDefaultValue(const UProperty* Property, uint8* PropertyData);
 
 	//VALIDATION
-	static bool CanHaveAMemberVariableOfType(const UBlueprintGeneratedStruct* Struct, const FEdGraphPinType& VarType, FString* OutMsg = NULL);
+	static bool CanHaveAMemberVariableOfType(const UUserDefinedStruct* Struct, const FEdGraphPinType& VarType, FString* OutMsg = NULL);
 
 	enum EStructureError
 	{
@@ -87,12 +128,13 @@ public:
 		FallbackStruct,
 		NotCompiled,
 		NotBlueprintType,
+		NotSupportedType,
+		EmptyStructure
 	};
 
 	/** Can the structure be a member variable for a BPGClass or BPGStruct */
 	static EStructureError IsStructureValid(const UScriptStruct* Struct, const UStruct* RecursionParent = NULL, FString* OutMsg = NULL);
-private:
-	static FName MakeUniqueStructName(UBlueprint* Blueprint, const FString& BaseName);
-	static void OnStructureChanged(FBPStructureDescription& StructureDesc, UBlueprint* Blueprint);
+
+	static void OnStructureChanged(UUserDefinedStruct* Struct);
 };
 
diff --git a/Engine/Source/Editor/UnrealEdMessages/Private/UnrealEdMessagesModule.cpp b/Engine/Source/Editor/UnrealEdMessages/Private/UnrealEdMessagesModule.cpp
--- a/Engine/Source/Editor/UnrealEdMessages/Private/UnrealEdMessagesModule.cpp
+++ b/Engine/Source/Editor/UnrealEdMessages/Private/UnrealEdMessagesModule.cpp
@@ -5,7 +5,7 @@
 =============================================================================*/
 
 #include "UnrealEdMessagesPrivatePCH.h"
-#include "UnrealEdMessages.generated.inl"
+
 
 
 /**
diff --git a/Engine/Source/Editor/WorldBrowser/Private/WorldBrowserModule.cpp b/Engine/Source/Editor/WorldBrowser/Private/WorldBrowserModule.cpp
--- a/Engine/Source/Editor/WorldBrowser/Private/WorldBrowserModule.cpp
+++ b/Engine/Source/Editor/WorldBrowser/Private/WorldBrowserModule.cpp
@@ -5,7 +5,7 @@
 #include "SWorldMainView.h"
 #include "StreamingLevels/StreamingLevelEdMode.h"
 
-#include "WorldBrowser.generated.inl"
+
 
 IMPLEMENT_MODULE( FWorldBrowserModule, WorldBrowser );
 
diff --git a/Engine/Source/Programs/AutomationTool/WindowsHostPlatform.cs b/Engine/Source/Programs/AutomationTool/WindowsHostPlatform.cs
--- a/Engine/Source/Programs/AutomationTool/WindowsHostPlatform.cs
+++ b/Engine/Source/Programs/AutomationTool/WindowsHostPlatform.cs
@@ -44,15 +44,12 @@ namespace AutomationTool
 		private bool FindBaseVSToolPaths(out string WindowsSDKDir, out string BaseVSToolPath)
 		{
 			WindowsSDKDir = "";
-			BaseVSToolPath = "";
+			BaseVSToolPath = WindowsPlatform.GetVSComnToolsPath();
 
 			if (WindowsPlatform.Compiler == WindowsCompiler.VisualStudio2013)
 			{
 				WindowsSDKDir = FindWindowsSDKInstallationFolder( "v8.1" );
 
-				// Grab path to Visual Studio binaries from the system environment
-				BaseVSToolPath = Environment.GetEnvironmentVariable("VS120COMNTOOLS");
-
 				if (string.IsNullOrEmpty(BaseVSToolPath))
 				{
 					Log.TraceError("Visual Studio 2013 must be installed in order to build this target.");
@@ -63,9 +60,6 @@ namespace AutomationTool
 			{
 				WindowsSDKDir = FindWindowsSDKInstallationFolder( "v8.0" );
 
-				// Grab path to Visual Studio binaries from the system environment
-				BaseVSToolPath = Environment.GetEnvironmentVariable("VS110COMNTOOLS");
-
 				if (string.IsNullOrEmpty(BaseVSToolPath))
 				{
 					Log.TraceError("Visual Studio 2012 must be installed in order to build this target.");
@@ -149,15 +143,7 @@ namespace AutomationTool
 		public override string GetMsDevExe()
 		{
 			// Locate MsDevEnv executable
-			string PotentialMsDevExe = null;
-			if (WindowsPlatform.Compiler == WindowsCompiler.VisualStudio2013)
-			{
-				PotentialMsDevExe = CommandUtils.GetEnvVar("VS120COMNTOOLS") + @"..\IDE\Devenv.com";
-			}
-			else if (WindowsPlatform.Compiler == WindowsCompiler.VisualStudio2012)
-			{
-				PotentialMsDevExe = CommandUtils.GetEnvVar("VS110COMNTOOLS") + @"..\IDE\Devenv.com";
-			}
+			string PotentialMsDevExe = Path.Combine(WindowsPlatform.GetVSComnToolsPath(), "..", "IDE", "Devenv.com");
 
 			var Info = new System.IO.FileInfo(PotentialMsDevExe);
 			if (Info.Exists)
diff --git a/Engine/Source/Programs/UnrealBuildTool/Configuration/UEBuildModule.cs b/Engine/Source/Programs/UnrealBuildTool/Configuration/UEBuildModule.cs
--- a/Engine/Source/Programs/UnrealBuildTool/Configuration/UEBuildModule.cs
+++ b/Engine/Source/Programs/UnrealBuildTool/Configuration/UEBuildModule.cs
@@ -706,8 +706,35 @@ namespace UnrealBuildTool
 	/** A module that is compiled from C++ code. */
 	public class UEBuildModuleCPP : UEBuildModule
 	{
+		public class AutoGenerateInlInfoClass
+		{
+			/** The filename of the *.generated.cpp file which was generated for the module */
+			public readonly string Filename;
+
+			public AutoGenerateInlInfoClass(string InFilename)
+			{
+				Debug.Assert(InFilename != null);
+
+				Filename = InFilename;
+			}
+		}
+
+		/** Information about the .inl file . */
+		public AutoGenerateInlInfoClass AutoGenerateInlInfo = null;
+
+		private class SourceFilesClass
+		{
+			public readonly List<FileItem> MissingFiles = new List<FileItem>();
+			public readonly List<FileItem> CPPFiles     = new List<FileItem>();
+			public readonly List<FileItem> CFiles       = new List<FileItem>();
+			public readonly List<FileItem> CCFiles      = new List<FileItem>();
+			public readonly List<FileItem> MMFiles      = new List<FileItem>();
+			public readonly List<FileItem> RCFiles      = new List<FileItem>();
+			public readonly List<FileItem> OtherFiles   = new List<FileItem>();
+		}
+
 		/** A list of the absolute paths of source files in the module. */
-		List<FileItem> SourceFiles;
+		readonly SourceFilesClass SourceFiles = new SourceFilesClass();
 
 		/** The preprocessor definitions used to compile this module's private implementation. */
 		List<string> Definitions;
@@ -750,12 +777,14 @@ namespace UnrealBuildTool
 		/** Whether we're building the redist static library (as well as using it). */
 		public bool bBuildingRedistStaticLibrary = false;
 
-		/** Public UObject header files within this module.  This is filled in during the Compile() call. */
-		private readonly HashSet<FileItem> PublicUObjectHeaders = new HashSet<FileItem>();
-
-		/** Private UObject header files within this module.  This is filled in during the Compile() call. */
-		private readonly HashSet<FileItem> PrivateUObjectHeaders = new HashSet<FileItem>();
+		public class ProcessedDependenciesClass
+		{
+			/** The file, if any, which is used as the unique PCH for this module */
+			public FileItem UniquePCHHeaderFile = null;
+		}
 
+		/** The processed dependencies for the class */
+		public ProcessedDependenciesClass ProcessedDependencies = null;
 
 		public UEBuildModuleCPP(
 			UEBuildTarget InTarget,
@@ -815,7 +844,40 @@ namespace UnrealBuildTool
 				)
 		{
 			IntelliSenseGatherer = InIntelliSenseGatherer;
-			SourceFiles = new List<FileItem>( InSourceFiles );
+
+			foreach (var SourceFile in InSourceFiles)
+			{
+				string Extension = Path.GetExtension(SourceFile.AbsolutePath).ToUpperInvariant();
+				if (!SourceFile.bExists)
+				{
+					SourceFiles.MissingFiles.Add(SourceFile);
+				}
+				else if (Extension == ".CPP")
+				{
+					SourceFiles.CPPFiles.Add(SourceFile);
+				}
+				else if (Extension == ".C")
+				{
+					SourceFiles.CFiles.Add(SourceFile);
+				}
+				else if (Extension == ".CC")
+				{
+					SourceFiles.CCFiles.Add(SourceFile);
+				}
+				else if (Extension == ".MM" || Extension == ".M")
+				{
+					SourceFiles.MMFiles.Add(SourceFile);
+				}
+				else if (Extension == ".RC")
+				{
+					SourceFiles.RCFiles.Add(SourceFile);
+				}
+				else
+				{
+					SourceFiles.OtherFiles.Add(SourceFile);
+				}
+			}
+
 			Definitions = ListFromOptionalEnumerableStringParameter(InDefinitions);
 			foreach (var Def in Definitions)
 			{
@@ -850,39 +912,7 @@ namespace UnrealBuildTool
 				return LinkInputFiles;
 			}
 
-			var ModuleCompileEnvironment = new CPPEnvironment(CompileEnvironment);
-
-			// Override compile environment
-			ModuleCompileEnvironment.Config.OptimizeCode = OptimizeCode;
-			ModuleCompileEnvironment.Config.bUseRTTI = bUseRTTI;
-			ModuleCompileEnvironment.Config.bFasterWithoutUnity = bFasterWithoutUnity;
-			ModuleCompileEnvironment.Config.OptimizeCode                           = OptimizeCode;
-			ModuleCompileEnvironment.Config.bUseRTTI                               = bUseRTTI;
-			ModuleCompileEnvironment.Config.bFasterWithoutUnity                    = bFasterWithoutUnity;
-			ModuleCompileEnvironment.Config.OptimizeCode = OptimizeCode;
-			ModuleCompileEnvironment.Config.bUseRTTI = bUseRTTI;
-			ModuleCompileEnvironment.Config.bEnableBufferSecurityChecks = bEnableBufferSecurityChecks;
-			ModuleCompileEnvironment.Config.bFasterWithoutUnity = bFasterWithoutUnity;
-			ModuleCompileEnvironment.Config.MinFilesUsingPrecompiledHeaderOverride = MinFilesUsingPrecompiledHeaderOverride;
-			ModuleCompileEnvironment.Config.bEnableExceptions                      = bEnableExceptions;
-			ModuleCompileEnvironment.Config.bEnableInlining                        = bEnableInlining;
-			ModuleCompileEnvironment.Config.OutputDirectory                        = Path.Combine(Binary.Config.IntermediateDirectory, Name);
-
-			// Switch the optimization flag if we're building a game module
-			if(Target.Configuration == UnrealTargetConfiguration.DebugGame && Type == UEBuildModuleType.GameModule)
-			{
-				ModuleCompileEnvironment.Config.TargetConfiguration = CPPTargetConfiguration.Debug;
-			}
-
-			// Add the module's private definitions.
-			ModuleCompileEnvironment.Config.Definitions.AddRange(Definitions);
-
-			// Setup the compile environment for the module.
-			SetupPrivateCompileEnvironment(
-				ref ModuleCompileEnvironment.Config.IncludePaths,
-				ref ModuleCompileEnvironment.Config.SystemIncludePaths,
-				ref ModuleCompileEnvironment.Config.Definitions
-				);
+			var ModuleCompileEnvironment = CreateModuleCompileEnvironment(CompileEnvironment);
 
 			if( IntelliSenseGatherer != null )
 			{
@@ -895,44 +925,14 @@ namespace UnrealBuildTool
 				return LinkInputFiles;
 			}
 
-
-			// Gather a list of CPP files in the project.
-			var CPPFiles = new List<FileItem>();
-			var CFiles = new List<FileItem>();
-			var CCFiles = new List<FileItem>();
-			var MMFiles = new List<FileItem>();
-			var RCFiles = new List<FileItem>();
-			var MissingFiles = new List<FileItem>();
-			foreach (var SourceFile in SourceFiles)
+			// Throw an error if the module's source file list referenced any non-existent files.
+			if (SourceFiles.MissingFiles.Count > 0)
 			{
-				string Extension = Path.GetExtension(SourceFile.AbsolutePath).ToUpperInvariant();
-				if (SourceFile.bExists)
-				{
-					if (Extension == ".CPP")
-					{
-						CPPFiles.Add(SourceFile);
-					}
-					else if (Extension == ".C")
-					{
-						CFiles.Add(SourceFile);
-					}
-					else if (Extension == ".CC")
-					{
-						CCFiles.Add(SourceFile);
-					}
-					else if (Extension == ".MM" || Extension == ".M")
-					{
-						MMFiles.Add(SourceFile);
-					}
-					else if (Extension == ".RC")
-					{
-						RCFiles.Add(SourceFile);
-					}
-				}
-				else
-				{
-					MissingFiles.Add(SourceFile);
-				}
+				throw new BuildException(
+					"UBT ERROR: Module \"{0}\" references non-existent files:\n{1} (perhaps a file was added to the project but not checked in)",
+					Name,
+					string.Join("\n", SourceFiles.MissingFiles.Select(M => M.AbsolutePath))
+				);
 			}
 
 			// For an executable or a static library do not use the default RC file - 
@@ -942,42 +942,17 @@ namespace UnrealBuildTool
 			if (Binary.Config.Type == UEBuildBinaryType.DynamicLinkLibrary)
 			{
 				// Add default PCLaunch.rc file if this module has no own resource file specified
-				if (RCFiles.Count <= 0)
+				if (SourceFiles.RCFiles.Count <= 0)
 				{
 					string DefRC = Path.Combine(Directory.GetCurrentDirectory(), "Runtime/Launch/Resources/Windows/PCLaunch.rc");
 					FileItem Item = FileItem.GetItemByFullPath(DefRC);
-					RCFiles.Add(Item);
+					SourceFiles.RCFiles.Add(Item);
 				}
 			}
 
-			// Throw an error if the module's source file list referenced any non-existent files.
-			if (MissingFiles.Count > 0)
-			{
-				string FileList = "";
-				foreach (var MissingFile in MissingFiles)
-				{
-					FileList += "\n";
-					FileList += MissingFile.AbsolutePath;
-				}
-				throw new BuildException(
-					"UBT ERROR: Module \"{0}\" references non-existent files:{1} (perhaps a file was added to the project but not checked in)",
-					Name,
-					FileList
-					);
-			}
-
-			// Clear out any references to precompiled header usage for source files..
-			foreach (var CPPFile in CPPFiles)
-			{
-				CPPFile.PCHHeaderNameInCode = null;
-				CPPFile.PrecompiledHeaderIncludeFilename = null;
-			}
-
-
 			// Check to see if this is an Engine module (including program or plugin modules).  That is, the module is located under the "Engine" folder
 			var IsGameModule = !Utils.IsFileUnderDirectory( this.ModuleDirectory, Path.Combine( ProjectFileGenerator.EngineRelativePath ) );
 
-
 			// Should we force a precompiled header to be generated for this module?  Usually, we only bother with a
 			// precompiled header if there are at least several source files in the module (after combining them for unity
 			// builds.)  But for game modules, it can be convenient to always have a precompiled header to single-file
@@ -1001,16 +976,13 @@ namespace UnrealBuildTool
 			{
 				bModuleUsesUnityBuild = false;
 			}
-			else if( !BuildConfiguration.bForceUnityBuild && IsGameModule && CPPFiles != null && CPPFiles.Count < BuildConfiguration.MinGameModuleSourceFilesForUnityBuild )
+			else if( !BuildConfiguration.bForceUnityBuild && IsGameModule && SourceFiles.CPPFiles.Count < BuildConfiguration.MinGameModuleSourceFilesForUnityBuild )
 			{
 				// Game modules with only a small number of source files are usually better off having faster iteration times
 				// on single source file changes, so we forcibly disable unity build for those modules
 				bModuleUsesUnityBuild = false;
 			}
 
-
-
-
 			// Precompiled header support.
 			bool bWasModuleCodeCompiled = false;
 			if (BuildPlatform.ShouldUsePCHFiles(CompileEnvironment.Config.TargetPlatform, CompileEnvironment.Config.TargetConfiguration))
@@ -1020,41 +992,26 @@ namespace UnrealBuildTool
 				// The code below will figure out whether this module will either use a "unique PCH" (private PCH that will only be included by
 				// this module's code files), or a "shared PCH" (potentially included by many code files in many modules.)  Only one or the other
 				// will be used.
-				FileItem UniquePCHHeaderFile = null;
 				FileItem SharedPCHHeaderFile = null;
 
 				// In the case of a shared PCH, we also need to keep track of which module that PCH's header file is a member of
 				string SharedPCHModuleName = String.Empty;
 
-				// The precompiled header environment for all source files in this module that use a precompiled header, if we even need one
-				PrecompileHeaderEnvironment ModulePCHEnvironment = null;
-
 				bool bDisableSharedPCHFiles = (Binary.Config.bCompileMonolithic && CompileEnvironment.Config.bIsBuildingLibrary);
 				if( BuildConfiguration.bUseSharedPCHs && bDisableSharedPCHFiles )
 				{
 					Log.TraceVerbose("Module '{0}' was not allowed to use SharedPCHs, because we're compiling to a library in monolithic mode", this.Name );
 				}
 
-				bool bUseSharedPCHFiles = BuildConfiguration.bUseSharedPCHs && (bDisableSharedPCHFiles == false);
-				bool bIsASharedPCHModule = false;
-				if (bUseSharedPCHFiles)
-				{
-					foreach( var CurSharedPCHHeaderFile in GlobalCompileEnvironment.SharedPCHHeaderFiles )
-					{
-						if( this == CurSharedPCHHeaderFile.Module )
-						{
-							bIsASharedPCHModule = true;
-							break;
-						}
-					}
-				}
+				bool bUseSharedPCHFiles  = BuildConfiguration.bUseSharedPCHs && (bDisableSharedPCHFiles == false);
+				bool bIsASharedPCHModule = bUseSharedPCHFiles && GlobalCompileEnvironment.SharedPCHHeaderFiles.Any( PCH => PCH.Module == this );
 
 				// Map from pch header string to the source files that use that PCH
 				var UsageMapPCH = new Dictionary<string, List<FileItem>>( StringComparer.InvariantCultureIgnoreCase );
 
 				// Determine what potential precompiled header is used by each source file.
 				double SharedPCHTotalTime = 0.0;
-				foreach( var CPPFile in CPPFiles )
+				foreach( var CPPFile in SourceFiles.CPPFiles )
 				{
 					if (bUseSharedPCHFiles)
 					{
@@ -1100,14 +1057,11 @@ namespace UnrealBuildTool
 									SharedPCHModuleName = LargestIncludedSharedPCHHeaderFile.Module.Name;
 									SharedPCHHeaderFile = LargestIncludedSharedPCHHeaderFile.PCHHeaderFile;
 								}
-								else
+								else if( SharedPCHHeaderFile != LargestIncludedSharedPCHHeaderFile.PCHHeaderFile )
 								{
-									if( SharedPCHHeaderFile != LargestIncludedSharedPCHHeaderFile.PCHHeaderFile )
-									{
-										// @todo UBT perf: It is fairly costly to perform this test, as we could easily early-out after we have SharedPCHHeaderFile and not bother testing the rest of the files in the module.  But this can be useful to find abusive modules that include a shared PCH header in the bowels of a non-PCH private source file.
-										Console.WriteLine( "WARNING: File '{0}' doesn't use same 'shared' precompiled header as other files in this module: '{1}' vs '{2}'.  This can greatly impact compile times.  Make sure that your module's private PCH header includes the 'largest' shared PCH header that your module uses", 
-											CPPFile.AbsolutePath, SharedPCHHeaderFile, LargestIncludedSharedPCHHeaderFile.PCHHeaderFile );
-									}
+									// @todo UBT perf: It is fairly costly to perform this test, as we could easily early-out after we have SharedPCHHeaderFile and not bother testing the rest of the files in the module.  But this can be useful to find abusive modules that include a shared PCH header in the bowels of a non-PCH private source file.
+									Console.WriteLine( "WARNING: File '{0}' doesn't use same 'shared' precompiled header as other files in this module: '{1}' vs '{2}'.  This can greatly impact compile times.  Make sure that your module's private PCH header includes the 'largest' shared PCH header that your module uses", 
+										CPPFile.AbsolutePath, SharedPCHHeaderFile, LargestIncludedSharedPCHHeaderFile.PCHHeaderFile );
 								}
 							}
 							else
@@ -1123,80 +1077,13 @@ namespace UnrealBuildTool
 						}
 					}
 
-					{
-						// Find headers used by the source file.
-						List<DependencyInclude> DirectIncludeFilenames = CPPEnvironment.GetDirectIncludeDependencies( CPPFile, ModuleCompileEnvironment.Config.TargetPlatform, ModuleCompileEnvironment.bHackHeaderGenerator );
-						if( BuildConfiguration.bPrintDebugInfo )
-						{
-							var IncludedFileNames = new StringBuilder();
-							foreach( var CurInclude in DirectIncludeFilenames )
-							{
-								if( IncludedFileNames.Length > 0 )
-								{
-									IncludedFileNames.Append( ", " );
-								}
-								IncludedFileNames.Append( Path.GetFileName( CurInclude.IncludeName ) );
-							}
-
-							Log.TraceVerbose( "Found direct includes for {0}: {1}", Path.GetFileName( CPPFile.AbsolutePath ), IncludedFileNames );
-						}
-
-						if( DirectIncludeFilenames.Count > 0 )
-						{
-							// The pch header should always be the first include in the source file.
-							// NOTE: This is not an absolute path.  This is just the literal include string from the source file!
-							CPPFile.PCHHeaderNameInCode = DirectIncludeFilenames[ 0 ].IncludeName;
-
-							// Resolve the PCH header to an absolute path.
-							// Check NullOrEmpty here because if the file could not be resolved we need to throw an exception
-							if (string.IsNullOrEmpty(DirectIncludeFilenames[0].IncludeResolvedName) ||
-								// ignore any preexisting resolve cache if we are not configured to use it.
-								!BuildConfiguration.bUseIncludeDependencyResolveCache ||
-								// if we are testing the resolve cache, we force UBT to resolve every time to look for conflicts
-								BuildConfiguration.bTestIncludeDependencyResolveCache)
-							{
-								string SourceFilesDirectory = Path.GetDirectoryName(CPPFile.AbsolutePath);
-								// search the include paths to resolve the file.
-								FileItem PrecompiledHeaderIncludeFile = ModuleCompileEnvironment.FindIncludedFile(CPPFile.PCHHeaderNameInCode, !BuildConfiguration.bCheckExternalHeadersForModification, SourceFilesDirectory);
-								if (PrecompiledHeaderIncludeFile != null)
-								{
-									CPPEnvironment.IncludeDependencyCache.CacheResolvedIncludeFullPath(CPPFile, 0, PrecompiledHeaderIncludeFile.AbsolutePath);
-									CPPFile.PrecompiledHeaderIncludeFilename = PrecompiledHeaderIncludeFile.AbsolutePath;
-
-									if (UniquePCHHeaderFile == null)
-									{
-										UniquePCHHeaderFile = PrecompiledHeaderIncludeFile;
-									}
-								}
-								else
-								{
-									throw new BuildException("The first include statement in source file '{0}' is trying to include the file '{1}' as the precompiled header for module '{2}', but that file could not be located in any of the module's include search paths.", CPPFile.AbsolutePath, CPPFile.PCHHeaderNameInCode, this.Name);
-								}
-							}
-							else
-							{
-								CPPFile.PrecompiledHeaderIncludeFilename = DirectIncludeFilenames[0].IncludeResolvedName;
-								if (UniquePCHHeaderFile == null)
-								{
-									UniquePCHHeaderFile = FileItem.GetItemByFullPath(CPPFile.PrecompiledHeaderIncludeFilename);
-								}
-							}
-						}
-					}
-
 					if( CPPFile.PrecompiledHeaderIncludeFilename == null )
 					{
 						throw new BuildException( "No PCH usage for file \"{0}\" . Missing #include header?", CPPFile.AbsolutePath );
 					}
 
 					// Create a new entry if not in the pch usage map
-					List<FileItem> PCHSourceFiles = null;
-					if( !UsageMapPCH.TryGetValue( CPPFile.PrecompiledHeaderIncludeFilename, out PCHSourceFiles ) )
-					{
-						PCHSourceFiles = new List<FileItem>();
-						UsageMapPCH.Add( CPPFile.PrecompiledHeaderIncludeFilename, PCHSourceFiles );
-					}
-					PCHSourceFiles.Add( CPPFile );
+					UsageMapPCH.GetOrAddNew( CPPFile.PrecompiledHeaderIncludeFilename ).Add( CPPFile );
 				}
 
 				if( BuildConfiguration.bPrintDebugInfo )
@@ -1219,25 +1106,19 @@ namespace UnrealBuildTool
 					// Keep track of the PCH file that is most used within this module
 					string MostFilesAreIncludingPCH = string.Empty;
 					int MostFilesIncluded = 0;
-					foreach( var CurPCH in UsageMapPCH )
+					foreach( var CurPCH in UsageMapPCH.Where( PCH => PCH.Value.Count > MostFilesIncluded ) )
 					{
-						if( CurPCH.Value.Count > MostFilesIncluded )
-						{
-							MostFilesAreIncludingPCH = CurPCH.Key;
-							MostFilesIncluded = CurPCH.Value.Count;
-						}
+						MostFilesAreIncludingPCH = CurPCH.Key;
+						MostFilesIncluded        = CurPCH.Value.Count;
 					}
 
 					// Find all of the files that are not including our "best" PCH header
 					var FilesNotIncludingBestPCH = new StringBuilder();
-					foreach( var CurPCH in UsageMapPCH )
+					foreach( var CurPCH in UsageMapPCH.Where( PCH => PCH.Key != MostFilesAreIncludingPCH ) )
 					{
-						if( CurPCH.Key != MostFilesAreIncludingPCH )
+						foreach( var SourceFile in CurPCH.Value )
 						{
-							foreach( var SourceFile in CurPCH.Value )
-							{
-								FilesNotIncludingBestPCH.AppendFormat( "{0} (including {1})\n", SourceFile.AbsolutePath, CurPCH.Key );
-							}
+							FilesNotIncludingBestPCH.AppendFormat( "{0} (including {1})\n", SourceFile.AbsolutePath, CurPCH.Key );
 						}
 					}
 
@@ -1249,80 +1130,41 @@ namespace UnrealBuildTool
 						FilesNotIncludingBestPCH );
 				}
 
+				// The precompiled header environment for all source files in this module that use a precompiled header, if we even need one
+				PrecompileHeaderEnvironment ModulePCHEnvironment = null;
 
 				// If there was one header that was included first by enough C++ files, use it as the precompiled header.
 				// Only use precompiled headers for projects with enough files to make the PCH creation worthwhile.
-				if( CPPFiles != null &&
-					( SharedPCHHeaderFile != null || CPPFiles.Count >= MinFilesUsingPrecompiledHeader ) )
+				if( SharedPCHHeaderFile != null || SourceFiles.CPPFiles.Count >= MinFilesUsingPrecompiledHeader )
 				{
 					if( SharedPCHHeaderFile != null )
 					{
-						// Check to see if we have a PCH header already setup that we can use
-						foreach( var SharedPCHEnvironment in GlobalCompileEnvironment.SharedPCHEnvironments )
+						ModulePCHEnvironment = ApplySharedPCH(GlobalCompileEnvironment, CompileEnvironment, ModuleCompileEnvironment, SourceFiles.CPPFiles, SharedPCHHeaderFile);
+						if (ModulePCHEnvironment != null)
 						{
-							if( SharedPCHEnvironment.PrecompiledHeaderIncludeFilename == SharedPCHHeaderFile )
-							{
-								// Don't mix CLR modes
-								if( SharedPCHEnvironment.CLRMode == ModuleCompileEnvironment.Config.CLRMode )
-								{
-									// Don't mix non-optimized code with optimized code (PCHs won't be compatible)
-									bool bIsDebugBuild = CompileEnvironment.Config.TargetConfiguration == CPPTargetConfiguration.Debug;
-									var SharedPCHCodeOptimization = SharedPCHEnvironment.OptimizeCode;
-									if( SharedPCHCodeOptimization == ModuleRules.CodeOptimization.InNonDebugBuilds && !bIsDebugBuild )
-									{
-										SharedPCHCodeOptimization = ModuleRules.CodeOptimization.Always;
-									}
-									var ModuleCodeOptimization = ModuleCompileEnvironment.Config.OptimizeCode;
-									if( ModuleCodeOptimization == ModuleRules.CodeOptimization.InNonDebugBuilds && !bIsDebugBuild )
-									{
-										ModuleCodeOptimization = ModuleRules.CodeOptimization.Always;
-									}
-
-									if( SharedPCHCodeOptimization.Equals( ModuleCodeOptimization ) )
-									{
-										// Found a valid shared PCH header to use!
-
-										// @todo SharedPCH: Ideally we would exhaustively check for a compatible compile environment (definitions, imports/exports, etc)
-										//    Currently, it's possible for the shared PCH to be compiled differently depending on which module UBT happened to have
-										//    include it first during the build phase.  This could create problems with determinstic builds, or turn up compile
-										//    errors unexpectedly due to compile environment differences.
-										ModulePCHEnvironment = SharedPCHEnvironment;
-										Log.TraceVerbose( "Module " + Name + " uses existing SharedPCH '" + ModulePCHEnvironment.PrecompiledHeaderIncludeFilename + "' (from module " + ModulePCHEnvironment.ModuleName + ")" );
-
-										// Update all CPPFiles to point to the shared PCH
-										foreach( var CPPFile in CPPFiles )
-										{
-											CPPFile.PCHHeaderNameInCode = SharedPCHHeaderFile.AbsolutePath;
-											CPPFile.PrecompiledHeaderIncludeFilename = SharedPCHHeaderFile.AbsolutePath;
-										}
-									}
-									else
-									{
-										Log.TraceVerbose( "Module {0} cannot use existing SharedPCH '{1}' (from module '{2}') because optimization levels don't match", Name, SharedPCHEnvironment.PrecompiledHeaderIncludeFilename.AbsolutePath, SharedPCHEnvironment.ModuleName );
-										SharedPCHHeaderFile = null;
-									}
-								}
-								else
-								{
-									Log.TraceVerbose( "Module {0} cannot use existing SharedPCH '{1}' (from module '{2}') because CLR modes don't match", Name, SharedPCHEnvironment.PrecompiledHeaderIncludeFilename.AbsolutePath, SharedPCHEnvironment.ModuleName );
-									SharedPCHHeaderFile = null;
-								}
-								break;
-							}
+							// @todo SharedPCH: Ideally we would exhaustively check for a compatible compile environment (definitions, imports/exports, etc)
+							//    Currently, it's possible for the shared PCH to be compiled differently depending on which module UBT happened to have
+							//    include it first during the build phase.  This could create problems with determinstic builds, or turn up compile
+							//    errors unexpectedly due to compile environment differences.
+							Log.TraceVerbose("Module " + Name + " uses existing SharedPCH '" + ModulePCHEnvironment.PrecompiledHeaderIncludeFilename + "' (from module " + ModulePCHEnvironment.ModuleName + ")");
+						}
+						else
+						{
+							SharedPCHHeaderFile = null;
 						}
 					}
 
 					// A shared PCH was not already set up for us, so set one up.
 					if( ModulePCHEnvironment == null )
 					{
-						var PCHHeaderFile = UniquePCHHeaderFile;
+						var PCHHeaderFile = ProcessedDependencies.UniquePCHHeaderFile;
 						var PCHModuleName = this.Name;
 						if( SharedPCHHeaderFile != null )
 						{
 							PCHHeaderFile = SharedPCHHeaderFile;
 							PCHModuleName = SharedPCHModuleName;
 						}
-						var PCHHeaderNameInCode = CPPFiles[ 0 ].PCHHeaderNameInCode;
+						var PCHHeaderNameInCode = SourceFiles.CPPFiles[ 0 ].PCHHeaderNameInCode;
 
 						ModulePCHEnvironment = new PrecompileHeaderEnvironment( PCHModuleName, PCHHeaderNameInCode, PCHHeaderFile, ModuleCompileEnvironment.Config.CLRMode, ModuleCompileEnvironment.Config.OptimizeCode );
 
@@ -1347,20 +1189,18 @@ namespace UnrealBuildTool
 					// Setup a new compile environment for this module's source files.  It's pretty much the exact same as the
 					// module's compile environment, except that it will include a PCH file.
 					var ModulePCHCompileEnvironment = new CPPEnvironment( ModuleCompileEnvironment );
-					{
-						ModulePCHCompileEnvironment.Config.PrecompiledHeaderAction = PrecompiledHeaderAction.Include;
-						ModulePCHCompileEnvironment.Config.PrecompiledHeaderIncludeFilename = ModulePCHEnvironment.PrecompiledHeaderIncludeFilename.AbsolutePath;
-						ModulePCHCompileEnvironment.Config.PCHHeaderNameInCode = ModulePCHEnvironment.PCHHeaderNameInCode;
+					ModulePCHCompileEnvironment.Config.PrecompiledHeaderAction          = PrecompiledHeaderAction.Include;
+					ModulePCHCompileEnvironment.Config.PrecompiledHeaderIncludeFilename = ModulePCHEnvironment.PrecompiledHeaderIncludeFilename.AbsolutePath;
+					ModulePCHCompileEnvironment.Config.PCHHeaderNameInCode              = ModulePCHEnvironment.PCHHeaderNameInCode;
 
-						if( SharedPCHHeaderFile != null )
-						{
-							// Shared PCH headers need to be force included, because we're basically forcing the module to use
-							// the precompiled header that we want, instead of the "first include" in each respective .cpp file
-							ModulePCHCompileEnvironment.Config.bForceIncludePrecompiledHeader = true;
-						}
+					if( SharedPCHHeaderFile != null )
+					{
+						// Shared PCH headers need to be force included, because we're basically forcing the module to use
+						// the precompiled header that we want, instead of the "first include" in each respective .cpp file
+						ModulePCHCompileEnvironment.Config.bForceIncludePrecompiledHeader = true;
 					}
 
-					var CPPFilesToBuild = CPPFiles;
+					var CPPFilesToBuild = SourceFiles.CPPFiles;
 					if (bModuleUsesUnityBuild)
 					{
 						// unity files generated for only the set of files which share the same PCH environment
@@ -1368,8 +1208,7 @@ namespace UnrealBuildTool
 					}
 
 					// Check if there are enough unity files to warrant pch generation (and we haven't already generated the shared one)
-					if( ModulePCHEnvironment.PrecompiledHeaderFile == null &&
-						( SharedPCHHeaderFile != null || CPPFilesToBuild.Count >= MinFilesUsingPrecompiledHeader ) )
+					if( ModulePCHEnvironment.PrecompiledHeaderFile == null && ( SharedPCHHeaderFile != null || CPPFilesToBuild.Count >= MinFilesUsingPrecompiledHeader ) )
 					{
 						bool bAllowDLLExports = true;
 						var PCHOutputDirectory = ModuleCompileEnvironment.Config.OutputDirectory;
@@ -1420,32 +1259,42 @@ namespace UnrealBuildTool
 
 				if( BuildConfiguration.bPrintPerformanceInfo )
 				{
-				var TotalPCHTime = DateTime.UtcNow - PCHTimerStart;
+					var TotalPCHTime = DateTime.UtcNow - PCHTimerStart;
 					Trace.TraceInformation( "PCH time for " + Name + " is " + TotalPCHTime.TotalSeconds + "s (shared PCHs: " + SharedPCHTotalTime + "s)" );
 				}
 			}
 
-			if( !bWasModuleCodeCompiled )
+			if( !bWasModuleCodeCompiled && SourceFiles.CPPFiles.Count > 0 )
 			{
-				if( CPPFiles.Count > 0 )
+				var CPPFilesToCompile = SourceFiles.CPPFiles;
+				if (bModuleUsesUnityBuild)
 				{
-					var CPPFilesToCompile = CPPFiles;
-					if (bModuleUsesUnityBuild)
-					{
-						CPPFilesToCompile = Unity.GenerateUnityCPPs( CPPFilesToCompile, ModuleCompileEnvironment, Name );
-					}
-					LinkInputFiles.AddRange( ModuleCompileEnvironment.CompileFiles( CPPFilesToCompile, Name ).ObjectFiles );
+					CPPFilesToCompile = Unity.GenerateUnityCPPs( CPPFilesToCompile, ModuleCompileEnvironment, Name );
 				}
+				LinkInputFiles.AddRange( ModuleCompileEnvironment.CompileFiles( CPPFilesToCompile, Name ).ObjectFiles );
+			}
+
+			string PCHPath = null;
+			if (ProcessedDependencies.UniquePCHHeaderFile != null)
+			{
+				PCHPath = UEToolChain.GetPlatformToolChain(ModuleCompileEnvironment.Config.TargetPlatform).ConvertPath( ProcessedDependencies.UniquePCHHeaderFile.AbsolutePath );
+			}
+
+			if (AutoGenerateInlInfo != null && !ModuleCompileEnvironment.bHackHeaderGenerator)
+			{
+				var GeneratedCppFileItem = FileItem.GetItemByPath( AutoGenerateInlInfo.Filename );
+
+				LinkInputFiles.AddRange( ModuleCompileEnvironment.CompileFiles( new List<FileItem>{ GeneratedCppFileItem }, Name ).ObjectFiles );
 			}
 
 			// Compile C files directly.
-			LinkInputFiles.AddRange(ModuleCompileEnvironment.CompileFiles(CFiles, Name).ObjectFiles);
+			LinkInputFiles.AddRange(ModuleCompileEnvironment.CompileFiles(SourceFiles.CFiles, Name).ObjectFiles);
 
 			// Compile CC files directly.
-			LinkInputFiles.AddRange(ModuleCompileEnvironment.CompileFiles(CCFiles, Name).ObjectFiles);
+			LinkInputFiles.AddRange(ModuleCompileEnvironment.CompileFiles(SourceFiles.CCFiles, Name).ObjectFiles);
 
 			// Compile MM files directly.
-			LinkInputFiles.AddRange(ModuleCompileEnvironment.CompileFiles(MMFiles, Name).ObjectFiles);
+			LinkInputFiles.AddRange(ModuleCompileEnvironment.CompileFiles(SourceFiles.MMFiles, Name).ObjectFiles);
 
 			// If we're building Rocket, generate a static library for this module
 			if(RedistStaticLibraryPath != null)
@@ -1453,13 +1302,13 @@ namespace UnrealBuildTool
 				// Create a link environment for it
 				LinkEnvironment RedistLinkEnvironment = new LinkEnvironment();
 				RedistLinkEnvironment.InputFiles.AddRange(LinkInputFiles);
-				RedistLinkEnvironment.Config.TargetArchitecture = CompileEnvironment.Config.TargetArchitecture;
-				RedistLinkEnvironment.Config.TargetConfiguration = CompileEnvironment.Config.TargetConfiguration;
-				RedistLinkEnvironment.Config.TargetPlatform = CompileEnvironment.Config.TargetPlatform;
-				RedistLinkEnvironment.Config.bIsBuildingDLL = false;
-				RedistLinkEnvironment.Config.bIsBuildingLibrary = true;
+				RedistLinkEnvironment.Config.TargetArchitecture    = CompileEnvironment.Config.TargetArchitecture;
+				RedistLinkEnvironment.Config.TargetConfiguration   = CompileEnvironment.Config.TargetConfiguration;
+				RedistLinkEnvironment.Config.TargetPlatform        = CompileEnvironment.Config.TargetPlatform;
+				RedistLinkEnvironment.Config.bIsBuildingDLL        = false;
+				RedistLinkEnvironment.Config.bIsBuildingLibrary    = true;
 				RedistLinkEnvironment.Config.IntermediateDirectory = Binary.Config.IntermediateDirectory;
-				RedistLinkEnvironment.Config.OutputFilePath = RedistStaticLibraryPath;
+				RedistLinkEnvironment.Config.OutputFilePath        = RedistStaticLibraryPath;
 
 				// Replace the items built so far with the library
 				RedistLinkEnvironment.LinkExecutable(false);
@@ -1467,87 +1316,149 @@ namespace UnrealBuildTool
 			}
 
 			// Compile RC files.
-			LinkInputFiles.AddRange(ModuleCompileEnvironment.CompileRCFiles(RCFiles).ObjectFiles);
+			LinkInputFiles.AddRange(ModuleCompileEnvironment.CompileRCFiles(SourceFiles.RCFiles).ObjectFiles);
 
-			
-			// Keep track of this module's public and private UObject source files, so that we can pass those off to UHT if needed
+			return LinkInputFiles;
+		}
+
+		private PrecompileHeaderEnvironment ApplySharedPCH(CPPEnvironment GlobalCompileEnvironment, CPPEnvironment CompileEnvironment, CPPEnvironment ModuleCompileEnvironment, List<FileItem> CPPFiles, FileItem SharedPCHHeaderFile)
+		{
+			// Check to see if we have a PCH header already setup that we can use
+			var SharedPCHEnvironment = GlobalCompileEnvironment.SharedPCHEnvironments.Find(Env => Env.PrecompiledHeaderIncludeFilename == SharedPCHHeaderFile);
+			if (SharedPCHEnvironment == null)
+				return null;
+
+			// Don't mix CLR modes
+			if (SharedPCHEnvironment.CLRMode != ModuleCompileEnvironment.Config.CLRMode)
 			{
-				string ModuleSourceFolder = Path.GetFullPath( this.ModuleDirectory );
-				foreach( var SourceFile in SourceFiles )
+				Log.TraceVerbose("Module {0} cannot use existing SharedPCH '{1}' (from module '{2}') because CLR modes don't match", Name, SharedPCHEnvironment.PrecompiledHeaderIncludeFilename.AbsolutePath, SharedPCHEnvironment.ModuleName);
+				return null;
+			}
+
+			// Don't mix non-optimized code with optimized code (PCHs won't be compatible)
+			var SharedPCHCodeOptimization = SharedPCHEnvironment.OptimizeCode;
+			var ModuleCodeOptimization    = ModuleCompileEnvironment.Config.OptimizeCode;
+
+			if (CompileEnvironment.Config.TargetConfiguration != CPPTargetConfiguration.Debug)
+			{
+				if (SharedPCHCodeOptimization == ModuleRules.CodeOptimization.InNonDebugBuilds)
 				{
-					// Will always be a cache hit (we did this earlier during Compile())
-					var IncludedFiles = ModuleCompileEnvironment.GetIncludeDependencies( SourceFile );
+					SharedPCHCodeOptimization = ModuleRules.CodeOptimization.Always;
+				}
 
-					// Also check for intrinsic classes like "Object.h", which are special cases because they are never included in compiled code and exist only for UHT to parse
-					{
-						// Runtime/CoreUObject/Classes/Object.h
-						{
-							var IntrinsicFileItem = FileItem.GetExistingItemByPath( Path.Combine( ProjectFileGenerator.EngineRelativePath, "Source", "Runtime", Path.Combine( "CoreUObject", "Classes", "Object.h" ) ) );	// @todo uht: In Classes folder right now
-							if( !IntrinsicFileItem.bExists )
-							{
-								throw new BuildException( "Expecting " + IntrinsicFileItem.AbsolutePath + " to exist" );
-							}
-							IntrinsicFileItem.HasUObjects = true;
-							IncludedFiles.Add( IntrinsicFileItem );
-						}
+				if (ModuleCodeOptimization == ModuleRules.CodeOptimization.InNonDebugBuilds)
+				{
+					ModuleCodeOptimization = ModuleRules.CodeOptimization.Always;
+				}
+			}
 
-						// Runtime/Engine/Classes/Model.h
-						{
-							var IntrinsicFileItem = FileItem.GetExistingItemByPath( Path.Combine( ProjectFileGenerator.EngineRelativePath, "Source", "Runtime", Path.Combine( "Engine", "Classes", "Intrinsic", "Model.h" ) ) );	// @todo uht: In Classes folder right now
-							if( !IntrinsicFileItem.bExists )
-							{
-								throw new BuildException( "Expecting " + IntrinsicFileItem.AbsolutePath + " to exist" );
-							}
-							IntrinsicFileItem.HasUObjects = true;
-							IncludedFiles.Add( IntrinsicFileItem );
-						}
-					}
+			if (SharedPCHCodeOptimization != ModuleCodeOptimization)
+			{
+				Log.TraceVerbose("Module {0} cannot use existing SharedPCH '{1}' (from module '{2}') because optimization levels don't match", Name, SharedPCHEnvironment.PrecompiledHeaderIncludeFilename.AbsolutePath, SharedPCHEnvironment.ModuleName);
+				return null;
+			}
 
-					// @todo uht: Could check SourceFile.HasUObjects if we want to include .cpps with USTRUCT/UCLASSES here
-					foreach( var IncludedFile in IncludedFiles )
-					{
-						if( IncludedFile.HasUObjects )
-						{
-							if( IncludedFile.AbsolutePath.StartsWith( ModuleSourceFolder + Path.DirectorySeparatorChar ) )
-							{
-								// Is it private or public?
-								bool bIsPublic = false;
-								{
-									// Get the part of the path that is relative to the module source folder
-									var RelativeSourceFilePath = IncludedFile.AbsolutePath.Substring( ModuleSourceFolder.Length + 1 );
-									if( RelativeSourceFilePath.StartsWith( "Classes" + Path.DirectorySeparatorChar, StringComparison.InvariantCultureIgnoreCase ) )
-									{
-										// Code files under the legacy 'Classes' directory are always considered public
-										bIsPublic = true;
-									}
-									else if( RelativeSourceFilePath.StartsWith( "Public" + Path.DirectorySeparatorChar, StringComparison.InvariantCultureIgnoreCase ) ||
-											 ( RelativeSourceFilePath.IndexOf( Path.DirectorySeparatorChar + "Public" + Path.DirectorySeparatorChar, StringComparison.InvariantCultureIgnoreCase ) != -1 ) )
-									{
-										// Code is under a 'Public' subdirectory somewhere in the hierarchy
-										bIsPublic = true;
-									}
-								}
+			// Update all CPPFiles to point to the shared PCH
+			foreach (var CPPFile in CPPFiles)
+			{
+				CPPFile.PCHHeaderNameInCode              = SharedPCHHeaderFile.AbsolutePath;
+				CPPFile.PrecompiledHeaderIncludeFilename = SharedPCHHeaderFile.AbsolutePath;
+			}
 
-								if( bIsPublic )
-								{
-									PublicUObjectHeaders.Add( IncludedFile );
-								}
-								else 
-								{
-									PrivateUObjectHeaders.Add( IncludedFile );
-								}
-							}
-							else
-							{
-								// Not from this module, so we don't need to worry about UObjects
-							}
-						}
-					}
+			return SharedPCHEnvironment;
+		}
 
-				}
+		public void ProcessAllCppDependencies(CPPEnvironment ModuleCompileEnvironment)
+		{
+			if (ProcessedDependencies != null)
+				return;
+
+			FileItem UniquePCH = null;
+			foreach( var CPPFile in SourceFiles.CPPFiles )
+			{
+				// Find headers used by the source file.
+				var PCH = ProcessDependencies(CPPFile, ModuleCompileEnvironment);
+				UniquePCH = UniquePCH ?? PCH;
 			}
 
-			return LinkInputFiles;
+			ProcessedDependencies = new ProcessedDependenciesClass{ UniquePCHHeaderFile = UniquePCH };
+		}
+
+		private FileItem ProcessDependencies(FileItem CPPFile, CPPEnvironment ModuleCompileEnvironment)
+		{
+			List<DependencyInclude> DirectIncludeFilenames = CPPEnvironment.GetDirectIncludeDependencies(CPPFile, ModuleCompileEnvironment.Config.TargetPlatform, ModuleCompileEnvironment.bHackHeaderGenerator);
+			if (BuildConfiguration.bPrintDebugInfo)
+			{
+				Log.TraceVerbose("Found direct includes for {0}: {1}", Path.GetFileName(CPPFile.AbsolutePath), string.Join(", ", DirectIncludeFilenames.Select(F => F.IncludeName)));
+			}
+
+			if (DirectIncludeFilenames.Count == 0)
+				return null;
+
+			var FirstInclude = DirectIncludeFilenames[0];
+
+			// The pch header should always be the first include in the source file.
+			// NOTE: This is not an absolute path.  This is just the literal include string from the source file!
+			CPPFile.PCHHeaderNameInCode = FirstInclude.IncludeName;
+
+			// Resolve the PCH header to an absolute path.
+			// Check NullOrEmpty here because if the file could not be resolved we need to throw an exception
+			if (!string.IsNullOrEmpty(FirstInclude.IncludeResolvedName) &&
+				// ignore any preexisting resolve cache if we are not configured to use it.
+				BuildConfiguration.bUseIncludeDependencyResolveCache &&
+				// if we are testing the resolve cache, we force UBT to resolve every time to look for conflicts
+				!BuildConfiguration.bTestIncludeDependencyResolveCache)
+			{
+				CPPFile.PrecompiledHeaderIncludeFilename = FirstInclude.IncludeResolvedName;
+				return FileItem.GetItemByFullPath(CPPFile.PrecompiledHeaderIncludeFilename);
+			}
+
+			string SourceFilesDirectory = Path.GetDirectoryName(CPPFile.AbsolutePath);
+
+			// search the include paths to resolve the file.
+			FileItem PrecompiledHeaderIncludeFile = ModuleCompileEnvironment.FindIncludedFile(CPPFile.PCHHeaderNameInCode, !BuildConfiguration.bCheckExternalHeadersForModification, SourceFilesDirectory);
+			if (PrecompiledHeaderIncludeFile == null)
+				throw new BuildException("The first include statement in source file '{0}' is trying to include the file '{1}' as the precompiled header for module '{2}', but that file could not be located in any of the module's include search paths.", CPPFile.AbsolutePath, CPPFile.PCHHeaderNameInCode, this.Name);
+
+			CPPEnvironment.IncludeDependencyCache.CacheResolvedIncludeFullPath(CPPFile, 0, PrecompiledHeaderIncludeFile.AbsolutePath);
+			CPPFile.PrecompiledHeaderIncludeFilename = PrecompiledHeaderIncludeFile.AbsolutePath;
+
+			return PrecompiledHeaderIncludeFile;
+		}
+
+		/// <summary>
+		/// Creates a compile environment from a base environment based on the module settings.
+		/// </summary>
+		/// <param name="BaseCompileEnvironment">An existing environment to base the module compile environment on.</param>
+		/// <returns>The new module compile environment.</returns>
+		public CPPEnvironment CreateModuleCompileEnvironment(CPPEnvironment BaseCompileEnvironment)
+		{
+			var Result = new CPPEnvironment(BaseCompileEnvironment);
+
+			// Override compile environment
+			Result.Config.bFasterWithoutUnity                    = bFasterWithoutUnity;
+			Result.Config.OptimizeCode                           = OptimizeCode;
+			Result.Config.bUseRTTI                               = bUseRTTI;
+			Result.Config.bEnableBufferSecurityChecks            = bEnableBufferSecurityChecks;
+			Result.Config.bFasterWithoutUnity                    = bFasterWithoutUnity;
+			Result.Config.MinFilesUsingPrecompiledHeaderOverride = MinFilesUsingPrecompiledHeaderOverride;
+			Result.Config.bEnableExceptions                      = bEnableExceptions;
+			Result.Config.bEnableInlining                        = bEnableInlining;
+			Result.Config.OutputDirectory                        = Path.Combine(Binary.Config.IntermediateDirectory, Name);
+
+			// Switch the optimization flag if we're building a game module
+			if (Target.Configuration == UnrealTargetConfiguration.DebugGame && Type == UEBuildModuleType.GameModule)
+			{
+				Result.Config.TargetConfiguration = CPPTargetConfiguration.Debug;
+			}
+
+			// Add the module's private definitions.
+			Result.Config.Definitions.AddRange(Definitions);
+
+			// Setup the compile environment for the module.
+			SetupPrivateCompileEnvironment(ref Result.Config.IncludePaths, ref Result.Config.SystemIncludePaths, ref Result.Config.Definitions);
+
+			return Result;
 		}
 
 		public override void SetupPrivateLinkEnvironment(
@@ -1621,6 +1532,12 @@ namespace UnrealBuildTool
 			Writer.WriteEndElement();
 		}
 
+		public HashSet<FileItem> _AllClassesHeaders     = new HashSet<FileItem>();
+		public HashSet<FileItem> _PublicUObjectHeaders  = new HashSet<FileItem>();
+		public HashSet<FileItem> _PrivateUObjectHeaders = new HashSet<FileItem>();
+
+		private UHTModuleInfo? CachedModuleUHTInfo = null;
+
 		/// <summary>
 		/// If any of this module's source files contain UObject definitions, this will return those header files back to the caller
 		/// </summary>
@@ -1628,14 +1545,16 @@ namespace UnrealBuildTool
 		/// <param name="PublicUObjectHeaders">Dependent UObject headers in Public source folders</param>
 		/// <param name="PrivateUObjectHeaders">Dependent UObject headers not in Public source folders</param>
 		/// <returns>
-		public UHTModuleInfo GetUHTModuleInfo()
+		public UHTModuleInfo GetUHTModuleInfo(CPPEnvironment CompileEnvironment)
 		{
+			if (CachedModuleUHTInfo.HasValue)
+				return CachedModuleUHTInfo.Value;
+
 			// If this module has a "Classes" folder, we'll assume that all headers in that folder are public headers, whether they were
 			// included by #include statements or not.  This is because in the old system, we relied on generated *Classes.h mega-headers
 			// that auto-included all of the files in that folder, so programmers didn't have to bother including the files themselves.
 			// With the new system (files in outside of the "Classes" folders), you always need to include what you're using!  We can't follow
 			// include dependencies on generated "*Classes.h" files that may not even exist at the time we're gathering UObject headers.
-			var AllClassesHeaders = new HashSet<FileItem>();
 			{
 				var ClassesFolder = Path.Combine( this.ModuleDirectory, "Classes" );
 				if( Directory.Exists( ClassesFolder ) )
@@ -1645,21 +1564,91 @@ namespace UnrealBuildTool
 					{
 						var UObjectHeaderFileItem = FileItem.GetExistingItemByPath( ClassHeader );
 						UObjectHeaderFileItem.HasUObjects = true;
-						AllClassesHeaders.Add( UObjectHeaderFileItem );
+						_AllClassesHeaders.Add( UObjectHeaderFileItem );
+					}
+				}
+			}
+
+			// Keep track of this module's public and private UObject source files, so that we can pass those off to UHT if needed
+			string ModuleSourceFolder = Path.GetFullPath(this.ModuleDirectory);
+			foreach (var SourceFile in SourceFiles.CFiles.Concat(SourceFiles.CPPFiles).Concat(SourceFiles.CCFiles).Concat(SourceFiles.MMFiles))
+			{
+				// Will always be a cache hit (we did this earlier during Compile())
+				var IncludedFiles = CompileEnvironment.GetIncludeDependencies(SourceFile);
+
+				// Also check for intrinsic classes like "Object.h", which are special cases because they are never included in compiled code and exist only for UHT to parse
+
+				// Runtime/CoreUObject/Classes/Object.h
+				{
+					// @todo uht: In Classes folder right now
+					var IntrinsicFileItem = FileItem.GetExistingItemByPath(Path.Combine(ProjectFileGenerator.EngineRelativePath, "Source", "Runtime", "CoreUObject", "Classes", "Object.h"));
+					if (!IntrinsicFileItem.bExists)
+					{
+						throw new BuildException("Expecting " + IntrinsicFileItem.AbsolutePath + " to exist");
+					}
+					IntrinsicFileItem.HasUObjects = true;
+					IncludedFiles.Add(IntrinsicFileItem);
+				}
+
+				// Runtime/Engine/Classes/Model.h
+				{
+					// @todo uht: In Classes folder right now
+					var IntrinsicFileItem = FileItem.GetExistingItemByPath(Path.Combine(ProjectFileGenerator.EngineRelativePath, "Source", "Runtime", "Engine", "Classes", "Intrinsic", "Model.h"));
+					if (!IntrinsicFileItem.bExists)
+					{
+						throw new BuildException("Expecting " + IntrinsicFileItem.AbsolutePath + " to exist");
+					}
+					IntrinsicFileItem.HasUObjects = true;
+					IncludedFiles.Add(IntrinsicFileItem);
+				}
+
+				// @todo uht: Could check SourceFile.HasUObjects if we want to include .cpps with USTRUCT/UCLASSES here
+				foreach (var IncludedFile in IncludedFiles)
+				{
+					// Check if it's not from this module, so we don't need to worry about UObjects
+					if (!IncludedFile.HasUObjects || !IncludedFile.AbsolutePath.StartsWith(ModuleSourceFolder + Path.DirectorySeparatorChar))
+						continue;
+
+					// Is it private or public?
+					bool bIsPublic = false;
+
+					// Get the part of the path that is relative to the module source folder
+					var RelativeSourceFilePath = IncludedFile.AbsolutePath.Substring(ModuleSourceFolder.Length + 1);
+					if (RelativeSourceFilePath.StartsWith("Classes" + Path.DirectorySeparatorChar, StringComparison.InvariantCultureIgnoreCase))
+					{
+						// Code files under the legacy 'Classes' directory are always considered public
+						bIsPublic = true;
+					}
+					else if (RelativeSourceFilePath.StartsWith("Public" + Path.DirectorySeparatorChar, StringComparison.InvariantCultureIgnoreCase) ||
+								(RelativeSourceFilePath.IndexOf(Path.DirectorySeparatorChar + "Public" + Path.DirectorySeparatorChar, StringComparison.InvariantCultureIgnoreCase) != -1))
+					{
+						// Code is under a 'Public' subdirectory somewhere in the hierarchy
+						bIsPublic = true;
+					}
+
+					if (bIsPublic)
+					{
+						_PublicUObjectHeaders.Add(IncludedFile);
+					}
+					else
+					{
+						_PrivateUObjectHeaders.Add(IncludedFile);
 					}
 				}
 			}
 
-			var AllClassesHeadersList = AllClassesHeaders.ToList();
+			var AllClassesHeadersList = _AllClassesHeaders.ToList();
 
-			return new UHTModuleInfo {
+			CachedModuleUHTInfo = new UHTModuleInfo {
 				ModuleName                  = this.Name,
 				ModuleDirectory             = this.ModuleDirectory,
 				PublicUObjectClassesHeaders = AllClassesHeadersList,
 				// Don't report public classes headers twice.  This logic can go away if we deprecate Classes folders.
-				PublicUObjectHeaders        = this.PublicUObjectHeaders.Where(File => !AllClassesHeadersList.Contains(File)).ToList(),
-				PrivateUObjectHeaders       = this.PrivateUObjectHeaders.ToList()
+				PublicUObjectHeaders        = _PublicUObjectHeaders.Where(File => !AllClassesHeadersList.Contains(File)).ToList(),
+				PrivateUObjectHeaders       = _PrivateUObjectHeaders.ToList()
 			};
+
+			return CachedModuleUHTInfo.Value;
 		}
 
 		public override void GetAllDependencyModules( ref Dictionary<string, UEBuildModule> ReferencedModules, ref List<UEBuildModule> OrderedModules, bool bIncludeDynamicallyLoaded, bool bForceCircular )
diff --git a/Engine/Source/Programs/UnrealBuildTool/Configuration/UEBuildPlatform.cs b/Engine/Source/Programs/UnrealBuildTool/Configuration/UEBuildPlatform.cs
--- a/Engine/Source/Programs/UnrealBuildTool/Configuration/UEBuildPlatform.cs
+++ b/Engine/Source/Programs/UnrealBuildTool/Configuration/UEBuildPlatform.cs
@@ -92,15 +92,7 @@ namespace UnrealBuildTool
 		public static void RegisterPlatformWithGroup(UnrealTargetPlatform InPlatform, UnrealPlatformGroup InGroup)
 		{
 			// find or add the list of groups for this platform
-			List<UnrealTargetPlatform> PlatformList;
-			if (PlatformGroupDictionary.TryGetValue(InGroup, out PlatformList) == false)
-			{
-				PlatformList = new List<UnrealTargetPlatform>();
-				PlatformGroupDictionary[InGroup] = PlatformList;
-			}
-
-			// add the platform to the list of platforms the group has
-			PlatformList.Add(InPlatform);
+			PlatformGroupDictionary.GetOrAddNew(InGroup).Add(InPlatform);
 		}
 
 		/**
diff --git a/Engine/Source/Programs/UnrealBuildTool/Configuration/UEBuildTarget.cs b/Engine/Source/Programs/UnrealBuildTool/Configuration/UEBuildTarget.cs
--- a/Engine/Source/Programs/UnrealBuildTool/Configuration/UEBuildTarget.cs
+++ b/Engine/Source/Programs/UnrealBuildTool/Configuration/UEBuildTarget.cs
@@ -1364,50 +1364,35 @@ namespace UnrealBuildTool
 				GlobalCompileEnvironment.SharedPCHHeaderFiles = SharedPCHHeaderFiles;
 			}
 
-			foreach (var Binary in AppBinaries)
-			{
-				OutputItems.AddRange(Binary.Build(GlobalCompileEnvironment, GlobalLinkEnvironment));
-			}
-
 			if( (BuildConfiguration.bXGEExport && UEBuildConfiguration.bGenerateManifest) || (!UEBuildConfiguration.bGenerateManifest && !UEBuildConfiguration.bCleanProject && !ProjectFileGenerator.bGenerateProjectFiles) )
 			{
+				var ToolChain = UEToolChain.GetPlatformToolChain(GlobalLinkEnvironment.Config.TargetPlatform);
+
 				var UObjectModules = new List<UHTModuleInfo>();
 
 				// Figure out which modules have UObjects that we may need to generate headers for
+				foreach( var Binary in AppBinaries )
 				{
-					foreach( var Binary in AppBinaries )
+					var DependencyModules = Binary.GetAllDependencyModules(bIncludeDynamicallyLoaded: false, bForceCircular: false);
+					foreach( var DependencyModuleCPP in DependencyModules.OfType<UEBuildModuleCPP>().Where( CPPModule => !UObjectModules.Any( Module => Module.ModuleName == CPPModule.Name ) ) )
 					{
-						var DependencyModules = Binary.GetAllDependencyModules(bIncludeDynamicallyLoaded: false, bForceCircular: false);
-						foreach( var DependencyModule in DependencyModules )
-						{
-							if( !DependencyModule.bIncludedInTarget )
-							{
-								throw new BuildException( "Expecting module {0} to have bIncludeInTarget set", DependencyModule.Name );
-							}
+						if( !DependencyModuleCPP.bIncludedInTarget )
+							throw new BuildException( "Expecting module {0} to have bIncludeInTarget set", DependencyModuleCPP.Name );
 
-							UEBuildModuleCPP DependencyModuleCPP = DependencyModule as UEBuildModuleCPP;
-							if( DependencyModuleCPP != null )
-							{
-								// Do we already have this module?
-								if( !UObjectModules.Any(Module => Module.ModuleName == DependencyModuleCPP.Name ) )
-								{
-									var UHTModuleInfo = DependencyModuleCPP.GetUHTModuleInfo();
-									if( UHTModuleInfo.PublicUObjectClassesHeaders.Count > 0 || UHTModuleInfo.PrivateUObjectHeaders.Count > 0 || UHTModuleInfo.PublicUObjectHeaders.Count > 0 )
-									{
-										UObjectModules.Add( UHTModuleInfo );
-										Log.TraceVerbose( "Detected UObject module: " + DependencyModuleCPP.Name );
-									}
-								}
-								else
-								{
-									// This module doesn't define any UObjects
-								}
-							}
-							else
-							{
-								// Non-C++ modules never have UObjects
-							}
-						}
+						var ModuleCompileEnvironment = DependencyModuleCPP.CreateModuleCompileEnvironment(GlobalCompileEnvironment);
+						DependencyModuleCPP.ProcessAllCppDependencies(ModuleCompileEnvironment);
+
+						var UHTModuleInfo = DependencyModuleCPP.GetUHTModuleInfo(ModuleCompileEnvironment);
+						if( UHTModuleInfo.PublicUObjectClassesHeaders.Count == 0 && UHTModuleInfo.PrivateUObjectHeaders.Count == 0 && UHTModuleInfo.PublicUObjectHeaders.Count == 0 )
+							continue;
+
+						UHTModuleInfo.PCH                      = ToolChain.ConvertPath( DependencyModuleCPP.ProcessedDependencies.UniquePCHHeaderFile.AbsolutePath );
+						UHTModuleInfo.GeneratedCPPFilenameBase = Path.Combine( UEBuildModuleCPP.GetGeneratedCodeDirectoryForModule(this, UHTModuleInfo.ModuleDirectory, UHTModuleInfo.ModuleName), UHTModuleInfo.ModuleName ) + ".generated";
+
+						DependencyModuleCPP.AutoGenerateInlInfo = new UEBuildModuleCPP.AutoGenerateInlInfoClass( UHTModuleInfo.GeneratedCPPFilenameBase + ".cpp" );
+
+						UObjectModules.Add( UHTModuleInfo );
+						Log.TraceVerbose( "Detected UObject module: " + DependencyModuleCPP.Name );
 					}
 				}
 
@@ -1424,7 +1409,12 @@ namespace UnrealBuildTool
 						return UHTResult;
 					}
 				}
+			}
 
+			// Build the target's binaries.
+			foreach (var Binary in AppBinaries)
+			{
+				OutputItems.AddRange(Binary.Build(GlobalCompileEnvironment, GlobalLinkEnvironment));
 			}
 
 			if (BuildConfiguration.WriteTargetInfoPath != null)
