#ttp 332992 - CORE: Make sure UBT adds default xml configs to the sln
#proj UE4
#branch UE4
#summary Implemented. Refactored a bit creation of project files to allow addition of files in custom project path. Additionally fixed UBT bug that was outputting "Unhandled Exception" rather than "UnrealBuildTool Exception" on XmlConfigLoader.Init() errors.

diff --git a/Engine/Source/Programs/UnrealBuildTool/System/Project.cs b/Engine/Source/Programs/UnrealBuildTool/System/Project.cs
--- a/Engine/Source/Programs/UnrealBuildTool/System/Project.cs
+++ b/Engine/Source/Programs/UnrealBuildTool/System/Project.cs
@@ -59,6 +59,23 @@ namespace UnrealBuildTool
 		}
 	}
 
+	/// <summary>
+	/// Class that stores info about aliased file.
+	/// </summary>
+	public class AliasedFile
+	{
+		public AliasedFile(string FileSystemPath, string ProjectPath)
+		{
+			this.FileSystemPath = FileSystemPath;
+			this.ProjectPath = ProjectPath;
+		}
+
+		// File system path.
+		public string FileSystemPath { get; private set; }
+
+		// Project path.
+		public string ProjectPath { get; private set; }
+	}
 
 	public abstract class ProjectFile : IntelliSenseGatherer
 	{
@@ -165,6 +182,17 @@ namespace UnrealBuildTool
 			}
 		}
 
+		/// Aliased (i.e. files is custom filter tree) in this project
+		public readonly List<AliasedFile> AliasedFiles = new List<AliasedFile>();
+
+		/// <summary>
+		/// Adds aliased file to the project.
+		/// </summary>
+		/// <param name="File">Aliased file.</param>
+		public void AddAliasedFileToProject(AliasedFile File)
+		{
+			AliasedFiles.Add(File);
+		}
 	
 		/// <summary>
 		/// Adds a file to this project, ignoring dupes
diff --git a/Engine/Source/Programs/UnrealBuildTool/System/ProjectFileGenerator.cs b/Engine/Source/Programs/UnrealBuildTool/System/ProjectFileGenerator.cs
--- a/Engine/Source/Programs/UnrealBuildTool/System/ProjectFileGenerator.cs
+++ b/Engine/Source/Programs/UnrealBuildTool/System/ProjectFileGenerator.cs
@@ -395,6 +395,7 @@ namespace UnrealBuildTool
 						if( IncludeConfigFiles )
 						{
 							AddEngineConfigFiles( EngineProject );
+							AddUBTConfigFilesToEngineProject(EngineProject);
 						}
 
 						// Engine localization files
@@ -567,6 +568,46 @@ namespace UnrealBuildTool
 			}
 		}
 
+		/// <summary>
+		/// Adds detected UBT configuration files (BuildConfiguration.xml) to engine project.
+		/// </summary>
+		/// <param name="EngineProject">Engine project to add files to.</param>
+		private void AddUBTConfigFilesToEngineProject(ProjectFile EngineProject)
+		{
+			/*
+			 * From the UE4/Engine/Programs/UnrealBuildTool/BuildConfiguration.xml
+			 * 
+			 * There are four possible location for this file:
+		     *   a. UE4/Engine/Programs/UnrealBuildTool
+		     *   b. UE4/Engine/Programs/NoRedist/UnrealBuildTool
+		     *   c. UE4/Engine/Saved/UnrealBuildTool
+		     *   d. My Documents/Unreal Engine/UnrealBuildTool
+			 */
+
+			var AbsoluteEnginePath = new DirectoryInfo(EngineRelativePath).FullName;
+			var AbsoluteMyDocumentsPath = System.Environment.GetFolderPath(
+				Environment.SpecialFolder.MyDocuments);
+
+			var BuildConfigurationPaths = new Tuple<string, string>[] {
+				new Tuple<string, string>("Branch", Path.Combine(AbsoluteEnginePath, "Programs", "UnrealBuildTool")),
+				new Tuple<string, string>("NotForLicensees", Path.Combine(AbsoluteEnginePath, "Programs", "NoRedist", "UnrealBuildTool")),
+				new Tuple<string, string>("User", Path.Combine(AbsoluteEnginePath, "Saved", "UnrealBuildTool")),
+				new Tuple<string, string>("User_MyDocuments", Path.Combine(AbsoluteMyDocumentsPath, "Unreal Engine", "UnrealBuildTool"))
+			};
+
+			foreach(var BuildConfigurationPath in BuildConfigurationPaths)
+			{
+				var FilePath = Path.Combine(BuildConfigurationPath.Item2, "BuildConfiguration.xml");
+
+				if(!File.Exists(FilePath))
+				{
+					continue;
+				}
+
+				EngineProject.AddAliasedFileToProject(new AliasedFile(FilePath, Path.Combine("Config", "UnrealBuildTool", BuildConfigurationPath.Item1)));
+			}
+		}
+
 		/// <summary>
 		/// Clean project files
 		/// </summary>
diff --git a/Engine/Source/Programs/UnrealBuildTool/System/UnrealBuildTool.cs b/Engine/Source/Programs/UnrealBuildTool/System/UnrealBuildTool.cs
--- a/Engine/Source/Programs/UnrealBuildTool/System/UnrealBuildTool.cs
+++ b/Engine/Source/Programs/UnrealBuildTool/System/UnrealBuildTool.cs
@@ -518,7 +518,15 @@ namespace UnrealBuildTool
 		{
 			InitLogging();
 
-			XmlConfigLoader.Init();
+			try
+			{
+				XmlConfigLoader.Init();
+			}
+			catch (BuildException Exception)
+			{
+				Log.TraceError("UnrealBuildTool Exception: " + Exception);
+				return (int) ECompilationResult.OtherCompilationError;
+			}
 
 			InitialEnvironment = Environment.GetEnvironmentVariables();
 			if (InitialEnvironment.Count < 1)
diff --git a/Engine/Source/Programs/UnrealBuildTool/System/VCProject.cs b/Engine/Source/Programs/UnrealBuildTool/System/VCProject.cs
--- a/Engine/Source/Programs/UnrealBuildTool/System/VCProject.cs
+++ b/Engine/Source/Programs/UnrealBuildTool/System/VCProject.cs
@@ -513,14 +513,8 @@ namespace UnrealBuildTool
 
 			// Source folders and files
 			{
-				VCFiltersFileContent.Append(
-					"	<ItemGroup>" + ProjectFileGenerator.NewLine );
-
-				VCProjectFileContent.Append(
-					"	<ItemGroup>" + ProjectFileGenerator.NewLine );
+				var LocalAliasedFiles = new List<AliasedFile>(AliasedFiles);
 
-				// Add all file directories to the filters file as solution filters
-				var FilterDirectories = new HashSet<string>();
 				foreach( var CurFile in SourceFiles )
 				{
 					// We want all source file and directory paths in the project files to be relative to the project file's
@@ -538,102 +532,36 @@ namespace UnrealBuildTool
 						FilterRelativeSourceDirectory = Path.GetDirectoryName( Utils.MakePathRelativeTo( CurFile.FilePath, CurFile.RelativeBaseFolder ) );
 					}
 
-					// No need to add the root directory relative to the project (it would just be an empty string!)
-					if( !String.IsNullOrWhiteSpace( FilterRelativeSourceDirectory ) )
-					{
-						// We only want each directory to appear once in the filters file
-						var PathRemaining = Utils.CleanDirectorySeparators( FilterRelativeSourceDirectory );
-						if( !FilterDirectories.Contains( PathRemaining ) )
-						{
-							// Make sure all subdirectories leading up to this directory each have their own filter, too!
-							var AllDirectoriesInPath = new List<string>();
-							var PathSoFar = "";
-							for( ; ; )
-							{
-								if( PathRemaining.Length > 0 )
-								{
-									var SlashIndex = PathRemaining.IndexOf( Path.DirectorySeparatorChar );
-									string SplitDirectory;
-									if( SlashIndex != -1 )
-									{
-										SplitDirectory = PathRemaining.Substring( 0, SlashIndex );
-										PathRemaining = PathRemaining.Substring( SplitDirectory.Length + 1 );
-									}
-									else
-									{
-										SplitDirectory = PathRemaining;
-										PathRemaining = "";
-									}
-									if( !String.IsNullOrEmpty( PathSoFar ) )
-									{
-										PathSoFar += Path.DirectorySeparatorChar;
-									}
-									PathSoFar += SplitDirectory;
-
-									AllDirectoriesInPath.Add( PathSoFar );
-								}
-								else
-								{
-									break;
-								}
-							}
-
-							foreach( var LeadingDirectory in AllDirectoriesInPath )
-							{
-								if( !FilterDirectories.Contains( LeadingDirectory ) )
-								{
-									FilterDirectories.Add( LeadingDirectory );
-
-									// Generate a unique GUID for this folder
-									// NOTE: When saving generated project files, we ignore differences in GUIDs if every other part of the file
-									//       matches identically with the pre-existing file
-									var FilterGUID = Guid.NewGuid().ToString( "B" ).ToUpperInvariant();
-
-									VCFiltersFileContent.Append(
-										"		<Filter Include=\"" + LeadingDirectory + "\">" + ProjectFileGenerator.NewLine +
-										"			<UniqueIdentifier>" + FilterGUID + "</UniqueIdentifier>" + ProjectFileGenerator.NewLine +
-										"		</Filter>" + ProjectFileGenerator.NewLine);
+					LocalAliasedFiles.Add(new AliasedFile(ProjectRelativeSourceFile, FilterRelativeSourceDirectory));
+				}
 
-									FiltersFileIsNeeded = true;
-								}
-							}
-						}
-					}
+				VCFiltersFileContent.Append(
+					"	<ItemGroup>" + ProjectFileGenerator.NewLine);
 
+				VCProjectFileContent.Append(
+					"	<ItemGroup>" + ProjectFileGenerator.NewLine);
 
-					// What type of file is this?
-					string VCFileType;
-					if( CurFile.FilePath.EndsWith( ".h", StringComparison.InvariantCultureIgnoreCase ) ||
-						CurFile.FilePath.EndsWith( ".inl", StringComparison.InvariantCultureIgnoreCase ) )
-					{
-						VCFileType = "ClInclude";
-					}
-					else if( CurFile.FilePath.EndsWith( ".cpp", StringComparison.InvariantCultureIgnoreCase ) )
-					{
-						VCFileType = "ClCompile";
-					}
-					else if( CurFile.FilePath.EndsWith( ".rc", StringComparison.InvariantCultureIgnoreCase ) )
-					{
-						VCFileType = "ResourceCompile";
-					}
-					else if( CurFile.FilePath.EndsWith( ".manifest", StringComparison.InvariantCultureIgnoreCase ) )
-					{
-						VCFileType = "Manifest";
-					}
-					else
+				// Add all file directories to the filters file as solution filters
+				var FilterDirectories = new HashSet<string>();
+				foreach (var AliasedFile in LocalAliasedFiles)
+				{
+					// No need to add the root directory relative to the project (it would just be an empty string!)
+					if (!String.IsNullOrWhiteSpace(AliasedFile.ProjectPath))
 					{
-						VCFileType = "None";
+						FiltersFileIsNeeded = EnsureFilterPathExists(AliasedFile.ProjectPath, VCFiltersFileContent, FilterDirectories);
 					}
 
+					var VCFileType = GetVCFileType(AliasedFile.FileSystemPath);
+
 					VCProjectFileContent.Append(
-						"		<" + VCFileType + " Include=\"" + ProjectRelativeSourceFile + "\" />" + ProjectFileGenerator.NewLine);
+						"		<" + VCFileType + " Include=\"" + AliasedFile.FileSystemPath + "\" />" + ProjectFileGenerator.NewLine);
 
-					if( !String.IsNullOrWhiteSpace( FilterRelativeSourceDirectory ) )
+					if (!String.IsNullOrWhiteSpace(AliasedFile.ProjectPath))
 					{
 						VCFiltersFileContent.Append(
-							"		<" + VCFileType + " Include=\"" + ProjectRelativeSourceFile + "\">" + ProjectFileGenerator.NewLine +
-							"			<Filter>" + FilterRelativeSourceDirectory + "</Filter>" + ProjectFileGenerator.NewLine +
-							"		</" + VCFileType + " >" + ProjectFileGenerator.NewLine );
+							"		<" + VCFileType + " Include=\"" + AliasedFile.FileSystemPath + "\">" + ProjectFileGenerator.NewLine +
+							"			<Filter>" + Utils.CleanDirectorySeparators(AliasedFile.ProjectPath) + "</Filter>" + ProjectFileGenerator.NewLine +
+							"		</" + VCFileType + " >" + ProjectFileGenerator.NewLine);
 
 						FiltersFileIsNeeded = true;
 					}
@@ -641,7 +569,7 @@ namespace UnrealBuildTool
 					{
 						// No need to specify the root directory relative to the project (it would just be an empty string!)
 						VCFiltersFileContent.Append(
-							"		<" + VCFileType + " Include=\"" + ProjectRelativeSourceFile + "\" />" + ProjectFileGenerator.NewLine);
+							"		<" + VCFileType + " Include=\"" + AliasedFile.FileSystemPath + "\" />" + ProjectFileGenerator.NewLine);
 					}
 				}
 
@@ -862,6 +790,100 @@ namespace UnrealBuildTool
 			return bSuccess;
 		}
 
+		private static bool EnsureFilterPathExists(string FilterRelativeSourceDirectory, StringBuilder VCFiltersFileContent, HashSet<string> FilterDirectories)
+		{
+			// We only want each directory to appear once in the filters file
+			var PathRemaining = Utils.CleanDirectorySeparators( FilterRelativeSourceDirectory );
+			var FiltersFileIsNeeded = false;
+			if( !FilterDirectories.Contains( PathRemaining ) )
+			{
+				// Make sure all subdirectories leading up to this directory each have their own filter, too!
+				var AllDirectoriesInPath = new List<string>();
+				var PathSoFar = "";
+				for( ; ; )
+				{
+					if( PathRemaining.Length > 0 )
+					{
+						var SlashIndex = PathRemaining.IndexOf( Path.DirectorySeparatorChar );
+						string SplitDirectory;
+						if( SlashIndex != -1 )
+						{
+							SplitDirectory = PathRemaining.Substring( 0, SlashIndex );
+							PathRemaining = PathRemaining.Substring( SplitDirectory.Length + 1 );
+						}
+						else
+						{
+							SplitDirectory = PathRemaining;
+							PathRemaining = "";
+						}
+						if( !String.IsNullOrEmpty( PathSoFar ) )
+						{
+							PathSoFar += Path.DirectorySeparatorChar;
+						}
+						PathSoFar += SplitDirectory;
+
+						AllDirectoriesInPath.Add( PathSoFar );
+					}
+					else
+					{
+						break;
+					}
+				}
+
+				foreach( var LeadingDirectory in AllDirectoriesInPath )
+				{
+					if( !FilterDirectories.Contains( LeadingDirectory ) )
+					{
+						FilterDirectories.Add( LeadingDirectory );
+
+						// Generate a unique GUID for this folder
+						// NOTE: When saving generated project files, we ignore differences in GUIDs if every other part of the file
+						//       matches identically with the pre-existing file
+						var FilterGUID = Guid.NewGuid().ToString( "B" ).ToUpperInvariant();
+
+						VCFiltersFileContent.Append(
+							"		<Filter Include=\"" + LeadingDirectory + "\">" + ProjectFileGenerator.NewLine +
+							"			<UniqueIdentifier>" + FilterGUID + "</UniqueIdentifier>" + ProjectFileGenerator.NewLine +
+							"		</Filter>" + ProjectFileGenerator.NewLine);
+
+						FiltersFileIsNeeded = true;
+					}
+				}
+			}
+
+			return FiltersFileIsNeeded;
+		}
+
+		/// <summary>
+		/// Returns the VCFileType element name based on the file path.
+		/// </summary>
+		/// <param name="Path">The path of the file to return type for.</param>
+		/// <returns>Name of the element in MSBuild project file for this file.</returns>
+		private string GetVCFileType(string Path)
+		{
+			// What type of file is this?
+			if (Path.EndsWith(".h", StringComparison.InvariantCultureIgnoreCase) ||
+				Path.EndsWith(".inl", StringComparison.InvariantCultureIgnoreCase))
+			{
+				return "ClInclude";
+			}
+			else if (Path.EndsWith(".cpp", StringComparison.InvariantCultureIgnoreCase))
+			{
+				return "ClCompile";
+			}
+			else if (Path.EndsWith(".rc", StringComparison.InvariantCultureIgnoreCase))
+			{
+				return "ResourceCompile";
+			}
+			else if (Path.EndsWith(".manifest", StringComparison.InvariantCultureIgnoreCase))
+			{
+				return "Manifest";
+			}
+			else
+			{
+				return "None";
+			}
+		}
 
 		// Anonymous function that writes project configuration data
 		private void WriteConfiguration(string ProjectName, UnrealTargetConfiguration Configuration, UnrealTargetPlatform Platform, string TargetFilePath, TargetRules TargetRulesObject, StringBuilder VCProjectFileContent, StringBuilder VCUserFileContent)
