Changed UE4/Config/UnrealBuildTool/Branch project folder name to Default to avoid impression that the file should be changed.

diff --git a/Engine/Source/Programs/UnrealBuildTool/System/ProjectFileGenerator.cs b/Engine/Source/Programs/UnrealBuildTool/System/ProjectFileGenerator.cs
--- a/Engine/Source/Programs/UnrealBuildTool/System/ProjectFileGenerator.cs
+++ b/Engine/Source/Programs/UnrealBuildTool/System/ProjectFileGenerator.cs
@@ -589,7 +589,7 @@ namespace UnrealBuildTool
 				Environment.SpecialFolder.MyDocuments);
 
 			var BuildConfigurationPaths = new Tuple<string, string>[] {
-				new Tuple<string, string>("Branch", Path.Combine(AbsoluteEnginePath, "Programs", "UnrealBuildTool")),
+				new Tuple<string, string>("Default", Path.Combine(AbsoluteEnginePath, "Programs", "UnrealBuildTool")),
 				new Tuple<string, string>("NotForLicensees", Path.Combine(AbsoluteEnginePath, "Programs", "NoRedist", "UnrealBuildTool")),
 				new Tuple<string, string>("User", Path.Combine(AbsoluteEnginePath, "Saved", "UnrealBuildTool")),
 				new Tuple<string, string>("User_MyDocuments", Path.Combine(AbsoluteMyDocumentsPath, "Unreal Engine", "UnrealBuildTool"))
