Moved NoRedist/UnrealBuildTool/BuildConfiguration.xml to NotForLicensees/UnrealBuildTool/BuildConfiguration.xml and refactored UBT BuildConfiguration.xml loader and ProjectFileGenerator, so they use same possible location source.

diff --git a/Engine/Programs/UnrealBuildTool/BuildConfiguration.xml b/Engine/Programs/UnrealBuildTool/BuildConfiguration.xml
--- a/Engine/Programs/UnrealBuildTool/BuildConfiguration.xml
+++ b/Engine/Programs/UnrealBuildTool/BuildConfiguration.xml
@@ -35,7 +35,7 @@
 	
 	There are four possible location for this file:
 		a. UE4/Engine/Programs/UnrealBuildTool
-		b. UE4/Engine/Programs/NoRedist/UnrealBuildTool
+		b. UE4/Engine/Programs/NotForLicensees/UnrealBuildTool
 		c. UE4/Engine/Saved/UnrealBuildTool
 		d. My Documents/Unreal Engine/UnrealBuildTool
 	
diff --git a/Engine/Source/Programs/UnrealBuildTool/System/ProjectFileGenerator.cs b/Engine/Source/Programs/UnrealBuildTool/System/ProjectFileGenerator.cs
--- a/Engine/Source/Programs/UnrealBuildTool/System/ProjectFileGenerator.cs
+++ b/Engine/Source/Programs/UnrealBuildTool/System/ProjectFileGenerator.cs
@@ -574,37 +574,16 @@ namespace UnrealBuildTool
 		/// <param name="EngineProject">Engine project to add files to.</param>
 		private void AddUBTConfigFilesToEngineProject(ProjectFile EngineProject)
 		{
-			/*
-			 * From the UE4/Engine/Programs/UnrealBuildTool/BuildConfiguration.xml
-			 * 
-			 * There are four possible location for this file:
-		     *   a. UE4/Engine/Programs/UnrealBuildTool
-		     *   b. UE4/Engine/Programs/NoRedist/UnrealBuildTool
-		     *   c. UE4/Engine/Saved/UnrealBuildTool
-		     *   d. My Documents/Unreal Engine/UnrealBuildTool
-			 */
-
-			var AbsoluteEnginePath = new DirectoryInfo(EngineRelativePath).FullName;
-			var AbsoluteMyDocumentsPath = System.Environment.GetFolderPath(
-				Environment.SpecialFolder.MyDocuments);
-
-			var BuildConfigurationPaths = new Tuple<string, string>[] {
-				new Tuple<string, string>("Default", Path.Combine(AbsoluteEnginePath, "Programs", "UnrealBuildTool")),
-				new Tuple<string, string>("NotForLicensees", Path.Combine(AbsoluteEnginePath, "Programs", "NoRedist", "UnrealBuildTool")),
-				new Tuple<string, string>("User", Path.Combine(AbsoluteEnginePath, "Saved", "UnrealBuildTool")),
-				new Tuple<string, string>("User_MyDocuments", Path.Combine(AbsoluteMyDocumentsPath, "Unreal Engine", "UnrealBuildTool"))
-			};
-
-			foreach(var BuildConfigurationPath in BuildConfigurationPaths)
-			{
-				var FilePath = Path.Combine(BuildConfigurationPath.Item2, "BuildConfiguration.xml");
+			foreach(var BuildConfigurationPath in XmlConfigLoader.ConfigLocationHierarchy)
+			{
+				var FilePath = Path.Combine(BuildConfigurationPath.FSLocation, "BuildConfiguration.xml");
 
 				if(!File.Exists(FilePath))
 				{
 					continue;
 				}
 
-				EngineProject.AddAliasedFileToProject(new AliasedFile(FilePath, Path.Combine("Config", "UnrealBuildTool", BuildConfigurationPath.Item1)));
+				EngineProject.AddAliasedFileToProject(new AliasedFile(FilePath, Path.Combine("Config", "UnrealBuildTool", BuildConfigurationPath.IDEFolderName)));
 			}
 		}
 
diff --git a/Engine/Source/Programs/UnrealBuildTool/Utilities/XmlConfigLoader.cs b/Engine/Source/Programs/UnrealBuildTool/Utilities/XmlConfigLoader.cs
--- a/Engine/Source/Programs/UnrealBuildTool/Utilities/XmlConfigLoader.cs
+++ b/Engine/Source/Programs/UnrealBuildTool/Utilities/XmlConfigLoader.cs
@@ -110,17 +110,32 @@ namespace UnrealBuildTool
 		}
 
 		/// <summary>
-		/// Loads BuildConfiguration from XML into memory.
+		/// Class that stores information about possible BuildConfiguration.xml
+		/// location and its name that will be displayed in IDE.
 		/// </summary>
-		private static void LoadData()
+		public class XmlConfigLocation
 		{
-			var ConfigXmlFileName = "BuildConfiguration.xml";
-			var UE4EnginePath = new FileInfo(Path.Combine(Utils.GetExecutingAssemblyDirectory(), "..", "..")).FullName;
+			// Possible location of the file system.
+			public string FSLocation { get; private set; }
+
+			// IDE folder name that will contain this location if file will be found.
+			public string IDEFolderName { get; private set; }
 
+			public XmlConfigLocation(string FSLocation, string IDEFolderName)
+			{
+				this.FSLocation = FSLocation;
+				this.IDEFolderName = IDEFolderName;
+			}
+		}
+
+		public static readonly XmlConfigLocation[] ConfigLocationHierarchy;
+
+		static XmlConfigLoader()
+		{
 			/*
 			 *	There are four possible location for this file:
 			 *		a. UE4/Engine/Programs/UnrealBuildTool
-			 *		b. UE4/Engine/Programs/NoRedist/UnrealBuildTool
+			 *		b. UE4/Engine/Programs/NotForLicensees/UnrealBuildTool
 			 *		c. UE4/Engine/Saved/UnrealBuildTool
 			 *		d. My Documnets/Unreal Engine/UnrealBuildTool
 			 *
@@ -129,17 +144,27 @@ namespace UnrealBuildTool
 			 *	priority. Not defined classes and fields are left alone.
 			 */
 
-			var ConfigLocationHierarchy = new string[]
+			var UE4EnginePath = new FileInfo(Path.Combine(Utils.GetExecutingAssemblyDirectory(), "..", "..")).FullName;
+
+			ConfigLocationHierarchy = new XmlConfigLocation[]
 			{
-				Path.Combine(UE4EnginePath, "Programs", "UnrealBuildTool"),
-				Path.Combine(UE4EnginePath, "Programs", "NoRedist", "UnrealBuildTool"),
-				Path.Combine(UE4EnginePath, "Saved", "UnrealBuildTool"),
-				Path.Combine(Environment.GetFolderPath(Environment.SpecialFolder.Personal), "Unreal Engine", "UnrealBuildTool")
+				new XmlConfigLocation(Path.Combine(UE4EnginePath, "Programs", "UnrealBuildTool"), "Default"),
+				new XmlConfigLocation(Path.Combine(UE4EnginePath, "Programs", "NotForLicensees", "UnrealBuildTool"), "NotForLicensees"),
+				new XmlConfigLocation(Path.Combine(UE4EnginePath, "Saved", "UnrealBuildTool"), "User"),
+				new XmlConfigLocation(Path.Combine(Environment.GetFolderPath(Environment.SpecialFolder.Personal), "Unreal Engine", "UnrealBuildTool"), "MyDocuments")
 			};
+		}
+
+		/// <summary>
+		/// Loads BuildConfiguration from XML into memory.
+		/// </summary>
+		private static void LoadData()
+		{
+			var ConfigXmlFileName = "BuildConfiguration.xml";
 
 			foreach (var PossibleConfigLocation in ConfigLocationHierarchy)
 			{
-				var FilePath = Path.Combine(PossibleConfigLocation, ConfigXmlFileName);
+				var FilePath = Path.Combine(PossibleConfigLocation.FSLocation, ConfigXmlFileName);
 
 				if (File.Exists(FilePath))
 				{
