Fixed a bug in Xcode project generator that was causing Mac source not getting added to UE4XcodeHelper target.

diff --git a/Engine/Source/Programs/UnrealBuildTool/System/XcodeProject.cs b/Engine/Source/Programs/UnrealBuildTool/System/XcodeProject.cs
--- a/Engine/Source/Programs/UnrealBuildTool/System/XcodeProject.cs
+++ b/Engine/Source/Programs/UnrealBuildTool/System/XcodeProject.cs
@@ -181,7 +181,7 @@ namespace UnrealBuildTool
 				foreach (string PlatformName in Enum.GetNames(typeof(UnrealTargetPlatform)))
 				{
 					string AltName = PlatformName == "Win32" || PlatformName == "Win64" ? "windows" : PlatformName.ToLower();
-					if ((SourceFile.FilePath.ToLower().Contains("/" + PlatformName.ToLower() + "/")) || (SourceFile.FilePath.ToLower().Contains("/" + AltName + "/")) && PlatformName != "Mac")
+					if ((SourceFile.FilePath.ToLower().Contains("/" + PlatformName.ToLower() + "/") || SourceFile.FilePath.ToLower().Contains("/" + AltName + "/")) && PlatformName != "Mac")
 					{
 						// UE4XcodeHelper is Mac only target, so skip other platforms files
 						return false;
