Build time fixes.

* Fixed a bug where the shared PCH file was being nulled out if it wasn't from a global PCH environment.
* The PCH environment for .generated.cpp files is now taken taken from associated module's environment.
* Split .generated.cpp files now include the module PCH at the top of every file.
* PCH inclusion is now always forced.

diff --git a/Engine/Source/Programs/UnrealBuildTool/Configuration/UEBuildModule.cs b/Engine/Source/Programs/UnrealBuildTool/Configuration/UEBuildModule.cs
--- a/Engine/Source/Programs/UnrealBuildTool/Configuration/UEBuildModule.cs
+++ b/Engine/Source/Programs/UnrealBuildTool/Configuration/UEBuildModule.cs
@@ -983,6 +983,9 @@ namespace UnrealBuildTool
 				bModuleUsesUnityBuild = false;
 			}
 
+			// The environment with which to compile the CPP files
+			var CPPCompileEnvironment = ModuleCompileEnvironment;
+
 			// Precompiled header support.
 			bool bWasModuleCodeCompiled = false;
 			if (BuildPlatform.ShouldUsePCHFiles(CompileEnvironment.Config.TargetPlatform, CompileEnvironment.Config.TargetConfiguration))
@@ -1137,9 +1140,11 @@ namespace UnrealBuildTool
 				// Only use precompiled headers for projects with enough files to make the PCH creation worthwhile.
 				if( SharedPCHHeaderFile != null || SourceFiles.CPPFiles.Count >= MinFilesUsingPrecompiledHeader )
 				{
+					FileItem PCHToUse;
+
 					if( SharedPCHHeaderFile != null )
 					{
-						ModulePCHEnvironment = ApplySharedPCH(GlobalCompileEnvironment, CompileEnvironment, ModuleCompileEnvironment, SourceFiles.CPPFiles, SharedPCHHeaderFile);
+						ModulePCHEnvironment = ApplySharedPCH(GlobalCompileEnvironment, CompileEnvironment, ModuleCompileEnvironment, SourceFiles.CPPFiles, ref SharedPCHHeaderFile);
 						if (ModulePCHEnvironment != null)
 						{
 							// @todo SharedPCH: Ideally we would exhaustively check for a compatible compile environment (definitions, imports/exports, etc)
@@ -1148,9 +1153,21 @@ namespace UnrealBuildTool
 							//    errors unexpectedly due to compile environment differences.
 							Log.TraceVerbose("Module " + Name + " uses existing SharedPCH '" + ModulePCHEnvironment.PrecompiledHeaderIncludeFilename + "' (from module " + ModulePCHEnvironment.ModuleName + ")");
 						}
-						else
+
+						PCHToUse = SharedPCHHeaderFile;
+					}
+					else
+					{
+						PCHToUse = ProcessedDependencies.UniquePCHHeaderFile;
+					}
+
+					if (PCHToUse != null)
+					{
+						// Update all CPPFiles to point to the PCH
+						foreach (var CPPFile in SourceFiles.CPPFiles)
 						{
-							SharedPCHHeaderFile = null;
+							CPPFile.PCHHeaderNameInCode              = PCHToUse.AbsolutePath;
+							CPPFile.PrecompiledHeaderIncludeFilename = PCHToUse.AbsolutePath;
 						}
 					}
 
@@ -1246,14 +1263,12 @@ namespace UnrealBuildTool
 
 						// if pch action was generated for the environment then use pch
 						ModulePCHCompileEnvironment.PrecompiledHeaderFile = ModulePCHEnvironment.PrecompiledHeaderFile;
-						LinkInputFiles.AddRange( ModulePCHCompileEnvironment.CompileFiles( CPPFilesToBuild, Name ).ObjectFiles );
-					}
-					else
-					{
-						// otherwise, compile non-pch
-						LinkInputFiles.AddRange( ModuleCompileEnvironment.CompileFiles( CPPFilesToBuild, Name ).ObjectFiles );
+
+						// Use this compile environment from now on
+						CPPCompileEnvironment = ModulePCHCompileEnvironment;
 					}
-					
+
+					LinkInputFiles.AddRange( CPPCompileEnvironment.CompileFiles( CPPFilesToBuild, Name ).ObjectFiles );
 					bWasModuleCodeCompiled = true;
 				}
 
@@ -1271,7 +1286,7 @@ namespace UnrealBuildTool
 				{
 					CPPFilesToCompile = Unity.GenerateUnityCPPs( CPPFilesToCompile, ModuleCompileEnvironment, Name );
 				}
-				LinkInputFiles.AddRange( ModuleCompileEnvironment.CompileFiles( CPPFilesToCompile, Name ).ObjectFiles );
+				LinkInputFiles.AddRange( CPPCompileEnvironment.CompileFiles( CPPFilesToCompile, Name ).ObjectFiles );
 			}
 
 			string PCHPath = null;
@@ -1284,7 +1299,7 @@ namespace UnrealBuildTool
 			{
 				var GeneratedCppFileItem = FileItem.GetItemByPath( AutoGenerateInlInfo.Filename );
 
-				LinkInputFiles.AddRange( ModuleCompileEnvironment.CompileFiles( new List<FileItem>{ GeneratedCppFileItem }, Name ).ObjectFiles );
+				LinkInputFiles.AddRange( CPPCompileEnvironment.CompileFiles( new List<FileItem>{ GeneratedCppFileItem }, Name ).ObjectFiles );
 			}
 
 			// Compile C files directly.
@@ -1321,10 +1336,11 @@ namespace UnrealBuildTool
 			return LinkInputFiles;
 		}
 
-		private PrecompileHeaderEnvironment ApplySharedPCH(CPPEnvironment GlobalCompileEnvironment, CPPEnvironment CompileEnvironment, CPPEnvironment ModuleCompileEnvironment, List<FileItem> CPPFiles, FileItem SharedPCHHeaderFile)
+		private PrecompileHeaderEnvironment ApplySharedPCH(CPPEnvironment GlobalCompileEnvironment, CPPEnvironment CompileEnvironment, CPPEnvironment ModuleCompileEnvironment, List<FileItem> CPPFiles, ref FileItem SharedPCHHeaderFile)
 		{
 			// Check to see if we have a PCH header already setup that we can use
-			var SharedPCHEnvironment = GlobalCompileEnvironment.SharedPCHEnvironments.Find(Env => Env.PrecompiledHeaderIncludeFilename == SharedPCHHeaderFile);
+			var SharedPCHHeaderFileCopy = SharedPCHHeaderFile;
+			var SharedPCHEnvironment = GlobalCompileEnvironment.SharedPCHEnvironments.Find(Env => Env.PrecompiledHeaderIncludeFilename == SharedPCHHeaderFileCopy);
 			if (SharedPCHEnvironment == null)
 				return null;
 
@@ -1332,6 +1348,7 @@ namespace UnrealBuildTool
 			if (SharedPCHEnvironment.CLRMode != ModuleCompileEnvironment.Config.CLRMode)
 			{
 				Log.TraceVerbose("Module {0} cannot use existing SharedPCH '{1}' (from module '{2}') because CLR modes don't match", Name, SharedPCHEnvironment.PrecompiledHeaderIncludeFilename.AbsolutePath, SharedPCHEnvironment.ModuleName);
+				SharedPCHHeaderFile = null;
 				return null;
 			}
 
@@ -1355,16 +1372,10 @@ namespace UnrealBuildTool
 			if (SharedPCHCodeOptimization != ModuleCodeOptimization)
 			{
 				Log.TraceVerbose("Module {0} cannot use existing SharedPCH '{1}' (from module '{2}') because optimization levels don't match", Name, SharedPCHEnvironment.PrecompiledHeaderIncludeFilename.AbsolutePath, SharedPCHEnvironment.ModuleName);
+				SharedPCHHeaderFile = null;
 				return null;
 			}
 
-			// Update all CPPFiles to point to the shared PCH
-			foreach (var CPPFile in CPPFiles)
-			{
-				CPPFile.PCHHeaderNameInCode              = SharedPCHHeaderFile.AbsolutePath;
-				CPPFile.PrecompiledHeaderIncludeFilename = SharedPCHHeaderFile.AbsolutePath;
-			}
-
 			return SharedPCHEnvironment;
 		}
 
diff --git a/Engine/Source/Programs/UnrealBuildTool/Windows/VCToolChain.cs b/Engine/Source/Programs/UnrealBuildTool/Windows/VCToolChain.cs
--- a/Engine/Source/Programs/UnrealBuildTool/Windows/VCToolChain.cs
+++ b/Engine/Source/Programs/UnrealBuildTool/Windows/VCToolChain.cs
@@ -797,7 +797,10 @@ namespace UnrealBuildTool
 							FileArguments += string.Format(" /Yu\"{0}\"", CompileEnvironment.Config.PCHHeaderNameInCode);
 							FileArguments += string.Format(" /Fp\"{0}\"", CompileEnvironment.PrecompiledHeaderFile.AbsolutePath);
 
-							if( CompileEnvironment.Config.bForceIncludePrecompiledHeader )
+							// Is it unsafe to always force inclusion?  Clang is doing it, and .generated.cpp files
+							// won't work otherwise, because they're not located in the context of the module,
+							// so they can't access the module's PCH without an absolute path.
+							//if( CompileEnvironment.Config.bForceIncludePrecompiledHeader )
 							{
 								// Force include the precompiled header file.  This is needed because we may have selected a
 								// precompiled header that is different than the first direct include in the C++ source file, but
diff --git a/Engine/Source/Programs/UnrealHeaderTool/Private/CodeGenerator.cpp b/Engine/Source/Programs/UnrealHeaderTool/Private/CodeGenerator.cpp
--- a/Engine/Source/Programs/UnrealHeaderTool/Private/CodeGenerator.cpp
+++ b/Engine/Source/Programs/UnrealHeaderTool/Private/CodeGenerator.cpp
@@ -4495,10 +4495,12 @@ void FNativeClassHeaderGenerator::ExportGeneratedCPP()
 																								LINE_TERMINATOR
 		);
 
+	FString ModulePCHInclude;
+
 	const auto* ModuleInfo = GPackageToManifestModuleMap.FindChecked(Package);
 	if (ModuleInfo->PCH.Len())
 	{
-		GeneratedCPPClassesIncludes.Logf(TEXT("#include \"%s\"") LINE_TERMINATOR, *ModuleInfo->PCH);
+		ModulePCHInclude = FString::Printf(TEXT("#include \"%s\"") LINE_TERMINATOR, *ModuleInfo->PCH);
 	}
 
 	// Write out the ordered class dependencies into a single header that we can easily include
@@ -4562,7 +4564,7 @@ void FNativeClassHeaderGenerator::ExportGeneratedCPP()
 
 		FString HeaderPath = ModuleInfo->GeneratedCPPFilenameBase + (GeneratedFunctionBodyTextSplit.Num() > 1 ? *FString::Printf(TEXT(".%d.cpp"), FileIdx + 1) : TEXT(".cpp"));
 		
-		SaveHeaderIfChanged(*HeaderPath, *(GeneratedCPPPreamble + GeneratedCPPClassesIncludes + FileText + GeneratedCPPEpilogue));
+		SaveHeaderIfChanged(*HeaderPath, *(GeneratedCPPPreamble + ModulePCHInclude + GeneratedCPPClassesIncludes + FileText + GeneratedCPPEpilogue));
 
 		if (GeneratedFunctionBodyTextSplit.Num() > 1)
 		{
@@ -4580,7 +4582,7 @@ void FNativeClassHeaderGenerator::ExportGeneratedCPP()
 		}
 
 		FString HeaderPath = ModuleInfo->GeneratedCPPFilenameBase + TEXT(".cpp");
-		SaveHeaderIfChanged(*HeaderPath, *(GeneratedCPPPreamble + FileText + GeneratedCPPEpilogue));
+		SaveHeaderIfChanged(*HeaderPath, *(GeneratedCPPPreamble + ModulePCHInclude + FileText + GeneratedCPPEpilogue));
 	}
 
 	if( FParse::Param( FCommandLine::Get(), TEXT("GenerateConstructors") ) && AllConstructors.Len())
