UBT now tries to create an empty template BuildConfiguration.xml file if it doesn't find it in user locations.

diff --git a/Engine/Source/Programs/UnrealBuildTool/Utilities/XmlConfigLoader.cs b/Engine/Source/Programs/UnrealBuildTool/Utilities/XmlConfigLoader.cs
--- a/Engine/Source/Programs/UnrealBuildTool/Utilities/XmlConfigLoader.cs
+++ b/Engine/Source/Programs/UnrealBuildTool/Utilities/XmlConfigLoader.cs
@@ -121,10 +121,14 @@ namespace UnrealBuildTool
 			// IDE folder name that will contain this location if file will be found.
 			public string IDEFolderName { get; private set; }
 
-			public XmlConfigLocation(string FSLocation, string IDEFolderName)
+			// Tells if this is a user location.
+			public bool bIsUserLocation { get; private set; }
+
+			public XmlConfigLocation(string FSLocation, string IDEFolderName, bool bIsUserLocation = false)
 			{
 				this.FSLocation = FSLocation;
 				this.IDEFolderName = IDEFolderName;
+				this.bIsUserLocation = bIsUserLocation;
 			}
 		}
 
@@ -150,8 +154,8 @@ namespace UnrealBuildTool
 			{
 				new XmlConfigLocation(Path.Combine(UE4EnginePath, "Programs", "UnrealBuildTool"), "Default"),
 				new XmlConfigLocation(Path.Combine(UE4EnginePath, "Programs", "NotForLicensees", "UnrealBuildTool"), "NotForLicensees"),
-				new XmlConfigLocation(Path.Combine(UE4EnginePath, "Saved", "UnrealBuildTool"), "User"),
-				new XmlConfigLocation(Path.Combine(Environment.GetFolderPath(Environment.SpecialFolder.Personal), "Unreal Engine", "UnrealBuildTool"), "MyDocuments")
+				new XmlConfigLocation(Path.Combine(UE4EnginePath, "Saved", "UnrealBuildTool"), "User", true),
+				new XmlConfigLocation(Path.Combine(Environment.GetFolderPath(Environment.SpecialFolder.Personal), "Unreal Engine", "UnrealBuildTool"), "MyDocuments", true)
 			};
 		}
 
@@ -160,7 +164,9 @@ namespace UnrealBuildTool
 		/// </summary>
 		private static void LoadData()
 		{
-			var ConfigXmlFileName = "BuildConfiguration.xml";
+			const string ConfigXmlFileName = "BuildConfiguration.xml";
+
+			bool bFoundUserLocation = false;
 
 			foreach (var PossibleConfigLocation in ConfigLocationHierarchy)
 			{
@@ -168,11 +174,57 @@ namespace UnrealBuildTool
 
 				if (File.Exists(FilePath))
 				{
+					bFoundUserLocation |= PossibleConfigLocation.bIsUserLocation;
+
 					Load(FilePath);
 				}
 			}
+
+			if(!bFoundUserLocation)
+			{
+				CreateUserXmlConfigTemplate();
+			}
 		}
 
+		/// <summary>
+		/// Creates template file in one of the user locations.
+		/// </summary>
+		private static void CreateUserXmlConfigTemplate()
+		{
+			const string ConfigXmlFileName = "BuildConfiguration.xml";
+
+			// If not found in any of user locations, try to create in one.
+			foreach (var PossibleConfigLocation in ConfigLocationHierarchy)
+			{
+				if (!PossibleConfigLocation.bIsUserLocation)
+				{
+					// Skip non-user locations.
+					continue;
+				}
+
+				try
+				{
+					Directory.CreateDirectory(PossibleConfigLocation.FSLocation);
+
+					var FilePath = Path.Combine(PossibleConfigLocation.FSLocation, ConfigXmlFileName);
+
+					const string TemplateContent =
+						"<Configuration>\n" +
+						"	<BuildConfiguration>\n" +
+						"	</BuildConfiguration>\n" +
+						"</Configuration>\n";
+
+					File.WriteAllText(FilePath, TemplateContent);
+
+					// If everything went fine then break, cause one user file was already created.
+					break;
+				}
+				catch (Exception)
+				{
+					// Ignore quietly.
+				}
+			}
+		}
 		/// <summary>
 		/// Sets values of this class with values from given XML file.
 		/// </summary>
