Quick fix to get IOS compiling with RPC Utility

diff --git a/Engine/Source/Programs/UnrealBuildTool/Configuration/UEBuildTarget.cs b/Engine/Source/Programs/UnrealBuildTool/Configuration/UEBuildTarget.cs
--- a/Engine/Source/Programs/UnrealBuildTool/Configuration/UEBuildTarget.cs
+++ b/Engine/Source/Programs/UnrealBuildTool/Configuration/UEBuildTarget.cs
@@ -1486,7 +1486,19 @@ namespace UnrealBuildTool
 				// Include an empty header so UEBuildModule.Compile does not complain about a lack of PCH
 				string HeaderFilename                    = LinkerFixupsName + "Name.h";
 				string LinkerFixupHeaderFilenameWithPath = Path.Combine(GlobalCompileEnvironment.Config.OutputDirectory, HeaderFilename);
-				LinkerFixupsFileContents.Add("#include \"" + LinkerFixupHeaderFilenameWithPath + "\"");
+
+				UnrealTargetPlatform TargetPlatform = UEBuildTarget.CPPTargetPlatformToUnrealTargetPlatform( GlobalCompileEnvironment.Config.TargetPlatform );
+
+				// Quick fix. 
+				// For some reason, when compiled over RPCUtility, the path can't be found, and we don't need the long path on Mac anyhow
+				if ( ExternalExecution.GetRuntimePlatform() != UnrealTargetPlatform.Mac && TargetPlatform == UnrealTargetPlatform.IOS )
+				{
+					LinkerFixupsFileContents.Add("#include \"" + HeaderFilename + "\"");
+				}
+				else
+				{
+					LinkerFixupsFileContents.Add("#include \"" + LinkerFixupHeaderFilenameWithPath + "\"");
+				}
 
 				// Add a function that is not referenced by anything that invokes all the empty functions in the different static libraries
 				LinkerFixupsFileContents.Add("void " + LinkerFixupsName + "()");
