Potential fix for iOS build farm error
#ue4
#ios

diff --git a/Engine/Source/Programs/UnrealBuildTool/IOS/IOSToolChain.cs b/Engine/Source/Programs/UnrealBuildTool/IOS/IOSToolChain.cs
--- a/Engine/Source/Programs/UnrealBuildTool/IOS/IOSToolChain.cs
+++ b/Engine/Source/Programs/UnrealBuildTool/IOS/IOSToolChain.cs
@@ -834,6 +834,13 @@ namespace UnrealBuildTool
 
 					if (Directory.Exists (Project))
 					{
+						// ensure the plist, entitlements, and provision files are properly copied
+						UEBuildDeploy DeployHandler = UEBuildDeploy.GetBuildDeploy(Target.Platform);
+						if (DeployHandler != null)
+						{
+							DeployHandler.PrepTargetForDeployment(Target);
+						}
+
 						// code sign the project
 						string CmdLine = XcodeDeveloperDir + "usr/bin/xcodebuild" +
 						                " -project \"" + Project + "\"" +
