Rocket build fixes.

#remove Removed the build of .generated.cpp if there are no source files to build against (like there would be in Rocket).
#remove Renamed AutoGenerateInlInfo* to AutoGenerateCppInfo*.
#fix Fixed the #includes in *.generated.cpp to be relative  instead of absolute in Rocket builds.

diff --git a/Engine/Source/Programs/UnrealBuildTool/Configuration/UEBuildModule.cs b/Engine/Source/Programs/UnrealBuildTool/Configuration/UEBuildModule.cs
--- a/Engine/Source/Programs/UnrealBuildTool/Configuration/UEBuildModule.cs
+++ b/Engine/Source/Programs/UnrealBuildTool/Configuration/UEBuildModule.cs
@@ -706,12 +706,12 @@ namespace UnrealBuildTool
 	/** A module that is compiled from C++ code. */
 	public class UEBuildModuleCPP : UEBuildModule
 	{
-		public class AutoGenerateInlInfoClass
+		public class AutoGenerateCppInfoClass
 		{
 			/** The filename of the *.generated.cpp file which was generated for the module */
 			public readonly string Filename;
 
-			public AutoGenerateInlInfoClass(string InFilename)
+			public AutoGenerateCppInfoClass(string InFilename)
 			{
 				Debug.Assert(InFilename != null);
 
@@ -720,9 +720,9 @@ namespace UnrealBuildTool
 		}
 
 		/** Information about the .inl file . */
-		public AutoGenerateInlInfoClass AutoGenerateInlInfo = null;
+		public AutoGenerateCppInfoClass AutoGenerateCppInfo = null;
 
-		private class SourceFilesClass
+		public class SourceFilesClass
 		{
 			public readonly List<FileItem> MissingFiles = new List<FileItem>();
 			public readonly List<FileItem> CPPFiles     = new List<FileItem>();
@@ -731,10 +731,24 @@ namespace UnrealBuildTool
 			public readonly List<FileItem> MMFiles      = new List<FileItem>();
 			public readonly List<FileItem> RCFiles      = new List<FileItem>();
 			public readonly List<FileItem> OtherFiles   = new List<FileItem>();
+
+			public int Count
+			{
+				get
+				{
+					return MissingFiles.Count +
+					       CPPFiles    .Count +
+					       CFiles      .Count +
+					       CCFiles     .Count +
+					       MMFiles     .Count +
+					       RCFiles     .Count +
+					       OtherFiles  .Count;
+				}
+			}
 		}
 
 		/** A list of the absolute paths of source files in the module. */
-		readonly SourceFilesClass SourceFiles = new SourceFilesClass();
+		public readonly SourceFilesClass SourceFiles = new SourceFilesClass();
 
 		/** The preprocessor definitions used to compile this module's private implementation. */
 		List<string> Definitions;
@@ -1295,9 +1309,9 @@ namespace UnrealBuildTool
 				PCHPath = UEToolChain.GetPlatformToolChain(ModuleCompileEnvironment.Config.TargetPlatform).ConvertPath( ProcessedDependencies.UniquePCHHeaderFile.AbsolutePath );
 			}
 
-			if (AutoGenerateInlInfo != null && !ModuleCompileEnvironment.bHackHeaderGenerator)
+			if (AutoGenerateCppInfo != null && !ModuleCompileEnvironment.bHackHeaderGenerator)
 			{
-				var GeneratedCppFileItem = FileItem.GetItemByPath( AutoGenerateInlInfo.Filename );
+				var GeneratedCppFileItem = FileItem.GetItemByPath( AutoGenerateCppInfo.Filename );
 
 				LinkInputFiles.AddRange( CPPCompileEnvironment.CompileFiles( new List<FileItem>{ GeneratedCppFileItem }, Name ).ObjectFiles );
 			}
diff --git a/Engine/Source/Programs/UnrealBuildTool/Configuration/UEBuildTarget.cs b/Engine/Source/Programs/UnrealBuildTool/Configuration/UEBuildTarget.cs
--- a/Engine/Source/Programs/UnrealBuildTool/Configuration/UEBuildTarget.cs
+++ b/Engine/Source/Programs/UnrealBuildTool/Configuration/UEBuildTarget.cs
@@ -1386,10 +1386,20 @@ namespace UnrealBuildTool
 						if( UHTModuleInfo.PublicUObjectClassesHeaders.Count == 0 && UHTModuleInfo.PrivateUObjectHeaders.Count == 0 && UHTModuleInfo.PublicUObjectHeaders.Count == 0 )
 							continue;
 
-						UHTModuleInfo.PCH                      = ToolChain.ConvertPath( DependencyModuleCPP.ProcessedDependencies.UniquePCHHeaderFile.AbsolutePath );
+						UHTModuleInfo.PCH                      = "";
 						UHTModuleInfo.GeneratedCPPFilenameBase = Path.Combine( UEBuildModuleCPP.GetGeneratedCodeDirectoryForModule(this, UHTModuleInfo.ModuleDirectory, UHTModuleInfo.ModuleName), UHTModuleInfo.ModuleName ) + ".generated";
 
-						DependencyModuleCPP.AutoGenerateInlInfo = new UEBuildModuleCPP.AutoGenerateInlInfoClass( UHTModuleInfo.GeneratedCPPFilenameBase + ".cpp" );
+						if (DependencyModuleCPP.ProcessedDependencies.UniquePCHHeaderFile != null)
+						{
+							UHTModuleInfo.PCH = DependencyModuleCPP.ProcessedDependencies.UniquePCHHeaderFile.AbsolutePath;
+						}
+
+						// If we've got this far and there are no source files then it's likely we're running Rocket and ignoring
+						// engine files, so we don't need a .generated.cpp either
+						if (DependencyModuleCPP.SourceFiles.Count != 0)
+						{
+							DependencyModuleCPP.AutoGenerateCppInfo = new UEBuildModuleCPP.AutoGenerateCppInfoClass( UHTModuleInfo.GeneratedCPPFilenameBase + ".cpp" );
+						}
 
 						UObjectModules.Add( UHTModuleInfo );
 						Log.TraceVerbose( "Detected UObject module: " + DependencyModuleCPP.Name );
diff --git a/Engine/Source/Programs/UnrealHeaderTool/Private/CodeGenerator.cpp b/Engine/Source/Programs/UnrealHeaderTool/Private/CodeGenerator.cpp
--- a/Engine/Source/Programs/UnrealHeaderTool/Private/CodeGenerator.cpp
+++ b/Engine/Source/Programs/UnrealHeaderTool/Private/CodeGenerator.cpp
@@ -4500,7 +4500,20 @@ void FNativeClassHeaderGenerator::ExportGeneratedCPP()
 	const auto* ModuleInfo = GPackageToManifestModuleMap.FindChecked(Package);
 	if (ModuleInfo->PCH.Len())
 	{
-		ModulePCHInclude = FString::Printf(TEXT("#include \"%s\"") LINE_TERMINATOR, *ModuleInfo->PCH);
+		FString PCH = ModuleInfo->PCH;
+		if(bUseRelativePaths)
+		{
+			// When building Rocket, we need relative paths in the generated headers so that they can be installed to any directory.
+			FPaths::MakePathRelativeTo(PCH, *ModuleInfo->GeneratedCPPFilenameBase);
+			PCH.ReplaceInline(TEXT("\\"), TEXT("/"));
+		}
+		// ...we cannot rebase paths relative to the install directory. It may be on a different drive.
+		else if (FPaths::MakePathRelativeTo(PCH, *GManifest.RootLocalPath))
+		{
+			PCH = GManifest.RootBuildPath / PCH;
+		}
+
+		ModulePCHInclude = FString::Printf(TEXT("#include \"%s\"") LINE_TERMINATOR, *PCH);
 	}
 
 	// Write out the ordered class dependencies into a single header that we can easily include
