build script, allow IOS on PC

diff --git a/Engine/Source/Programs/UnrealBuildTool/IOS/UEBuildIOS.cs b/Engine/Source/Programs/UnrealBuildTool/IOS/UEBuildIOS.cs
--- a/Engine/Source/Programs/UnrealBuildTool/IOS/UEBuildIOS.cs
+++ b/Engine/Source/Programs/UnrealBuildTool/IOS/UEBuildIOS.cs
@@ -25,12 +25,6 @@ namespace UnrealBuildTool
 		 */
 		public override void RegisterBuildPlatform()
 		{
-			[REDACTED_EMAIL]: Add platform support
-			if ((UnrealBuildTool.RunningRocket() || UnrealBuildTool.BuildingRocket()) && ExternalExecution.GetRuntimePlatform() != UnrealTargetPlatform.Mac)
-			{
-				return;
-			}
-
 			// Register this build platform for IOS
 			Log.TraceVerbose("        Registering for {0}", UnrealTargetPlatform.IOS.ToString());
 			UEBuildPlatform.RegisterBuildPlatform(UnrealTargetPlatform.IOS, this);
