Exclude client and server targets in binary Rocket. It's impossible to build them since we don't include precompiled static libraries.

diff --git a/Engine/Source/Programs/UnrealBuildTool/System/ProjectFileGenerator.cs b/Engine/Source/Programs/UnrealBuildTool/System/ProjectFileGenerator.cs
--- a/Engine/Source/Programs/UnrealBuildTool/System/ProjectFileGenerator.cs
+++ b/Engine/Source/Programs/UnrealBuildTool/System/ProjectFileGenerator.cs
@@ -1519,142 +1519,145 @@ namespace UnrealBuildTool
 					}
 				}
 
-				if( WantProjectFileForTarget )
+				if (WantProjectFileForTarget)
 				{
 					// Create target rules for all of the platforms and configuration combinations that we want to enable support for.
 					// Just use Win64 or Mac as we only need to recover the target type and both should be supported for all targets...
 					string UnusedTargetFilePath;
-                    var TargetRulesObject = RulesCompiler.CreateTargetRules( TargetName, new TargetInfo( Utils.IsRunningOnMono ? UnrealTargetPlatform.Mac : UnrealTargetPlatform.Win64, UnrealTargetConfiguration.Development ), false, out UnusedTargetFilePath );
+					var TargetRulesObject = RulesCompiler.CreateTargetRules(TargetName, new TargetInfo(Utils.IsRunningOnMono ? UnrealTargetPlatform.Mac : UnrealTargetPlatform.Win64, UnrealTargetConfiguration.Development), false, out UnusedTargetFilePath);
 
-
-					bool IsProgramTarget = false;
-
-					string GameFolder = null;
-					string ProjectFileNameBase = null;
-					if( TargetRulesObject.Type == TargetRules.TargetType.Program )
-					{
-						IsProgramTarget = true;
-						ProjectFileNameBase = TargetName;
-					}
-					else if( IsEngineTarget )
-					{
-						ProjectFileNameBase = EngineProjectFileNameBase;
-					}
-					else
+					// Exclude client and server targets under binary Rocket; it's impossible to build without precompiled engine binaries
+					if (!UnrealBuildTool.RunningRocket() || (TargetRulesObject.Type != TargetRules.TargetType.Client && TargetRulesObject.Type != TargetRules.TargetType.Server))
 					{
-						// Figure out which game project this target belongs to
-						UProjectInfo ProjectInfo = FindGameContainingFile(AllGames, TargetFilePath);
-						if(ProjectInfo == null)
+						bool IsProgramTarget = false;
+
+						string GameFolder = null;
+						string ProjectFileNameBase = null;
+						if (TargetRulesObject.Type == TargetRules.TargetType.Program)
 						{
-							throw new BuildException( "Found a non-engine target file (" + TargetFilePath + ") that did not exist within any of the known game folders" );
+							IsProgramTarget = true;
+							ProjectFileNameBase = TargetName;
+						}
+						else if (IsEngineTarget)
+						{
+							ProjectFileNameBase = EngineProjectFileNameBase;
+						}
+						else
+						{
+							// Figure out which game project this target belongs to
+							UProjectInfo ProjectInfo = FindGameContainingFile(AllGames, TargetFilePath);
+							if (ProjectInfo == null)
+							{
+								throw new BuildException("Found a non-engine target file (" + TargetFilePath + ") that did not exist within any of the known game folders");
+							}
+							GameFolder = ProjectInfo.Folder;
+							ProjectFileNameBase = ProjectInfo.GameName;
 						}
-						GameFolder = ProjectInfo.Folder;
-						ProjectFileNameBase = ProjectInfo.GameName;
-					}
 
-					// @todo projectfiles: We should move all of the Target.cs files out of sub-folders to clean up the project directories a bit (e.g. GameUncooked folder)
+						// @todo projectfiles: We should move all of the Target.cs files out of sub-folders to clean up the project directories a bit (e.g. GameUncooked folder)
 
-					var ProjectFilePath = Path.Combine( IntermediateProjectFilesPath, ProjectFileNameBase + ProjectFileExtension );
+						var ProjectFilePath = Path.Combine(IntermediateProjectFilesPath, ProjectFileNameBase + ProjectFileExtension);
 
-					if (TargetRules.IsGameType(TargetRulesObject.Type) &&
-						(TargetRules.IsEditorType(TargetRulesObject.Type) == false))
-					{
-						// Allow platforms to generate stub projects here...
-						UEPlatformProjectGenerator.GenerateGameProjectStubs(
-							InGenerator: this, 
-							InTargetName: TargetName,
-							InTargetFilepath: TargetFilePath, 
-							InTargetRules: TargetRulesObject,
-							InPlatforms: SupportedPlatforms,
-							InConfigurations: SupportedConfigurations);
-					}
+						if (TargetRules.IsGameType(TargetRulesObject.Type) &&
+							(TargetRules.IsEditorType(TargetRulesObject.Type) == false))
+						{
+							// Allow platforms to generate stub projects here...
+							UEPlatformProjectGenerator.GenerateGameProjectStubs(
+								InGenerator: this,
+								InTargetName: TargetName,
+								InTargetFilepath: TargetFilePath,
+								InTargetRules: TargetRulesObject,
+								InPlatforms: SupportedPlatforms,
+								InConfigurations: SupportedConfigurations);
+						}
 
-					ProjectFilePath = Utils.MakePathRelativeTo( ProjectFilePath, MasterProjectRelativePath );
+						ProjectFilePath = Utils.MakePathRelativeTo(ProjectFilePath, MasterProjectRelativePath);
 
-					bool bProjectAlreadyExisted;
-					var ProjectFile = FindOrAddProject( ProjectFilePath, IncludeInGeneratedProjects: true, bAlreadyExisted:out bProjectAlreadyExisted );
-					ProjectFile.IsGeneratedProject = true;
-					ProjectFile.IsStubProject = false;
+						bool bProjectAlreadyExisted;
+						var ProjectFile = FindOrAddProject(ProjectFilePath, IncludeInGeneratedProjects: true, bAlreadyExisted: out bProjectAlreadyExisted);
+						ProjectFile.IsGeneratedProject = true;
+						ProjectFile.IsStubProject = false;
 
-					bool IsTemplateTarget = false;
-					{
-						// Check to see if this is a template target.  That is, the target is located under the "Templates" folder
-						string TargetFileRelativeToTemplatesDirectory = Utils.MakePathRelativeTo(TargetFilePath, Path.Combine(RootRelativePath, "Templates"));
-						if (!TargetFileRelativeToTemplatesDirectory.StartsWith("..") && !Path.IsPathRooted(TargetFileRelativeToTemplatesDirectory))
+						bool IsTemplateTarget = false;
 						{
-							IsTemplateTarget = true;
+							// Check to see if this is a template target.  That is, the target is located under the "Templates" folder
+							string TargetFileRelativeToTemplatesDirectory = Utils.MakePathRelativeTo(TargetFilePath, Path.Combine(RootRelativePath, "Templates"));
+							if (!TargetFileRelativeToTemplatesDirectory.StartsWith("..") && !Path.IsPathRooted(TargetFileRelativeToTemplatesDirectory))
+							{
+								IsTemplateTarget = true;
+							}
 						}
-					}
-					string BaseFolder = null;
-					if( IsProgramTarget )
-					{
-						ProgramProjects[ TargetName ] = ProjectFile;
-						BaseFolder = Path.GetDirectoryName( TargetFilePath );
-					}
-					else if( IsEngineTarget )
-					{
-						EngineProject = ProjectFile;
-						BaseFolder = EngineRelativePath;
-					}
-					else
-					{
-						GameProjects[ GameFolder ] = ProjectFile;
-						if( IsTemplateTarget )
+						string BaseFolder = null;
+						if (IsProgramTarget)
 						{
-							TemplateGameProjects.Add( ProjectFile );
+							ProgramProjects[TargetName] = ProjectFile;
+							BaseFolder = Path.GetDirectoryName(TargetFilePath);
 						}
-						BaseFolder = GameFolder;
-
-						if( !bProjectAlreadyExisted )
+						else if (IsEngineTarget)
 						{
-							// Add the .uproject file for this game/template
-							var UProjectFilePath = Path.Combine( BaseFolder, ProjectFileNameBase + ".uproject" );
-							if( File.Exists( UProjectFilePath ) )
+							EngineProject = ProjectFile;
+							BaseFolder = EngineRelativePath;
+						}
+						else
+						{
+							GameProjects[GameFolder] = ProjectFile;
+							if (IsTemplateTarget)
 							{
-								ProjectFile.AddFileToProject( UProjectFilePath, BaseFolder );
+								TemplateGameProjects.Add(ProjectFile);
 							}
-							else
+							BaseFolder = GameFolder;
+
+							if (!bProjectAlreadyExisted)
 							{
-								throw new BuildException( "Not expecting to find a game with no .uproject file.  File '{0}' doesn't exist", UProjectFilePath );
+								// Add the .uproject file for this game/template
+								var UProjectFilePath = Path.Combine(BaseFolder, ProjectFileNameBase + ".uproject");
+								if (File.Exists(UProjectFilePath))
+								{
+									ProjectFile.AddFileToProject(UProjectFilePath, BaseFolder);
+								}
+								else
+								{
+									throw new BuildException("Not expecting to find a game with no .uproject file.  File '{0}' doesn't exist", UProjectFilePath);
+								}
 							}
-						}
 
-					}
-
-					foreach( var ExistingProjectTarget in ProjectFile.ProjectTargets )
-					{
-						if( ExistingProjectTarget.TargetRules.ConfigurationName.Equals( TargetRulesObject.ConfigurationName, StringComparison.InvariantCultureIgnoreCase ) )
-						{
-							throw new BuildException( "Not expecting project {0} to already have a target rules of with configuration name {1} ({2}) while trying to add: {3}", ProjectFilePath, TargetRulesObject.ConfigurationName, ExistingProjectTarget.TargetRules.ToString(), TargetRulesObject.ToString() );
 						}
 
-						// Not expecting to have both a game and a program in the same project.  These would alias because we share the project and solution configuration names (just because it makes sense to)
-						if( ExistingProjectTarget.TargetRules.Type == TargetRules.TargetType.Game && ExistingProjectTarget.TargetRules.Type == TargetRules.TargetType.Program ||
-							ExistingProjectTarget.TargetRules.Type == TargetRules.TargetType.Program && ExistingProjectTarget.TargetRules.Type == TargetRules.TargetType.Game )
+						foreach (var ExistingProjectTarget in ProjectFile.ProjectTargets)
 						{
-							throw new BuildException( "Not expecting project {0} to already have a Game/Program target ({1}) associated with it while trying to add: {2}", ProjectFilePath, ExistingProjectTarget.TargetRules.ToString(), TargetRulesObject.ToString() );
+							if (ExistingProjectTarget.TargetRules.ConfigurationName.Equals(TargetRulesObject.ConfigurationName, StringComparison.InvariantCultureIgnoreCase))
+							{
+								throw new BuildException("Not expecting project {0} to already have a target rules of with configuration name {1} ({2}) while trying to add: {3}", ProjectFilePath, TargetRulesObject.ConfigurationName, ExistingProjectTarget.TargetRules.ToString(), TargetRulesObject.ToString());
+							}
+
+							// Not expecting to have both a game and a program in the same project.  These would alias because we share the project and solution configuration names (just because it makes sense to)
+							if (ExistingProjectTarget.TargetRules.Type == TargetRules.TargetType.Game && ExistingProjectTarget.TargetRules.Type == TargetRules.TargetType.Program ||
+								ExistingProjectTarget.TargetRules.Type == TargetRules.TargetType.Program && ExistingProjectTarget.TargetRules.Type == TargetRules.TargetType.Game)
+							{
+								throw new BuildException("Not expecting project {0} to already have a Game/Program target ({1}) associated with it while trying to add: {2}", ProjectFilePath, ExistingProjectTarget.TargetRules.ToString(), TargetRulesObject.ToString());
+							}
 						}
-					}
 
-					var ProjectTarget = new ProjectTarget()
-						{
-							TargetRules = TargetRulesObject,
-							TargetFilePath = TargetFilePath
-						};
+						var ProjectTarget = new ProjectTarget()
+							{
+								TargetRules = TargetRulesObject,
+								TargetFilePath = TargetFilePath
+							};
 
-					if( TargetName == "ShaderCompileWorker" )		// @todo projectfiles: Ideally, the target rules file should set this
-					{
-						ProjectTarget.ForceDevelopmentConfiguration = true;
-					}
+						if (TargetName == "ShaderCompileWorker")		// @todo projectfiles: Ideally, the target rules file should set this
+						{
+							ProjectTarget.ForceDevelopmentConfiguration = true;
+						}
 
-					ProjectFile.ProjectTargets.Add( ProjectTarget );
+						ProjectFile.ProjectTargets.Add(ProjectTarget);
 
-					// Make sure the *.Target.cs file is in the project.
-					ProjectFile.AddFileToProject( TargetFilePath, BaseFolder );
+						// Make sure the *.Target.cs file is in the project.
+						ProjectFile.AddFileToProject(TargetFilePath, BaseFolder);
 
 
-					// We special case ShaderCompileWorker.  It needs to always be compiled in Development mode.
-					Log.TraceVerbose( "Generating target {0} for {1}", TargetRulesObject.Type.ToString(), ProjectFilePath );
+						// We special case ShaderCompileWorker.  It needs to always be compiled in Development mode.
+						Log.TraceVerbose("Generating target {0} for {1}", TargetRulesObject.Type.ToString(), ProjectFilePath);
+					}
 				}
 			}
 		}
