#UE4: HTML5 Allow building and cooking for HTML5 platform on Mac.

  - Install Emscripten SDK > 1.16
  - Setup system wide enviorment variable EMSCRIPTEN. (  http://stackoverflow.com/questions/135688/setting-environment-variables-in-os-x )
  - Build HTML5 using Build.sh from command line.

[REDACTED_HANDLE] : Figure out whether to support cooking if SDK is not installed on host platform.

diff --git a/Engine/Build/BatchFiles/Mac/Build.sh b/Engine/Build/BatchFiles/Mac/Build.sh
--- a/Engine/Build/BatchFiles/Mac/Build.sh
+++ b/Engine/Build/BatchFiles/Mac/Build.sh
@@ -6,53 +6,72 @@
 # Setup Mono
 source Engine/Build/BatchFiles/Mac/SetupMono.sh Engine/Build/BatchFiles/Mac
 
+# override env if action is specified on command line 
+if [ $1 == "clean" ] 
+	then 
+		ACTION="clean" 
+fi 
+
+
+echo Building UBT... 
+
+xbuild /property:Configuration=Development  /nologo Engine/Source/Programs/UnrealBuildTool/UnrealBuildTool_Mono.csproj| grep -i error  
+
 case $ACTION in
 	"")
-		echo "Building $1..."
-		xbuild /property:Configuration=Development /verbosity:quiet /nologo Engine/Source/Programs/UnrealBuildTool/UnrealBuildTool_Mono.csproj |grep -i error
+		echo Building $1...
 
                 Platform=""
-                AdditionalFlags="-deploy"
-		                
-                if [ $2 = "iphoneos" ]
-                then
-	                Platform="IOS"
-			AdditionalFlags+=" -nocreatestub"
-                else
-	                if [ $2 = "iphonesimulator" ]
-	                then
+                AdditionalFlags=""
+		
+
+		case $2 in 
+	                "iphoneos") 
 		                Platform="IOS"
-		                AdditionalFlags+=" -simulator"
-				AdditionalFlags+=" -nocreatestub"
-	                else
-		                Platform="Mac"
-	                fi
-                fi
+				AdditionalFlags+=" -deploy -nocreatestub "
+			;; 
+  			"iphonesimulator")
+		        	Platform="IOS"
+		         	AdditionalFlags+=" -deploy -simulator -nocreatestub"
+			;;
 
+			"HTML5") 
+				Platform="HTML5"
+			;;
+			*)	
+				Platform="Mac"
+				AdditionalFlags+=" -deploy " 
+			;; 
+		esac
+
+		echo Running command : Engine/Binaries/DotNET/UnrealBuildTool.exe $1 $Platform $3 $AdditionalFlags "$4" 
 		mono Engine/Binaries/DotNET/UnrealBuildTool.exe $1 $Platform $3 $AdditionalFlags "$4"
 		;;
 	"clean")
 		echo "Cleaning $2 $3 $4..."
-		xbuild /property:Configuration=Development /verbosity:quiet /nologo Engine/Source/Programs/UnrealBuildTool/UnrealBuildTool_Mono.csproj |grep -i error
 
                 Platform=""
                 AdditionalFlags="-clean"
-		                
-                if [ $3 = "iphoneos" ]
-                then
-	                Platform="IOS"
-			AdditionalFlags+=" -nocreatestub"
-                else
-	                if [ $3 = "iphonesimulator" ]
-	                then
-		                Platform="IOS"
-		                AdditionalFlags+=" -simulator"
-				AdditionalFlags+=" -nocreatestub"
-	                else
-		                Platform="Mac"
-	                fi
-                fi
 
+		case $3 in 
+			"iphoneos")
+	        	        Platform="IOS"
+				AdditionalFlags+=" -nocreatestub"
+			;;
+			"iphonesimulator")
+			        Platform="IOS"
+		        	AdditionalFlags+=" -simulator"
+				AdditionalFlags+=" -nocreatestub"
+			;; 
+			"HTML5")
+				Platform="HTML5"
+			;;
+			*) 
+        			Platform="Mac"
+			;;
+		
+		esac
+		echo Running command: mono Engine/Binaries/DotNET/UnrealBuildTool.exe $2 $Platform $4 $AdditionalFlags "$5"
 		mono Engine/Binaries/DotNET/UnrealBuildTool.exe $2 $Platform $4 $AdditionalFlags "$5"
 		;;
 esac
diff --git a/Engine/Source/Programs/AutomationTool/HTML5/HTML5Platform.Automation.cs b/Engine/Source/Programs/AutomationTool/HTML5/HTML5Platform.Automation.cs
--- a/Engine/Source/Programs/AutomationTool/HTML5/HTML5Platform.Automation.cs
+++ b/Engine/Source/Programs/AutomationTool/HTML5/HTML5Platform.Automation.cs
@@ -6,9 +6,6 @@ using System.Text;
 using System.IO;
 using AutomationTool;
 using UnrealBuildTool;
-#if !MONO
-using Microsoft.Win32; 
-#endif 
 
 public class HTML5Platform : Platform
 {
@@ -26,35 +23,20 @@ public class HTML5Platform : Platform
 		// we need to operate in the root
 		using (new PushedDirectory(Path.Combine(Params.BaseStageDirectory, "HTML5")))
 		{
-
-            string BaseSDKPath = "";
-#if !MONO
-            RegistryKey localKey = RegistryKey.OpenBaseKey(Microsoft.Win32.RegistryHive.LocalMachine, RegistryView.Registry64);
-            RegistryKey key = localKey.OpenSubKey("Software\\Emscripten64");
-            if (key != null)
-            {
-                string sdkdir = key.GetValue("Install_Dir") as string;
-                // emscripten path is the highest numbered directory 
-                DirectoryInfo dInfo = new DirectoryInfo(sdkdir + "\\emscripten");
-                string Latest_Ver = (from S in dInfo.GetDirectories() select S.Name).ToList().Last();
-                BaseSDKPath = sdkdir + @"\emscripten\" + Latest_Ver;
-            }
-            else
-            {
-                BaseSDKPath = Environment.ExpandEnvironmentVariables("%EMSCRIPTEN%");
-                
-            }
-#else 
-            BaseSDKPath = Environment.ExpandEnvironmentVariables("%EMSCRIPTEN%"); 
-#endif
+			string BaseSDKPath = Environment.GetEnvironmentVariable("EMSCRIPTEN"); 
             // make the file_packager command line
-            string PackagerPath = "\""  + BaseSDKPath + "\\tools\\file_packager.py\"";
-			string CmdLine = string.Format("/c {0} {1} --preload . --pre-run --js-output={1}.js", PackagerPath, FinalDataLocation);
-
-			// package it up!
-			RunAndLog(CmdEnv, CommandUtils.CombinePaths(Environment.SystemDirectory, "cmd.exe"), CmdLine);
+			if (Utils.IsRunningOnMono) 
+			{
+				string PackagerPath = BaseSDKPath + "/tools/file_packager.py";
+				string CmdLine = string.Format ("-c \" python {0} {1} --preload . --pre-run --js-output={1}.js \" ", PackagerPath, FinalDataLocation);
+				RunAndLog (CmdEnv, "/bin/bash", CmdLine);	
+			} else 
+			{ 
+				string PackagerPath = "\"" + BaseSDKPath + "\\tools\\file_packager.py\"";
+				string CmdLine = string.Format ("/c {0} {1} --preload . --pre-run --js-output={1}.js", PackagerPath, FinalDataLocation);
+				RunAndLog (CmdEnv, CommandUtils.CombinePaths (Environment.SystemDirectory, "cmd.exe"), CmdLine);	
+			}
 		}
-		
 		// put the HTML file to the binaries directory
 		string TemplateFile = Path.Combine(CombinePaths(CmdEnv.LocalRoot, "Engine"), "Build", "HTML5", "Game.html.template");
 		string OutputFile = Path.Combine(Path.GetDirectoryName(Path.GetFullPath(Params.ProjectGameExeFilename)), Params.ProjectBinariesFolder, Params.ShortProjectName) + ".html";
diff --git a/Engine/Source/Programs/UnrealBuildTool/HTML5/HTML5ToolChain.cs b/Engine/Source/Programs/UnrealBuildTool/HTML5/HTML5ToolChain.cs
--- a/Engine/Source/Programs/UnrealBuildTool/HTML5/HTML5ToolChain.cs
+++ b/Engine/Source/Programs/UnrealBuildTool/HTML5/HTML5ToolChain.cs
@@ -13,50 +13,19 @@ namespace UnrealBuildTool
 	{
 		// cache the location of SDK tools
 		static string EMCCPath;
-        static string EMLinkPath;
 		static string PythonPath;
 
-		static private bool bEnableFastIteration = UnrealBuildTool.CommandLineContains("-fastiteration");
-
 		public override void RegisterToolChain()
 		{
 			// Make sure the SDK is installed
             // look up installed SDK. 
-            string BaseSDKPath = ""; 
-			if (!Utils.IsRunningOnMono)
-			{
-				Microsoft.Win32.RegistryKey LocalKey = Microsoft.Win32.RegistryKey.OpenBaseKey(Microsoft.Win32.RegistryHive.LocalMachine, Microsoft.Win32.RegistryView.Registry64); 
-				Microsoft.Win32.RegistryKey Key = LocalKey.OpenSubKey("Software\\Emscripten64"); 
-	            if (Key != null)
-	            {
-	                string SDKDir = Key.GetValue("Install_Dir") as string;
-					if (SDKDir != null)
-					{
-						SDKDir = Path.Combine(SDKDir, "emscripten");
-						if (Directory.Exists(SDKDir))
-						{
-							// emscripten path is the highest numbered directory 
-							DirectoryInfo DirInfo = new DirectoryInfo(SDKDir);
-							string Latest_Ver = (from S in DirInfo.GetDirectories() select S.Name).ToList().Last();
-							BaseSDKPath = Path.Combine(SDKDir, Latest_Ver);
-						}
-					}
-	            }
-			}
-
-			// if the above didn't set the SDK path, use the environment variable (for anyone who installed by unzipping or similar)
-			if (BaseSDKPath == "")
-			{
-                BaseSDKPath = Environment.GetEnvironmentVariable("EMSCRIPTEN");
-			}
-
+			string BaseSDKPath = Environment.GetEnvironmentVariable("EMSCRIPTEN");
 			if (!String.IsNullOrEmpty(BaseSDKPath))
 			{
 				BaseSDKPath = BaseSDKPath.Replace("\"", "");
 				if (!String.IsNullOrEmpty(BaseSDKPath))
                 {
 					EMCCPath = Path.Combine(BaseSDKPath, "emcc");
-                    EMLinkPath = Path.Combine(BaseSDKPath, "emlink.py");
 					// also figure out where python lives (if no envvar, assume it's in the path)
 					PythonPath = Environment.GetEnvironmentVariable("PYTHON");
 					if (PythonPath == null)
@@ -64,14 +33,11 @@ namespace UnrealBuildTool
 						PythonPath = Utils.IsRunningOnMono ? "python" : "python.exe";
 					}
                     EMCCPath = "\"" + EMCCPath + "\""; 
-
 					// set some environment variable we'll need
-
 					//Environment.SetEnvironmentVariable("EMCC_DEBUG", "cache");
 					Environment.SetEnvironmentVariable("EMCC_CORES", "8");
 					Environment.SetEnvironmentVariable("EMCC_FORCE_STDLIBS", "1");
 					Environment.SetEnvironmentVariable("EMCC_OPTIMIZE_NORMALLY", "1");
-
 					// finally register the toolchain that is now ready to go
                     Log.TraceVerbose("        Registered for {0}", CPPTargetPlatform.HTML5.ToString());
 					UEToolChain.RegisterPlatformToolChain(CPPTargetPlatform.HTML5, this);
@@ -140,7 +106,6 @@ namespace UnrealBuildTool
             }
 
             // NOTE: This may slow down the compiler's startup time!
-            if ( !bEnableFastIteration )
             { 
                 Result += " --memory-init-file 1";
             }
@@ -351,14 +316,13 @@ namespace UnrealBuildTool
 				}
 
 				// Add the source file path to the command-line.
-                string bfastlinkstring = bEnableFastIteration ? "" : " -c ";
-                string FileArguments = string.Format(bfastlinkstring + " \"{0}\"", SourceFile.AbsolutePath);
+                string FileArguments = string.Format(" \"{0}\"", SourceFile.AbsolutePath);
 
             	// Add the object file to the produced item list.
 				FileItem ObjectFile = FileItem.GetItemByPath(
 					Path.Combine(
 						CompileEnvironment.Config.OutputDirectory,
-						Path.GetFileName(SourceFile.AbsolutePath) + (bEnableFastIteration ? ".js" : ".o")
+						Path.GetFileName(SourceFile.AbsolutePath) + (".o")
 						)
 					);
 				CompileAction.ProducedItems.Add(ObjectFile);
@@ -377,8 +341,7 @@ namespace UnrealBuildTool
 				CompileAction.WorkingDirectory = Path.GetFullPath(".");
 				CompileAction.CommandPath = PythonPath;
                 
-                string fastlinkString = SourceFile.Info.FullName.Contains("Launch") ?  " -s MAIN_MODULE=1 " : "-s SIDE_MODULE=1";
-                CompileAction.CommandArguments = EMCCPath + Arguments + (bEnableFastIteration ? fastlinkString : "" )+ FileArguments + CompileEnvironment.Config.AdditionalArguments;
+				CompileAction.CommandArguments = EMCCPath + Arguments + FileArguments + CompileEnvironment.Config.AdditionalArguments;
                
                 System.Console.WriteLine(CompileAction.CommandArguments); 
 				CompileAction.StatusDescription = Path.GetFileName(SourceFile.AbsolutePath);
@@ -485,15 +448,14 @@ namespace UnrealBuildTool
 			LinkAction.bCanExecuteRemotely = false;
 			LinkAction.WorkingDirectory = Path.GetFullPath(".");
 			LinkAction.CommandPath = PythonPath;
-			LinkAction.CommandArguments = bEnableFastIteration ? EMLinkPath : EMCCPath;
+			LinkAction.CommandArguments = EMCCPath;
 		    LinkAction.CommandArguments += GetLinkArguments(LinkEnvironment);
 
 			// Add the input files to a response file, and pass the response file on the command-line.
 			foreach (FileItem InputFile in LinkEnvironment.InputFiles)
 			{
                 System.Console.WriteLine("File  {0} ", InputFile.AbsolutePath);
-                string fastlinkString = InputFile.AbsolutePath.Contains("Launch.cpp") ?  " -m " : " -s ";
-                LinkAction.CommandArguments += string.Format((bEnableFastIteration ? fastlinkString : "") + " \"{0}\"", InputFile.AbsolutePath);
+                LinkAction.CommandArguments += string.Format(" \"{0}\"", InputFile.AbsolutePath);
 				LinkAction.PrerequisiteItems.Add(InputFile);
 			}
             foreach (string InputFile in LinkEnvironment.Config.AdditionalLibraries)
@@ -509,12 +471,10 @@ namespace UnrealBuildTool
 			OutputFile = FileItem.GetItemByPath(LinkEnvironment.Config.OutputFilePath);
 			LinkAction.ProducedItems.Add(OutputFile);
 			LinkAction.CommandArguments += string.Format(" -o \"{0}\"", OutputFile.AbsolutePath);
-            if ( !bEnableFastIteration ) 
-            { 
- 			    FileItem OutputBC = FileItem.GetItemByPath(LinkEnvironment.Config.OutputFilePath.Replace(".js", ".bc").Replace(".html", ".bc"));
- 			    LinkAction.ProducedItems.Add(OutputBC);
- 			    LinkAction.CommandArguments += string.Format(" --save-bc \"{0}\"", OutputBC.AbsolutePath);
-            } 
+
+		    FileItem OutputBC = FileItem.GetItemByPath(LinkEnvironment.Config.OutputFilePath.Replace(".js", ".bc").Replace(".html", ".bc"));
+		    LinkAction.ProducedItems.Add(OutputBC);
+		    LinkAction.CommandArguments += string.Format(" --save-bc \"{0}\"", OutputBC.AbsolutePath);
 
      		LinkAction.StatusDescription = Path.GetFileName(OutputFile.AbsolutePath);
 			LinkAction.OutputEventHandler = new DataReceivedEventHandler(RemoteOutputReceivedEventHandler);
diff --git a/Engine/Source/Programs/UnrealBuildTool/Mac/UEBuildMac.cs b/Engine/Source/Programs/UnrealBuildTool/Mac/UEBuildMac.cs
--- a/Engine/Source/Programs/UnrealBuildTool/Mac/UEBuildMac.cs
+++ b/Engine/Source/Programs/UnrealBuildTool/Mac/UEBuildMac.cs
@@ -92,12 +92,13 @@ namespace UnrealBuildTool
                     }
                 }
 
-				// allow standalone tools to use target platform modules, without needing Engine
-				if (UEBuildConfiguration.bForceBuildTargetPlatforms)
-				{
-					InModule.AddDynamicallyLoadedModule("MacTargetPlatform");
-					InModule.AddDynamicallyLoadedModule("MacNoEditorTargetPlatform");
-				}
+		// allow standalone tools to use target platform modules, without needing Engine
+		if (UEBuildConfiguration.bForceBuildTargetPlatforms)
+		{
+			InModule.AddPlatformSpecificDynamicallyLoadedModule("HTML5TargetPlatform");
+			InModule.AddDynamicallyLoadedModule("MacTargetPlatform");
+			InModule.AddDynamicallyLoadedModule("MacNoEditorTargetPlatform");
+		}
 
                 if (bBuildShaderFormats)
                 {
diff --git a/Engine/Source/Runtime/Engine/Engine.Build.cs b/Engine/Source/Runtime/Engine/Engine.Build.cs
--- a/Engine/Source/Runtime/Engine/Engine.Build.cs
+++ b/Engine/Source/Runtime/Engine/Engine.Build.cs
@@ -113,8 +113,9 @@ public class Engine : ModuleRules
 			{
 				DynamicallyLoadedModuleNames.AddRange(
 					new string[] {
-						"MacTargetPlatform",
-						"MacNoEditorTargetPlatform"
+					    "MacTargetPlatform",
+					    "MacNoEditorTargetPlatform",
+					    "HTML5TargetPlatform"
 					}
 				);
 			}
