Make iOS work with a Windows build of the pre-built binaries
#ue4
#ios

diff --git a/Engine/Source/Programs/UnrealBuildTool/IOS/IOSToolChain.cs b/Engine/Source/Programs/UnrealBuildTool/IOS/IOSToolChain.cs
--- a/Engine/Source/Programs/UnrealBuildTool/IOS/IOSToolChain.cs
+++ b/Engine/Source/Programs/UnrealBuildTool/IOS/IOSToolChain.cs
@@ -4,6 +4,7 @@
 using System;
 using System.Collections;
 using System.Collections.Generic;
+using System.Linq;
 using System.IO;
 using System.Diagnostics;
 using System.Security.AccessControl;
@@ -63,6 +64,8 @@ namespace UnrealBuildTool
 		/** Which library archiver to use */
 		private static string IOSArchiver = "libtool";
 
+		public List<string> BuiltBinaries = new List<string>();
+
 		/// <summary>
 		/// Function to call to reset default data.
 		/// </summary>
@@ -543,7 +546,7 @@ namespace UnrealBuildTool
 						LinkAction.PrerequisiteItems.Add(RemoteLibFile);
 
 						// and add to the commandline
-						LinkAction.CommandArguments += string.Format(" \"{0}\"", ConvertPath(AdditionalLibrary));
+						LinkAction.CommandArguments += string.Format(" \"{0}\"", ConvertPath(Path.GetFullPath(AdditionalLibrary)));
 					}
 					else
 					{
@@ -780,6 +783,16 @@ namespace UnrealBuildTool
 			}
 		}
 
+		public override void PreBuildSync()
+		{
+			if (ExternalExecution.GetRuntimePlatform() != UnrealTargetPlatform.Mac)
+			{
+				BuiltBinaries = new List<string>();
+			}
+
+			base.PreBuildSync();
+		}
+
 		public override void PostBuildSync(UEBuildTarget Target)
 		{
 			base.PostBuildSync(Target);
@@ -878,6 +891,28 @@ namespace UnrealBuildTool
 			}
 			else
 			{
+				// store off the binaries
+				foreach (UEBuildBinary Binary in Target.AppBinaries)
+				{
+					BuiltBinaries.Add(Path.GetFullPath(Binary.ToString()));
+				}
+
+				// Generate static libraries for monolithic games in Rocket
+				if ((UnrealBuildTool.BuildingRocket() || UnrealBuildTool.RunningRocket()) && TargetRules.IsAGame(Target.Rules.Type))
+				{
+					List<UEBuildModule> Modules = Target.AppBinaries[0].GetAllDependencyModules(true, false);
+					foreach (UEBuildModuleCPP Module in Modules.OfType<UEBuildModuleCPP>())
+					{
+						if (Utils.IsFileUnderDirectory(Module.ModuleDirectory, BuildConfiguration.RelativeEnginePath) && Module.Binary == Target.AppBinaries[0])
+						{
+							if (Module.bBuildingRedistStaticLibrary)
+							{
+								BuiltBinaries.Add(Path.GetFullPath(Module.RedistStaticLibraryPath));
+							}
+						}
+					}
+				}
+
 				// check to see if the DangerouslyFast mode is valid (in other words, a build has gone through since a Rebuild/Clean operation)
 				string DangerouslyFastValidFile = Path.Combine(Target.GlobalLinkEnvironment.Config.IntermediateDirectory, "DangerouslyFastIsNotDangerous");
 				bool bUseDangerouslyFastModeWasRequested = bUseDangerouslyFastMode;
diff --git a/Engine/Source/Programs/UnrealBuildTool/IOS/UEDeployIOS.cs b/Engine/Source/Programs/UnrealBuildTool/IOS/UEDeployIOS.cs
--- a/Engine/Source/Programs/UnrealBuildTool/IOS/UEDeployIOS.cs
+++ b/Engine/Source/Programs/UnrealBuildTool/IOS/UEDeployIOS.cs
@@ -195,7 +195,51 @@ namespace UnrealBuildTool.IOS
 
 				return PrepForUATPackageOrDeploy(GameName, ProjectDirectory, BuildPath + "/" + DecoratedGameName, "../../Engine", false);
 			}
-
+			else
+			{
+				// If it is requested, send the app bundle back to the platform executing these commands.
+				if (BuildConfiguration.bCopyAppBundleBackToDevice)
+				{
+					Log.TraceInformation("Copying binaries back to this device...");
+
+					IOSToolChain Toolchain = UEToolChain.GetPlatformToolChain(CPPTargetPlatform.IOS) as IOSToolChain;
+					bool bIsStaticLibrary = InTarget.OutputPath.EndsWith(".a");
+
+					try
+					{
+						string BinaryDir = Path.GetDirectoryName(InTarget.OutputPath) + "\\";
+						if (BinaryDir.EndsWith(InTarget.AppName + "\\Binaries\\IOS\\") && InTarget.Rules.Type != TargetRules.TargetType.Game)
+						{
+							BinaryDir = BinaryDir.Replace(InTarget.Rules.Type.ToString(), "Game");
+						}
+
+						string RemoteBinariesDir = Toolchain.ConvertPath(BinaryDir);
+						string LocalBinariesDir = BinaryDir;
+
+						// Get the app bundle's name
+						string AppFullName = InTarget.AppName;
+						if (InTarget.Configuration != UnrealTargetConfiguration.Development)
+						{
+							AppFullName += "-" + InTarget.Platform.ToString();
+							AppFullName += "-" + InTarget.Configuration.ToString();
+						}
+
+						List<string> NotBundledBinaries = new List<string>();
+						foreach (string BinaryPath in Toolchain.BuiltBinaries)
+						{
+							if (!BinaryPath.Contains("Dummy"))
+							{
+								RPCUtilHelper.CopyFile(Toolchain.ConvertPath(BinaryPath), BinaryPath, false);
+							}
+						}
+						Log.TraceInformation("Copied binaries successfully.");
+					}
+					catch (Exception)
+					{
+						Log.TraceInformation("Copying binaries back to this device failed.");
+					}
+				}
+			}
 			return true;
 		}
 
diff --git a/Engine/Source/Programs/UnrealBuildTool/ToolChain/RemoteToolChain.cs b/Engine/Source/Programs/UnrealBuildTool/ToolChain/RemoteToolChain.cs
--- a/Engine/Source/Programs/UnrealBuildTool/ToolChain/RemoteToolChain.cs
+++ b/Engine/Source/Programs/UnrealBuildTool/ToolChain/RemoteToolChain.cs
@@ -53,6 +53,7 @@ namespace UnrealBuildTool
             ErrorMessageTokens.Add("IPP ERROR");
             ErrorMessageTokens.Add("System.Net.Sockets.SocketException");
 
+			BranchDirectory = BranchDirectory.Replace("Engine\\Binaries\\DotNET", "");
 			BranchDirectory = BranchDirectory.Replace("Engine\\Source\\", "");
 			BranchDirectoryMac = BranchDirectory.Replace("\\", "/");
         }
