Fixing 'unreachable code detected' warrnings when compiling UBT

diff --git a/Engine/Source/Programs/UnrealBuildTool/Windows/UEBuildWindows.cs b/Engine/Source/Programs/UnrealBuildTool/Windows/UEBuildWindows.cs
--- a/Engine/Source/Programs/UnrealBuildTool/Windows/UEBuildWindows.cs
+++ b/Engine/Source/Programs/UnrealBuildTool/Windows/UEBuildWindows.cs
@@ -334,9 +334,9 @@ namespace UnrealBuildTool
 
 			// Win32 XP is only supported at this time.
 			// The XP compilation is currently only supported with a VS2013 build.
-            SupportWindowsXP = false 
-				&& InBuildTarget.Platform == UnrealTargetPlatform.Win32
-				&& Compiler == WindowsCompiler.VisualStudio2013;
+			SupportWindowsXP = false;
+				//&& InBuildTarget.Platform == UnrealTargetPlatform.Win32
+				//&& Compiler == WindowsCompiler.VisualStudio2013;
 
 			if( SupportWindowsXP )
 			{
