Added AppData location to UBT possible configuration locations.
The logic of data loading is also slightly changed. The global user config can now be found in two locations: in MyDocuments and AppData. If both of them exist the preferred is most recently changed one. If one of them exists it is chosen. If none exists then the template is created in AppData.

diff --git a/Engine/Source/Programs/UnrealBuildTool/System/ProjectFileGenerator.cs b/Engine/Source/Programs/UnrealBuildTool/System/ProjectFileGenerator.cs
--- a/Engine/Source/Programs/UnrealBuildTool/System/ProjectFileGenerator.cs
+++ b/Engine/Source/Programs/UnrealBuildTool/System/ProjectFileGenerator.cs
@@ -576,14 +576,17 @@ namespace UnrealBuildTool
 		{
 			foreach(var BuildConfigurationPath in XmlConfigLoader.ConfigLocationHierarchy)
 			{
-				var FilePath = Path.Combine(BuildConfigurationPath.FSLocation, "BuildConfiguration.xml");
-
-				if(!File.Exists(FilePath))
+				if(!BuildConfigurationPath.bExists)
 				{
 					continue;
 				}
 
-				EngineProject.AddAliasedFileToProject(new AliasedFile(FilePath, Path.Combine("Config", "UnrealBuildTool", BuildConfigurationPath.IDEFolderName)));
+				EngineProject.AddAliasedFileToProject(
+						new AliasedFile(
+							BuildConfigurationPath.FSLocation,
+							Path.Combine("Config", "UnrealBuildTool", BuildConfigurationPath.IDEFolderName)
+						)
+					);
 			}
 		}
 
diff --git a/Engine/Source/Programs/UnrealBuildTool/Utilities/XmlConfigLoader.cs b/Engine/Source/Programs/UnrealBuildTool/Utilities/XmlConfigLoader.cs
--- a/Engine/Source/Programs/UnrealBuildTool/Utilities/XmlConfigLoader.cs
+++ b/Engine/Source/Programs/UnrealBuildTool/Utilities/XmlConfigLoader.cs
@@ -115,20 +115,74 @@ namespace UnrealBuildTool
 		/// </summary>
 		public class XmlConfigLocation
 		{
-			// Possible location of the file system.
+			/// <summary>
+			/// Returns location of the BuildConfiguration.xml.
+			/// </summary>
+			/// <returns>Location of the BuildConfiguration.xml.</returns>
+			private static string GetConfigLocation(IEnumerable<string> PossibleLocations, out bool bExists)
+			{
+				if(PossibleLocations.Count() == 0)
+				{
+					throw new ArgumentException("Empty possible locations", "PossibleLocations");
+				}
+
+				const string ConfigXmlFileName = "BuildConfiguration.xml";
+
+				// Filter out non-existing
+				var ExistingLocations = new List<string>();
+
+				foreach(var PossibleLocation in PossibleLocations)
+				{
+					var FilePath = Path.Combine(PossibleLocation, ConfigXmlFileName);
+
+					if(File.Exists(FilePath))
+					{
+						ExistingLocations.Add(FilePath);
+					}
+				}
+
+				if(ExistingLocations.Count == 0)
+				{
+					bExists = false;
+					return Path.Combine(PossibleLocations.First(), ConfigXmlFileName);
+				}
+
+				bExists = true;
+
+				if(ExistingLocations.Count == 1)
+				{
+					return ExistingLocations.First();
+				}
+
+				// Choose most recently used from existing.
+				return ExistingLocations.OrderBy(Location => File.GetLastWriteTime(Location)).Last();
+			}
+
+			// Possible location of the config file in the file system.
 			public string FSLocation { get; private set; }
 
 			// IDE folder name that will contain this location if file will be found.
 			public string IDEFolderName { get; private set; }
 
-			// Tells if this is a user location.
-			public bool bIsUserLocation { get; private set; }
+			// Tells if UBT has to create a template config file if it does not exist in the location.
+			public bool bCreateIfDoesNotExist { get; private set; }
+
+			// Tells if config file exists in this location.
+			public bool bExists { get; private set; }
 
-			public XmlConfigLocation(string FSLocation, string IDEFolderName, bool bIsUserLocation = false)
+			public XmlConfigLocation(string[] FSLocations, string IDEFolderName, bool bCreateIfDoesNotExist = false)
 			{
-				this.FSLocation = FSLocation;
+				bool bExists;
+
+				this.FSLocation = GetConfigLocation(FSLocations, out bExists);
 				this.IDEFolderName = IDEFolderName;
-				this.bIsUserLocation = bIsUserLocation;
+				this.bCreateIfDoesNotExist = bCreateIfDoesNotExist;
+				this.bExists = bExists;
+			}
+
+			public XmlConfigLocation(string FSLocation, string IDEFolderName, bool bCreateIfDoesNotExist = false)
+				: this(new string[] { FSLocation }, IDEFolderName, bCreateIfDoesNotExist)
+			{
 			}
 		}
 
@@ -141,7 +195,8 @@ namespace UnrealBuildTool
 			 *		a. UE4/Engine/Programs/UnrealBuildTool
 			 *		b. UE4/Engine/Programs/NotForLicensees/UnrealBuildTool
 			 *		c. UE4/Engine/Saved/UnrealBuildTool
-			 *		d. My Documnets/Unreal Engine/UnrealBuildTool
+			 *		d. <AppData or My Documnets>/Unreal Engine/UnrealBuildTool -- the location is
+			 *		   chosen by existence and if both exist most recently used.
 			 *
 			 *	The UBT is looking for it in all four places in the given order and
 			 *	overrides already read data with the loaded ones, hence d. has the
@@ -155,7 +210,10 @@ namespace UnrealBuildTool
 				new XmlConfigLocation(Path.Combine(UE4EnginePath, "Programs", "UnrealBuildTool"), "Default"),
 				new XmlConfigLocation(Path.Combine(UE4EnginePath, "Programs", "NotForLicensees", "UnrealBuildTool"), "NotForLicensees"),
 				new XmlConfigLocation(Path.Combine(UE4EnginePath, "Saved", "UnrealBuildTool"), "User", true),
-				new XmlConfigLocation(Path.Combine(Environment.GetFolderPath(Environment.SpecialFolder.Personal), "Unreal Engine", "UnrealBuildTool"), "MyDocuments", true)
+				new XmlConfigLocation(new string[] {
+					Path.Combine(Environment.GetFolderPath(Environment.SpecialFolder.ApplicationData), "Unreal Engine", "UnrealBuildTool"),
+					Path.Combine(Environment.GetFolderPath(Environment.SpecialFolder.Personal), "Unreal Engine", "UnrealBuildTool")
+				}, "Global", true)
 			};
 		}
 
@@ -164,65 +222,51 @@ namespace UnrealBuildTool
 		/// </summary>
 		private static void LoadData()
 		{
-			const string ConfigXmlFileName = "BuildConfiguration.xml";
-
-			bool bFoundUserLocation = false;
-
 			foreach (var PossibleConfigLocation in ConfigLocationHierarchy)
 			{
-				var FilePath = Path.Combine(PossibleConfigLocation.FSLocation, ConfigXmlFileName);
-
-				if (File.Exists(FilePath))
+				if(!PossibleConfigLocation.bExists)
 				{
-					bFoundUserLocation |= PossibleConfigLocation.bIsUserLocation;
-
-					Load(FilePath);
+					continue;
 				}
-			}
 
-			if(!bFoundUserLocation)
-			{
-				CreateUserXmlConfigTemplate();
+				Load(PossibleConfigLocation.FSLocation);
 			}
-		}
 
-		/// <summary>
-		/// Creates template file in one of the user locations.
-		/// </summary>
-		private static void CreateUserXmlConfigTemplate()
-		{
-			const string ConfigXmlFileName = "BuildConfiguration.xml";
-
-			// If not found in any of user locations, try to create in one.
 			foreach (var PossibleConfigLocation in ConfigLocationHierarchy)
 			{
-				if (!PossibleConfigLocation.bIsUserLocation)
+				if(!PossibleConfigLocation.bCreateIfDoesNotExist || PossibleConfigLocation.bExists)
 				{
-					// Skip non-user locations.
 					continue;
 				}
 
-				try
-				{
-					Directory.CreateDirectory(PossibleConfigLocation.FSLocation);
+				CreateUserXmlConfigTemplate(PossibleConfigLocation.FSLocation);
+			}
+		}
 
-					var FilePath = Path.Combine(PossibleConfigLocation.FSLocation, ConfigXmlFileName);
 
-					const string TemplateContent =
-						"<Configuration>\n" +
-						"	<BuildConfiguration>\n" +
-						"	</BuildConfiguration>\n" +
-						"</Configuration>\n";
+		/// <summary>
+		/// Creates template file in one of the user locations if it not
+		/// </summary>
+		/// <param name="Location">Location of the xml config file to create.</param>
+		private static void CreateUserXmlConfigTemplate(string Location)
+		{
+			try
+			{
+				Directory.CreateDirectory(Path.GetDirectoryName(Location));
 
-					File.WriteAllText(FilePath, TemplateContent);
+				var FilePath = Path.Combine(Location);
 
-					// If everything went fine then break, cause one user file was already created.
-					break;
-				}
-				catch (Exception)
-				{
-					// Ignore quietly.
-				}
+				const string TemplateContent =
+					"<Configuration>\n" +
+					"	<BuildConfiguration>\n" +
+					"	</BuildConfiguration>\n" +
+					"</Configuration>\n";
+
+				File.WriteAllText(FilePath, TemplateContent);
+			}
+			catch (Exception)
+			{
+				// Ignore quietly.
 			}
 		}
 		/// <summary>
