Rocket build fixes

#change All headers in Classes, Public and Private are treated as candidates for UHT.
#change All candidate headers are scraped for the presence of a UCLASS, USTRUCT, UENUM, UINTERFACE and UDELEGATE at the start of a line and are only passed to UHT if one of these is present.
#fix Removed UENUMs from various games which were picked up as false positives by UBT but were never intended to be parsed by UHT.
#fix SocketSubsystem.h was including itself.

diff --git a/Engine/Source/Programs/UnrealBuildTool/Configuration/UEBuildModule.cs b/Engine/Source/Programs/UnrealBuildTool/Configuration/UEBuildModule.cs
--- a/Engine/Source/Programs/UnrealBuildTool/Configuration/UEBuildModule.cs
+++ b/Engine/Source/Programs/UnrealBuildTool/Configuration/UEBuildModule.cs
@@ -6,6 +6,7 @@ using System.Text;
 using System.Diagnostics;
 using System.IO;
 using System.Linq;
+using System.Text.RegularExpressions;
 using System.Xml;
 
 namespace UnrealBuildTool
@@ -1575,89 +1576,27 @@ namespace UnrealBuildTool
 			if (CachedModuleUHTInfo.HasValue)
 				return CachedModuleUHTInfo.Value;
 
-			// If this module has a "Classes" folder, we'll assume that all headers in that folder are public headers, whether they were
-			// included by #include statements or not.  This is because in the old system, we relied on generated *Classes.h mega-headers
-			// that auto-included all of the files in that folder, so programmers didn't have to bother including the files themselves.
-			// With the new system (files in outside of the "Classes" folders), you always need to include what you're using!  We can't follow
-			// include dependencies on generated "*Classes.h" files that may not even exist at the time we're gathering UObject headers.
+			var FolderSets = new KeyValuePair<string, HashSet<FileItem>>[] {
+				new KeyValuePair<string, HashSet<FileItem>>("Classes", _AllClassesHeaders),
+				new KeyValuePair<string, HashSet<FileItem>>("Public",  _PublicUObjectHeaders),
+				new KeyValuePair<string, HashSet<FileItem>>("Private", _PrivateUObjectHeaders)
+			};
+
+			foreach (var Folder in FolderSets)
 			{
-				var ClassesFolder = Path.Combine( this.ModuleDirectory, "Classes" );
+				var ClassesFolder = Path.Combine( this.ModuleDirectory, Folder.Key );
 				if( Directory.Exists( ClassesFolder ) )
 				{
 					var ClassesFiles = Directory.GetFiles( ClassesFolder, "*.h", SearchOption.AllDirectories );
 					foreach( var ClassHeader in ClassesFiles )
 					{
 						var UObjectHeaderFileItem = FileItem.GetExistingItemByPath( ClassHeader );
-						UObjectHeaderFileItem.HasUObjects = true;
-						_AllClassesHeaders.Add( UObjectHeaderFileItem );
-					}
-				}
-			}
-
-			// Keep track of this module's public and private UObject source files, so that we can pass those off to UHT if needed
-			string ModuleSourceFolder = Path.GetFullPath(this.ModuleDirectory);
-			foreach (var SourceFile in SourceFiles.CFiles.Concat(SourceFiles.CPPFiles).Concat(SourceFiles.CCFiles).Concat(SourceFiles.MMFiles))
-			{
-				// Will always be a cache hit (we did this earlier during Compile())
-				var IncludedFiles = CompileEnvironment.GetIncludeDependencies(SourceFile);
-
-				// Also check for intrinsic classes like "Object.h", which are special cases because they are never included in compiled code and exist only for UHT to parse
-
-				// Runtime/CoreUObject/Classes/Object.h
-				{
-					// @todo uht: In Classes folder right now
-					var IntrinsicFileItem = FileItem.GetExistingItemByPath(Path.Combine(ProjectFileGenerator.EngineRelativePath, "Source", "Runtime", "CoreUObject", "Classes", "Object.h"));
-					if (!IntrinsicFileItem.bExists)
-					{
-						throw new BuildException("Expecting " + IntrinsicFileItem.AbsolutePath + " to exist");
-					}
-					IntrinsicFileItem.HasUObjects = true;
-					IncludedFiles.Add(IntrinsicFileItem);
-				}
-
-				// Runtime/Engine/Classes/Model.h
-				{
-					// @todo uht: In Classes folder right now
-					var IntrinsicFileItem = FileItem.GetExistingItemByPath(Path.Combine(ProjectFileGenerator.EngineRelativePath, "Source", "Runtime", "Engine", "Classes", "Intrinsic", "Model.h"));
-					if (!IntrinsicFileItem.bExists)
-					{
-						throw new BuildException("Expecting " + IntrinsicFileItem.AbsolutePath + " to exist");
-					}
-					IntrinsicFileItem.HasUObjects = true;
-					IncludedFiles.Add(IntrinsicFileItem);
-				}
-
-				// @todo uht: Could check SourceFile.HasUObjects if we want to include .cpps with USTRUCT/UCLASSES here
-				foreach (var IncludedFile in IncludedFiles)
-				{
-					// Check if it's not from this module, so we don't need to worry about UObjects
-					if (!IncludedFile.HasUObjects || !IncludedFile.AbsolutePath.StartsWith(ModuleSourceFolder + Path.DirectorySeparatorChar))
-						continue;
-
-					// Is it private or public?
-					bool bIsPublic = false;
-
-					// Get the part of the path that is relative to the module source folder
-					var RelativeSourceFilePath = IncludedFile.AbsolutePath.Substring(ModuleSourceFolder.Length + 1);
-					if (RelativeSourceFilePath.StartsWith("Classes" + Path.DirectorySeparatorChar, StringComparison.InvariantCultureIgnoreCase))
-					{
-						// Code files under the legacy 'Classes' directory are always considered public
-						bIsPublic = true;
-					}
-					else if (RelativeSourceFilePath.StartsWith("Public" + Path.DirectorySeparatorChar, StringComparison.InvariantCultureIgnoreCase) ||
-								(RelativeSourceFilePath.IndexOf(Path.DirectorySeparatorChar + "Public" + Path.DirectorySeparatorChar, StringComparison.InvariantCultureIgnoreCase) != -1))
-					{
-						// Code is under a 'Public' subdirectory somewhere in the hierarchy
-						bIsPublic = true;
-					}
-
-					if (bIsPublic)
-					{
-						_PublicUObjectHeaders.Add(IncludedFile);
-					}
-					else
-					{
-						_PrivateUObjectHeaders.Add(IncludedFile);
+						var FileContents = Utils.ReadAllText(UObjectHeaderFileItem.AbsolutePath);
+						if (Regex.IsMatch(FileContents, "^\\s*U(CLASS|STRUCT|ENUM|INTERFACE|DELEGATE)\\b", RegexOptions.Multiline))
+						{
+							UObjectHeaderFileItem.HasUObjects = true;
+							Folder.Value.Add( UObjectHeaderFileItem );
+						}
 					}
 				}
 			}
diff --git a/Engine/Source/Runtime/Engine/Public/TimerManager.h b/Engine/Source/Runtime/Engine/Public/TimerManager.h
--- a/Engine/Source/Runtime/Engine/Public/TimerManager.h
+++ b/Engine/Source/Runtime/Engine/Public/TimerManager.h
@@ -81,7 +81,6 @@ struct FTimerUnifiedDelegate
 	}
 };
 
-UENUM()
 namespace ETimerStatus
 {
 	enum Type
diff --git a/Engine/Source/Runtime/Sockets/Private/HTML5/Device/SocketSubsystem.h b/Engine/Source/Runtime/Sockets/Private/HTML5/Device/SocketSubsystem.h
--- a/Engine/Source/Runtime/Sockets/Private/HTML5/Device/SocketSubsystem.h
+++ b/Engine/Source/Runtime/Sockets/Private/HTML5/Device/SocketSubsystem.h
@@ -2,7 +2,6 @@
 
 #pragma once
 
-#include "SocketSubsystem.h"
 #include "BSDSockets/SocketSubsystemBSD.h"
 #include "SocketSubsystemPackage.h"
 
