change more NoRedist references to NotForLicensees

diff --git a/Engine/Source/Programs/UnrealBuildTool/Configuration/UEBuildConfiguration.cs b/Engine/Source/Programs/UnrealBuildTool/Configuration/UEBuildConfiguration.cs
--- a/Engine/Source/Programs/UnrealBuildTool/Configuration/UEBuildConfiguration.cs
+++ b/Engine/Source/Programs/UnrealBuildTool/Configuration/UEBuildConfiguration.cs
@@ -147,18 +147,18 @@ namespace UnrealBuildTool
 			   && Directory.Exists(UEBuildConfiguration.UEThirdPartyDirectory + "Steamworks/" + SteamVersion) == true;
 
 			bCompileMcpOSS = bCompileMcpOSS
-			   && Directory.Exists("Runtime/Online/NoRedist/OnlineSubsystemMcp") == true;
+			   && Directory.Exists("Runtime/Online/NotForLicensees/OnlineSubsystemMcp") == true;
 
 
 			bCompileSimplygon = bCompileSimplygon
-				&& Directory.Exists(UEBuildConfiguration.UEThirdPartyDirectory + "NoRedist") == true
-				&& Directory.Exists(UEBuildConfiguration.UEThirdPartyDirectory + "NoRedist/Simplygon") == true
+				&& Directory.Exists(UEBuildConfiguration.UEThirdPartyDirectory + "NotForLicensees") == true
+                && Directory.Exists(UEBuildConfiguration.UEThirdPartyDirectory + "NotForLicensees/Simplygon") == true
 				&& Directory.Exists("Developer/SimplygonMeshReduction") == true
 				&& !(ProjectFileGenerator.bGenerateProjectFiles && ProjectFileGenerator.bGeneratingRocketProjectFiles);
 
 			bCompileSpeedTree = bCompileSpeedTree
-				&& Directory.Exists(UEBuildConfiguration.UEThirdPartyDirectory + "NoRedist") == true
-				&& Directory.Exists(UEBuildConfiguration.UEThirdPartyDirectory + "NoRedist/SpeedTree") == true
+                && Directory.Exists(UEBuildConfiguration.UEThirdPartyDirectory + "NotForLicensees") == true
+                && Directory.Exists(UEBuildConfiguration.UEThirdPartyDirectory + "NotForLicensees/SpeedTree") == true
 				&& !(ProjectFileGenerator.bGenerateProjectFiles && ProjectFileGenerator.bGeneratingRocketProjectFiles);
 		}
 
