Template build fixes.

#change Header files for UHT are now discovered in any location below the module folder.

diff --git a/Engine/Source/Programs/UnrealBuildTool/Configuration/UEBuildModule.cs b/Engine/Source/Programs/UnrealBuildTool/Configuration/UEBuildModule.cs
--- a/Engine/Source/Programs/UnrealBuildTool/Configuration/UEBuildModule.cs
+++ b/Engine/Source/Programs/UnrealBuildTool/Configuration/UEBuildModule.cs
@@ -1576,39 +1576,37 @@ namespace UnrealBuildTool
 			if (CachedModuleUHTInfo.HasValue)
 				return CachedModuleUHTInfo.Value;
 
-			var FolderSets = new KeyValuePair<string, HashSet<FileItem>>[] {
-				new KeyValuePair<string, HashSet<FileItem>>("Classes", _AllClassesHeaders),
-				new KeyValuePair<string, HashSet<FileItem>>("Public",  _PublicUObjectHeaders),
-				new KeyValuePair<string, HashSet<FileItem>>("Private", _PrivateUObjectHeaders)
-			};
+			var ClassesFolder = Path.Combine(this.ModuleDirectory, "Classes");
+			var PublicFolder  = Path.Combine(this.ModuleDirectory, "Public");
 
-			foreach (var Folder in FolderSets)
+			var ClassesFiles = Directory.GetFiles( this.ModuleDirectory, "*.h", SearchOption.AllDirectories );
+			foreach (var ClassHeader in ClassesFiles)
 			{
-				var ClassesFolder = Path.Combine( this.ModuleDirectory, Folder.Key );
-				if( Directory.Exists( ClassesFolder ) )
+				var UObjectHeaderFileItem = FileItem.GetExistingItemByPath( ClassHeader );
+				var FileContents = Utils.ReadAllText(UObjectHeaderFileItem.AbsolutePath);
+				if (!Regex.IsMatch(FileContents, "^\\s*U(CLASS|STRUCT|ENUM|INTERFACE|DELEGATE)\\b", RegexOptions.Multiline))
+					continue;
+
+				UObjectHeaderFileItem.HasUObjects = true;
+				if (UObjectHeaderFileItem.AbsolutePath.StartsWith(ClassesFolder))
 				{
-					var ClassesFiles = Directory.GetFiles( ClassesFolder, "*.h", SearchOption.AllDirectories );
-					foreach( var ClassHeader in ClassesFiles )
-					{
-						var UObjectHeaderFileItem = FileItem.GetExistingItemByPath( ClassHeader );
-						var FileContents = Utils.ReadAllText(UObjectHeaderFileItem.AbsolutePath);
-						if (Regex.IsMatch(FileContents, "^\\s*U(CLASS|STRUCT|ENUM|INTERFACE|DELEGATE)\\b", RegexOptions.Multiline))
-						{
-							UObjectHeaderFileItem.HasUObjects = true;
-							Folder.Value.Add( UObjectHeaderFileItem );
-						}
-					}
+					_AllClassesHeaders.Add(UObjectHeaderFileItem);
+				}
+				else if (UObjectHeaderFileItem.AbsolutePath.StartsWith(PublicFolder))
+				{
+					_PublicUObjectHeaders.Add(UObjectHeaderFileItem);
+				}
+				else
+				{
+					_PrivateUObjectHeaders.Add(UObjectHeaderFileItem);
 				}
 			}
 
-			var AllClassesHeadersList = _AllClassesHeaders.ToList();
-
 			CachedModuleUHTInfo = new UHTModuleInfo {
 				ModuleName                  = this.Name,
 				ModuleDirectory             = this.ModuleDirectory,
-				PublicUObjectClassesHeaders = AllClassesHeadersList,
-				// Don't report public classes headers twice.  This logic can go away if we deprecate Classes folders.
-				PublicUObjectHeaders        = _PublicUObjectHeaders.Where(File => !AllClassesHeadersList.Contains(File)).ToList(),
+				PublicUObjectClassesHeaders = _AllClassesHeaders    .ToList(),
+				PublicUObjectHeaders        = _PublicUObjectHeaders .ToList(),
 				PrivateUObjectHeaders       = _PrivateUObjectHeaders.ToList()
 			};
 
