#ue4 Setup PS4 for automated SDK switching.

diff --git a/Engine/Source/Programs/UnrealBuildTool/Android/UEBuildAndroid.cs b/Engine/Source/Programs/UnrealBuildTool/Android/UEBuildAndroid.cs
--- a/Engine/Source/Programs/UnrealBuildTool/Android/UEBuildAndroid.cs
+++ b/Engine/Source/Programs/UnrealBuildTool/Android/UEBuildAndroid.cs
@@ -30,14 +30,14 @@ namespace UnrealBuildTool
             return false;
         }
 
-		public override bool HasRequiredSDKsInstalled()
+		public override SDKStatus HasRequiredSDKsInstalled()
 		{
 			string NDKPath = Environment.GetEnvironmentVariable("NDKROOT");
 
 			// we don't have an NDKROOT specified
 			if (String.IsNullOrEmpty(NDKPath))
 			{
-				return false;
+				return SDKStatus.Invalid;
 			}
 
 			NDKPath = NDKPath.Replace("\"", "");
@@ -45,15 +45,15 @@ namespace UnrealBuildTool
 			// can't find llvm-3.3 or llvm-3.1 in the toolchains
 			if (!Directory.Exists(Path.Combine(NDKPath, @"toolchains\llvm-3.3")) && !Directory.Exists(Path.Combine(NDKPath, @"toolchains\llvm-3.1")))
 			{
-				return false;
+				return SDKStatus.Invalid;
 			}
 
-			return true;
+			return SDKStatus.Valid;
 		}
 
 		public override void RegisterBuildPlatform()
 		{
-			if ((ProjectFileGenerator.bGenerateProjectFiles == true) || (HasRequiredSDKsInstalled() == true))
+			if ((ProjectFileGenerator.bGenerateProjectFiles == true) || (HasRequiredSDKsInstalled() == SDKStatus.Valid))
 			{
 				bool bRegisterBuildPlatform = true;
 
diff --git a/Engine/Source/Programs/UnrealBuildTool/Configuration/UEBuildPlatform.cs b/Engine/Source/Programs/UnrealBuildTool/Configuration/UEBuildPlatform.cs
--- a/Engine/Source/Programs/UnrealBuildTool/Configuration/UEBuildPlatform.cs
+++ b/Engine/Source/Programs/UnrealBuildTool/Configuration/UEBuildPlatform.cs
@@ -28,6 +28,15 @@ namespace UnrealBuildTool
             Uninstall
         };
 
+		public enum SDKStatus
+		{
+			Valid,			// Desired SDK is installed and set up.
+			Invalid,		// Could not find the desired SDK, SDK setup failed, etc.
+			Unknown,		// Not enough info to determine SDK state. 
+		};
+
+		public static readonly string SDKRootEnvVar = "UE_SDKS_ROOT";
+
         /**
          * Attempt to convert a string to an UnrealTargetPlatform enum entry
          * 
@@ -203,15 +212,21 @@ namespace UnrealBuildTool
          * 
          * @return Valid SDK string
          */
-        public virtual string GetPathToPlatformSDKs()
+        public virtual string GetPathToPlatformSDKs(out bool SDKRootSet)
         {
-            string SDKRoot = Environment.GetEnvironmentVariable("UE_SDKS_ROOT");
-            if (SDKRoot != null && SDKRoot != "")
+            string SDKPath = "";
+            SDKRootSet = false;
+
+            string SDKRoot = Environment.GetEnvironmentVariable(SDKRootEnvVar);
+            if (SDKRoot != null)
             {
-                return Path.Combine(SDKRoot, "Host" + ExternalExecution.GetRuntimePlatform(), GetSDKTargetPlatformName());
+                SDKRootSet = true;
+                if (SDKRoot != "")
+                {
+                    SDKPath = Path.Combine(SDKRoot, "Host" + ExternalExecution.GetRuntimePlatform(), GetSDKTargetPlatformName());
+                }
             }
-
-            return "";
+            return SDKPath;
         }
 
         /**
@@ -332,6 +347,25 @@ namespace UnrealBuildTool
             return bHookCanBeNonExistent;
         }
 
+        /**
+         * Returns the delimiter used to separate paths in the PATH environment variable for the platform we are executing on.
+         */
+        public String GetPathVarDelimiter()
+        {
+            switch (ExternalExecution.GetRuntimePlatform())
+            {
+                case UnrealTargetPlatform.Linux:
+                case UnrealTargetPlatform.Mac:
+                    return ":";
+                case UnrealTargetPlatform.Win32:
+                case UnrealTargetPlatform.Win64:
+                    return ";";
+                default:
+                    Log.TraceWarning("PATH var delimiter unknown for platform " + ExternalExecution.GetRuntimePlatform().ToString() + " using ';'");
+                    return ";";
+            }
+        }
+
         /**
          * Loads environment variables from SDK
          * 
@@ -347,14 +381,19 @@ namespace UnrealBuildTool
             {
                 using (StreamReader Reader = new StreamReader(EnvVarFile))
                 {
-                    for (; ; )
+                    List<string> PathAdds = new List<string>();
+                    List<string> PathRemoves = new List<string>();
+
+                    List<string> EnvVarNames = new List<string>();
+                    List<string> EnvVarValues = new List<string>();
+                    for (; ;)
                     {
                         string VariableString = Reader.ReadLine();
                         if (VariableString == null)
                         {
-                            return true;
+                            break;
                         }
-
+						
                         string[] Parts = VariableString.Split('=');
                         if (Parts.Length != 2)
                         {
@@ -363,12 +402,93 @@ namespace UnrealBuildTool
                             return false;
                         }
 
+                        if (String.Compare(Parts[0], "strippath", true) == 0)
+                        {
+                            PathRemoves.Add(Parts[1]);
+                        }
+                        else if (String.Compare(Parts[0], "addpath", true) == 0)
+                        {
+                            PathAdds.Add(Parts[1]);
+                        }
+                        else
+                        {
+                            // convenience for setup.bat writers.  Trim any accidental whitespace from var names/values.
+                            EnvVarNames.Add(Parts[0].Trim());
+                            EnvVarValues.Add(Parts[1].Trim());
+                        }
+                    }
+
+                    // don't actually set anything until we successfully validate and read all values in.
+                    // we don't want to set a few vars, return a failure, and then have a platform try to
+                    // build against a manually installed SDK with half-set env vars.
+                    for (int i = 0; i < EnvVarNames.Count; ++i)
+                    {
+                        string EnvVarName = EnvVarNames[i];
+                        string EnvVarValue = EnvVarValues[i];
+                        if (String.Compare(EnvVarName, "strippath", true) == 0)
                         if (BuildConfiguration.bPrintDebugInfo)
                         {
-                            Console.WriteLine("Setting variable '{0}' to '{1}'", Parts[0], Parts[1]);
+                            Console.WriteLine("Setting variable '{0}' to '{1}'", EnvVarName, EnvVarValue);
+                        }
+                        Environment.SetEnvironmentVariable(EnvVarName, EnvVarValue);
+                    }
+                       
+                    
+                    // actually perform the PATH stripping / adding.
+                    String OrigPathVar = Environment.GetEnvironmentVariable("PATH");
+                    String PathDelimiter = GetPathVarDelimiter();
+                    String[] PathVars = OrigPathVar.Split(PathDelimiter.ToCharArray());
+
+                    List<String> ModifiedPathVars = new List<string>();
+                    ModifiedPathVars.AddRange(PathVars);
+
+                    // perform removes first, in case they overlap with any adds.
+                    foreach (String PathRemove in PathRemoves)
+                    {
+                        foreach (String PathVar in PathVars)
+                        {
+                            if (PathVar.IndexOf(PathRemove, StringComparison.OrdinalIgnoreCase) >= 0)
+                            {
+                                if (BuildConfiguration.bPrintDebugInfo)
+                                {
+                                    Console.WriteLine("Removing Path: '{0}'", PathVar);
+                                }
+                                ModifiedPathVars.Remove(PathVar);
+                            }
+                        }
+                    }
+
+                    // remove all the of ADDs so that if this function is executed multiple times, the paths will be guarateed to be in the same order after each run.
+                    // If we did not do this, a 'remove' that matched some, but not all, of our 'adds' would cause the order to change.
+                    foreach (String PathAdd in PathAdds)
+                    {
+                        foreach (String PathVar in PathVars)
+                        {
+                            if (String.Compare(PathAdd, PathVar, true) == 0)
+                            {
+                                if (BuildConfiguration.bPrintDebugInfo)
+                                {
+                                    Console.WriteLine("Removing Path: '{0}'", PathVar);
+                                }
+                                ModifiedPathVars.Remove(PathVar);
+                            }
                         }
-                        Environment.SetEnvironmentVariable(Parts[0], Parts[1]);
                     }
+
+                    // perform adds, but don't add duplicates
+                    foreach (String PathAdd in PathAdds)
+                    {
+                        if (!ModifiedPathVars.Contains(PathAdd))
+                        {
+                            Console.WriteLine("Adding Path: '{0}'", PathAdd);
+                            ModifiedPathVars.Add(PathAdd);
+                        }
+                    }
+
+                    String ModifiedPath = String.Join(PathDelimiter, ModifiedPathVars);
+                    Environment.SetEnvironmentVariable("PATH", ModifiedPath);
+
+                    return true;
                 }
             }
 
@@ -378,45 +498,57 @@ namespace UnrealBuildTool
         /**
          *	Whether the required external SDKs are installed for this platform
          */
-		public virtual bool HasRequiredSDKsInstalled()
-		{
+		public virtual SDKStatus HasRequiredSDKsInstalled()
+		{			
             if (PlatformSupportsSDKSwitching())
-            {
+            {				
                 // if we don't have a path to platform SDK's, take no action
-                string PlatformSDKRoot = GetPathToPlatformSDKs();
-                if (PlatformSDKRoot != "")
-                {
-                    string CurrentSDKString;
-                    if (!GetCurrentlyInstalledSDKString(PlatformSDKRoot, out CurrentSDKString) || CurrentSDKString != GetRequiredSDKString())
+                bool bSDKPathSet;
+                string PlatformSDKRoot = GetPathToPlatformSDKs(out bSDKPathSet);				
+                if (bSDKPathSet)
+                {					
+                    if (PlatformSDKRoot != "")
                     {
-                        // switch over (note that version string can be empty)
-                        if (!RunSDKHooks(PlatformSDKRoot, CurrentSDKString, SDKHookType.Uninstall))
+                        string CurrentSDKString;
+                        if (!GetCurrentlyInstalledSDKString(PlatformSDKRoot, out CurrentSDKString) || CurrentSDKString != GetRequiredSDKString())
                         {
-                            Console.WriteLine("Failed to uninstall currently installed SDK {0}", CurrentSDKString);
-                            return false;
+                        
+                            // switch over (note that version string can be empty)
+                            if (!RunSDKHooks(PlatformSDKRoot, CurrentSDKString, SDKHookType.Uninstall))
+                            {
+                                Console.WriteLine("Failed to uninstall currently installed SDK {0}", CurrentSDKString);
+                                return SDKStatus.Invalid;
+                            }
+                            // delete Manifest file to avoid multiple uninstalls
+                            SetCurrentlyInstalledSDKString(PlatformSDKRoot, "");
+
+                            if (!RunSDKHooks(PlatformSDKRoot, GetRequiredSDKString(), SDKHookType.Install, false))
+                            {								
+                                Console.WriteLine("Failed to install required SDK {0}.  Attemping to uninstall", GetRequiredSDKString());
+                                RunSDKHooks(PlatformSDKRoot, GetRequiredSDKString(), SDKHookType.Uninstall, false);
+
+                                return SDKStatus.Invalid;
+                            }
+                            SetCurrentlyInstalledSDKString(PlatformSDKRoot, GetRequiredSDKString());
                         }
-                        // delete Manifest file to avoid multiple uninstalls
-                        SetCurrentlyInstalledSDKString(PlatformSDKRoot, "");
 
-                        if (!RunSDKHooks(PlatformSDKRoot, GetRequiredSDKString(), SDKHookType.Install, false))
+                        // load environment variables from current SDK
+                        if (!SetupEnvironmentFromSDK(PlatformSDKRoot, GetRequiredSDKString()))
                         {
-                            Console.WriteLine("Failed to install required SDK {0}", GetRequiredSDKString());
-                            return false;
-                        }
-                        SetCurrentlyInstalledSDKString(PlatformSDKRoot, GetRequiredSDKString());
+                            Console.WriteLine("Failed to load environment from required SDK {0}", GetRequiredSDKString());
+                            return SDKStatus.Invalid;
+                        }                        
+                        return SDKStatus.Valid;
                     }
-
-                    // load environment variables from current SDK
-                    if (!SetupEnvironmentFromSDK(PlatformSDKRoot, GetRequiredSDKString()))
+                    else
                     {
-                        Console.WriteLine("Failed to load environment from required SDK {0}", GetRequiredSDKString());
-                        return false;
+                        Console.WriteLine("Failed to install required SDK {0}. {1} is blank.", GetRequiredSDKString(), SDKRootEnvVar);
+                        return SDKStatus.Invalid;
                     }
                 }
             }
-
-			return true;
-		}
+            return SDKStatus.Unknown;
+        }
 
         /**
          *	If this platform can be compiled with XGE
diff --git a/Engine/Source/Programs/UnrealBuildTool/HTML5/UEBuildHTML5.cs b/Engine/Source/Programs/UnrealBuildTool/HTML5/UEBuildHTML5.cs
--- a/Engine/Source/Programs/UnrealBuildTool/HTML5/UEBuildHTML5.cs
+++ b/Engine/Source/Programs/UnrealBuildTool/HTML5/UEBuildHTML5.cs
@@ -29,10 +29,10 @@ namespace UnrealBuildTool
 		/**
 		 *	Whether the required external SDKs are installed for this platform
 		 */
-		public override bool HasRequiredSDKsInstalled()
+		public override SDKStatus HasRequiredSDKsInstalled()
 		{
 			string BaseSDKPath = Environment.GetEnvironmentVariable("EMSCRIPTEN");
-			return string.IsNullOrEmpty(BaseSDKPath) == false;
+			return (string.IsNullOrEmpty(BaseSDKPath) == false) ? SDKStatus.Valid : SDKStatus.Invalid;
 		}
 
         public override bool CanUseXGE()
@@ -52,7 +52,7 @@ namespace UnrealBuildTool
 			}
 
 			// Make sure the SDK is installed
-			if ((ProjectFileGenerator.bGenerateProjectFiles == true) || (HasRequiredSDKsInstalled() == true))
+			if ((ProjectFileGenerator.bGenerateProjectFiles == true) || (HasRequiredSDKsInstalled() == SDKStatus.Valid))
 			{
 				bool bRegisterBuildPlatform = true;
 
diff --git a/Engine/Source/Programs/UnrealBuildTool/IOS/UEBuildIOS.cs b/Engine/Source/Programs/UnrealBuildTool/IOS/UEBuildIOS.cs
--- a/Engine/Source/Programs/UnrealBuildTool/IOS/UEBuildIOS.cs
+++ b/Engine/Source/Programs/UnrealBuildTool/IOS/UEBuildIOS.cs
@@ -20,6 +20,11 @@ namespace UnrealBuildTool
 			return IOSArchitecture;
 		}
 
+		public override SDKStatus HasRequiredSDKsInstalled()
+		{
+			return SDKStatus.Valid;
+		}
+
 		/**
 		 *	Register the platform with the UEBuildPlatform class
 		 */
diff --git a/Engine/Source/Programs/UnrealBuildTool/Linux/UEBuildLinux.cs b/Engine/Source/Programs/UnrealBuildTool/Linux/UEBuildLinux.cs
--- a/Engine/Source/Programs/UnrealBuildTool/Linux/UEBuildLinux.cs
+++ b/Engine/Source/Programs/UnrealBuildTool/Linux/UEBuildLinux.cs
@@ -49,33 +49,33 @@ namespace UnrealBuildTool.Linux
         /**
          *	Whether the required external SDKs are installed for this platform
          */
-        public override bool HasRequiredSDKsInstalled()
+        public override SDKStatus HasRequiredSDKsInstalled()
         {
             if (ExternalExecution.GetRuntimePlatform() == UnrealTargetPlatform.Linux)
             {
-                return true;
+                return SDKStatus.Valid;
             }
 
             // attempt to switch SDKs
-            if (!base.HasRequiredSDKsInstalled())
+            if (base.HasRequiredSDKsInstalled() == SDKStatus.Valid)
             {
-                return false;
+                return SDKStatus.Valid;
             }
 
             string BaseLinuxPath = Environment.GetEnvironmentVariable("LINUX_ROOT");
 
             // we don't have an LINUX_ROOT specified
             if (String.IsNullOrEmpty(BaseLinuxPath))
-                return false;
+                return SDKStatus.Invalid;
 
             // paths to our toolchains
             BaseLinuxPath = BaseLinuxPath.Replace("\"", "");
             string ClangPath = Path.Combine(BaseLinuxPath, @"bin\Clang++.exe");
             
             if (File.Exists(ClangPath))
-                return true;
+                return SDKStatus.Valid;
 
-            return false;
+            return SDKStatus.Invalid;
         }
 
         /**
@@ -89,7 +89,7 @@ namespace UnrealBuildTool.Linux
 				return;
 			}
 
-            if ((ProjectFileGenerator.bGenerateProjectFiles == true) || (HasRequiredSDKsInstalled() == true))
+			if ((ProjectFileGenerator.bGenerateProjectFiles == true) || (HasRequiredSDKsInstalled() == UEBuildPlatform.SDKStatus.Valid))
             {
                 bool bRegisterBuildPlatform = true;
 
diff --git a/Engine/Source/Programs/UnrealBuildTool/Mac/UEBuildMac.cs b/Engine/Source/Programs/UnrealBuildTool/Mac/UEBuildMac.cs
--- a/Engine/Source/Programs/UnrealBuildTool/Mac/UEBuildMac.cs
+++ b/Engine/Source/Programs/UnrealBuildTool/Mac/UEBuildMac.cs
@@ -16,6 +16,11 @@ namespace UnrealBuildTool
             return false;
         }
 
+		public override SDKStatus HasRequiredSDKsInstalled()
+		{
+			return SDKStatus.Valid;
+		}
+
 		/**
 		 *	Register the platform with the UEBuildPlatform class
 		 */
diff --git a/Engine/Source/Programs/UnrealBuildTool/System/VCProject.cs b/Engine/Source/Programs/UnrealBuildTool/System/VCProject.cs
--- a/Engine/Source/Programs/UnrealBuildTool/System/VCProject.cs
+++ b/Engine/Source/Programs/UnrealBuildTool/System/VCProject.cs
@@ -183,7 +183,7 @@ namespace UnrealBuildTool
 				return false;
 			}
 
-			if (BuildPlatform.HasRequiredSDKsInstalled() == false)
+			if (BuildPlatform.HasRequiredSDKsInstalled() != UEBuildPlatform.SDKStatus.Valid)
 			{
 				return false;
 			}
@@ -416,7 +416,7 @@ namespace UnrealBuildTool
 							continue;
 						}
 						UEBuildPlatform BuildPlatform = UEBuildPlatform.GetBuildPlatform(Platform, true);
-						if ((BuildPlatform != null) && (BuildPlatform.HasRequiredSDKsInstalled() == true))
+						if ((BuildPlatform != null) && (BuildPlatform.HasRequiredSDKsInstalled() == UEBuildPlatform.SDKStatus.Valid))
 						{
 							// Now go through all of the target types for this project
 							if( ProjectTargets.Count == 0 )
diff --git a/Engine/Source/Programs/UnrealBuildTool/System/VCProjectFileGenerator.cs b/Engine/Source/Programs/UnrealBuildTool/System/VCProjectFileGenerator.cs
--- a/Engine/Source/Programs/UnrealBuildTool/System/VCProjectFileGenerator.cs
+++ b/Engine/Source/Programs/UnrealBuildTool/System/VCProjectFileGenerator.cs
@@ -181,12 +181,12 @@ namespace UnrealBuildTool
 			// version of Visual Studio.  This is their chance to override that.
 			foreach( var SupportedPlatform in SupportedPlatforms )
 			{
-				UEBuildPlatform BuildPlatform = UEBuildPlatform.GetBuildPlatform( SupportedPlatform, true );
+				UEBuildPlatform BuildPlatform = UEBuildPlatform.GetBuildPlatform( SupportedPlatform, true );				
 				if( BuildPlatform != null )
 				{
 					// Don't worry about platforms that we're missing SDKs for
-					if( BuildPlatform.HasRequiredSDKsInstalled() )
-					{
+					if (BuildPlatform.HasRequiredSDKsInstalled() == UEBuildPlatform.SDKStatus.Valid)
+					{						
 						// Make sure Visual Studio 2013 project files will work...
 						if( ProjectFileFormat == VCProjectFileFormat.VisualStudio2013 )
 						{
@@ -208,7 +208,7 @@ namespace UnrealBuildTool
 								}
 							}
 						}
-					}
+					}					
 				}
 			}
 		}
diff --git a/Engine/Source/Programs/UnrealBuildTool/WinRT/UEBuildWinRT.cs b/Engine/Source/Programs/UnrealBuildTool/WinRT/UEBuildWinRT.cs
--- a/Engine/Source/Programs/UnrealBuildTool/WinRT/UEBuildWinRT.cs
+++ b/Engine/Source/Programs/UnrealBuildTool/WinRT/UEBuildWinRT.cs
@@ -42,9 +42,9 @@ namespace UnrealBuildTool
 		/**
 		 *	Whether the required external SDKs are installed for this platform
 		 */
-		public override bool HasRequiredSDKsInstalled()
+		public override SDKStatus HasRequiredSDKsInstalled()
 		{
-			return IsVisualStudioInstalled();
+			return IsVisualStudioInstalled() ? SDKStatus.Valid : SDKStatus.Invalid;
 		}
 
 		/**
diff --git a/Engine/Source/Programs/UnrealBuildTool/Windows/UEBuildWindows.cs b/Engine/Source/Programs/UnrealBuildTool/Windows/UEBuildWindows.cs
--- a/Engine/Source/Programs/UnrealBuildTool/Windows/UEBuildWindows.cs
+++ b/Engine/Source/Programs/UnrealBuildTool/Windows/UEBuildWindows.cs
@@ -53,6 +53,11 @@ namespace UnrealBuildTool
 			}
 		}
 
+		public override SDKStatus HasRequiredSDKsInstalled()
+		{
+			return SDKStatus.Valid;
+		}
+
 		/**
 		 * Returns VisualStudio common tools path for current compiler.
 		 * 
