Add script setup version for SDK switching to allow easier setup script fixes.

diff --git a/Engine/Source/Programs/UnrealBuildTool/Configuration/UEBuildPlatform.cs b/Engine/Source/Programs/UnrealBuildTool/Configuration/UEBuildPlatform.cs
--- a/Engine/Source/Programs/UnrealBuildTool/Configuration/UEBuildPlatform.cs
+++ b/Engine/Source/Programs/UnrealBuildTool/Configuration/UEBuildPlatform.cs
@@ -18,6 +18,9 @@ namespace UnrealBuildTool
         /** Name of the file that holds currently install SDK version string */
         static string CurrentlyInstalledSDKStringManifest = "CurrentlyInstalled.txt";
 
+        /** name of the file that holds the last succesfully run SDK setup script version */
+        static string LastRunScriptVersionManifest = "CurrentlyInstalled.Version.txt";
+
         /** Name of the file that holds environment variables of current SDK */
         static string SDKEnvironmentVarsFile = "OutputEnvVars.txt";
 
@@ -207,6 +210,15 @@ namespace UnrealBuildTool
             return "";
         }
 
+        /**
+        * Gets the version number of the SDK setup script itself.  The version in the base should ALWAYS be 1.0.  If you need to force a rebuild for a given platform, override this for the given platform.
+        * @return Setup script version
+        */
+        public virtual String GetRequiredScriptVersionString()
+        {
+            return "1.0";
+        }
+
         /** 
          * Returns path to platform SDKs
          * 
@@ -260,6 +272,38 @@ namespace UnrealBuildTool
             return false;
         }
 
+        /**
+         * Gets the version of the last successfully run setup script.
+         * 
+         * @param PlatformSDKRoot absolute path to platform SDK root
+         * @param OutLastRunScriptVersion version string
+         * 
+         * @return true if was able to read it
+         */
+        public virtual bool GetLastRunScriptVersionString(string PlatformSDKRoot, out string OutLastRunScriptVersion)
+        {
+            if (Directory.Exists(PlatformSDKRoot))
+            {
+                string VersionFilename = Path.Combine(PlatformSDKRoot, LastRunScriptVersionManifest);
+                if (File.Exists(VersionFilename))
+                {
+                    using (StreamReader Reader = new StreamReader(VersionFilename))
+                    {
+                        string Version = Reader.ReadLine();
+                        if (Version != null)
+                        {
+                            OutLastRunScriptVersion = Version;
+                            return true;
+                        }
+                    }
+                }
+            }
+
+            OutLastRunScriptVersion = "";
+            return false;
+        }
+       
+
         /**
          * Sets currently installed version
          * 
@@ -288,6 +332,26 @@ namespace UnrealBuildTool
             return false;
         }
 
+        public virtual bool SetLastRunScriptVersion(string PlatformSDKRoot, string LastRunScriptVersion)
+        {
+            if (Directory.Exists(PlatformSDKRoot))
+            {
+                string VersionFilename = Path.Combine(PlatformSDKRoot, LastRunScriptVersionManifest);
+                if (File.Exists(VersionFilename))
+                {
+                    File.Delete(VersionFilename);
+                }
+
+                using (StreamWriter Writer = File.CreateText(VersionFilename))
+                {
+                    Writer.WriteLine(LastRunScriptVersion);
+                    return true;
+                }
+            }
+
+            return false;
+        }
+
         /**
          * Returns Hook names as needed by the platform
          * (e.g. can be overriden with custom executables or scripts)
@@ -509,8 +573,16 @@ namespace UnrealBuildTool
                 {					
                     if (PlatformSDKRoot != "")
                     {
+                        // check script version so script fixes can be propagated without touching every build machine's CurrentlyInstalled file manually.
+                        bool bScriptVersionMatches = false;
+                        string CurrentScriptVersionString;
+                        if (GetLastRunScriptVersionString(PlatformSDKRoot, out CurrentScriptVersionString) && CurrentScriptVersionString == GetRequiredScriptVersionString())
+                        {
+                            bScriptVersionMatches = true;
+                        }
+
                         string CurrentSDKString;
-                        if (!GetCurrentlyInstalledSDKString(PlatformSDKRoot, out CurrentSDKString) || CurrentSDKString != GetRequiredSDKString())
+                        if (!GetCurrentlyInstalledSDKString(PlatformSDKRoot, out CurrentSDKString) || CurrentSDKString != GetRequiredSDKString() || !bScriptVersionMatches)
                         {
                         
                             // switch over (note that version string can be empty)
@@ -521,6 +593,7 @@ namespace UnrealBuildTool
                             }
                             // delete Manifest file to avoid multiple uninstalls
                             SetCurrentlyInstalledSDKString(PlatformSDKRoot, "");
+                            SetLastRunScriptVersion(PlatformSDKRoot, "");
 
                             if (!RunSDKHooks(PlatformSDKRoot, GetRequiredSDKString(), SDKHookType.Install, false))
                             {								
@@ -530,6 +603,7 @@ namespace UnrealBuildTool
                                 return SDKStatus.Invalid;
                             }
                             SetCurrentlyInstalledSDKString(PlatformSDKRoot, GetRequiredSDKString());
+                            SetLastRunScriptVersion(PlatformSDKRoot, GetRequiredScriptVersionString());
                         }
 
                         // load environment variables from current SDK
