[GitHub] 60 : Simplify GetFilenameWithoutAnyExtensions
#ttp 333297

diff --git a/Engine/Source/Programs/UnrealBuildTool/System/Utils.cs b/Engine/Source/Programs/UnrealBuildTool/System/Utils.cs
--- a/Engine/Source/Programs/UnrealBuildTool/System/Utils.cs
+++ b/Engine/Source/Programs/UnrealBuildTool/System/Utils.cs
@@ -707,12 +707,17 @@ namespace UnrealBuildTool
 		/// <returns>Clean filename.</returns>
 		public static string GetFilenameWithoutAnyExtensions(string Filename)
 		{
-			do
+			Filename = Path.GetFileName(Filename);
+
+			var DotIndex = Filename.IndexOf('.');
+			if (DotIndex == -1)
+			{
+				return Filename; // No need to copy string
+			}
+			else
 			{
-				Filename = Path.GetFileNameWithoutExtension(Filename);
+				return Filename.Substring(0, DotIndex);
 			}
-			while (Filename.IndexOf('.') >= 0);
-			return Filename;
 		}
 
 		/// <summary>
