[GitHub] 71 : Make the implementation of UEBuildModule truly overridable
#ttp 333474

diff --git a/Engine/Source/Programs/UnrealBuildTool/Configuration/UEBuildModule.cs b/Engine/Source/Programs/UnrealBuildTool/Configuration/UEBuildModule.cs
--- a/Engine/Source/Programs/UnrealBuildTool/Configuration/UEBuildModule.cs
+++ b/Engine/Source/Programs/UnrealBuildTool/Configuration/UEBuildModule.cs
@@ -47,16 +47,16 @@ namespace UnrealBuildTool
 
 		/** Whether this module is included in the current target.  Only set after UEBuildBinary.BindModules is called. */
 		public bool bIncludedInTarget = false;
-		
-		List<string> PublicDefinitions;
-		List<string> PublicIncludePaths;
-		List<string> PrivateIncludePaths;
-		List<string> PublicSystemIncludePaths;
-		List<string> PublicLibraryPaths;
-		List<string> PublicAdditionalLibraries;
-		List<string> PublicFrameworks;
-		List<string> PublicAdditionalFrameworks;
-		List<string> PublicAdditionalShadowFiles;
+
+		protected readonly List<string> PublicDefinitions;
+		protected readonly List<string> PublicIncludePaths;
+		protected readonly List<string> PrivateIncludePaths;
+		protected readonly List<string> PublicSystemIncludePaths;
+		protected readonly List<string> PublicLibraryPaths;
+		protected readonly List<string> PublicAdditionalLibraries;
+		protected readonly List<string> PublicFrameworks;
+		protected readonly List<string> PublicAdditionalFrameworks;
+		protected readonly List<string> PublicAdditionalShadowFiles;
 
 		/** Names of modules with header files that this module's public interface needs access to. */
 		protected List<string> PublicIncludePathModuleNames;
@@ -138,13 +138,13 @@ namespace UnrealBuildTool
 		}
 
 		/** Adds a public definition */
-		public void AddPublicDefinition(string Definition)
+		public virtual void AddPublicDefinition(string Definition)
 		{
 			PublicDefinitions.Add(Definition);
 		}
 
 		/** Adds a public include path. */
-		public void AddPublicIncludePath(string PathName, bool bCheckForDuplicates = true)
+		public virtual void AddPublicIncludePath(string PathName, bool bCheckForDuplicates = true)
 		{
 			if (bCheckForDuplicates == true)
 			{
@@ -157,7 +157,7 @@ namespace UnrealBuildTool
 		}
 
 		/** Adds a public library path */
-		public void AddPublicLibraryPath(string PathName, bool bCheckForDuplicates = true)
+		public virtual void AddPublicLibraryPath(string PathName, bool bCheckForDuplicates = true)
 		{
 			if (bCheckForDuplicates == true)
 			{
@@ -170,7 +170,7 @@ namespace UnrealBuildTool
 		}
 
 		/** Adds a public additional library */
-		public void AddPublicAdditionalLibrary(string LibraryName, bool bCheckForDuplicates = true)
+		public virtual void AddPublicAdditionalLibrary(string LibraryName, bool bCheckForDuplicates = true)
 		{
 			if (bCheckForDuplicates == true)
 			{
@@ -183,7 +183,7 @@ namespace UnrealBuildTool
 		}
 
 		/** Adds a public framework */
-		public void AddPublicFramework(string LibraryName, bool bCheckForDuplicates = true)
+		public virtual void AddPublicFramework(string LibraryName, bool bCheckForDuplicates = true)
 		{
 			if (bCheckForDuplicates == true)
 			{
@@ -196,7 +196,7 @@ namespace UnrealBuildTool
 		}
 
 		/** Adds a public additional framework */
-		public void AddPublicAdditionalFramework(string LibraryName, bool bCheckForDuplicates = true)
+		public virtual void AddPublicAdditionalFramework(string LibraryName, bool bCheckForDuplicates = true)
 		{
 			if (bCheckForDuplicates == true)
 			{
@@ -209,13 +209,13 @@ namespace UnrealBuildTool
 		}
 
 		/** Adds a file that will be synced to a remote machine if a remote build is being executed */
-		public void AddPublicAdditionalShadowFile(string FileName)
+		public virtual void AddPublicAdditionalShadowFile(string FileName)
 		{
 			PublicAdditionalShadowFiles.Add(FileName);
 		}
 
 		/** Adds a public include path module. */
-		public void AddPublicIncludePathModule(string ModuleName, bool bCheckForDuplicates = true)
+		public virtual void AddPublicIncludePathModule(string ModuleName, bool bCheckForDuplicates = true)
 		{
 			if (bCheckForDuplicates == true)
 			{
@@ -228,7 +228,7 @@ namespace UnrealBuildTool
 		}
 
 		/** Adds a public dependency module. */
-		public void AddPublicDependencyModule(string ModuleName, bool bCheckForDuplicates = true)
+		public virtual void AddPublicDependencyModule(string ModuleName, bool bCheckForDuplicates = true)
 		{
 			if (bCheckForDuplicates == true)
 			{
@@ -241,7 +241,7 @@ namespace UnrealBuildTool
 		}
 		
 		/** Adds a dynamically loaded module. */
-		public void AddDynamicallyLoadedModule(string ModuleName, bool bCheckForDuplicates = true)
+		public virtual void AddDynamicallyLoadedModule(string ModuleName, bool bCheckForDuplicates = true)
 		{
 			if (bCheckForDuplicates == true)
 			{
@@ -253,7 +253,7 @@ namespace UnrealBuildTool
 			DynamicallyLoadedModuleNames.Add(ModuleName);
 		}
         /** Adds a dynamically loaded module platform specifc. */
-        public void AddPlatformSpecificDynamicallyLoadedModule(string ModuleName, bool bCheckForDuplicates = true)
+		public virtual void AddPlatformSpecificDynamicallyLoadedModule(string ModuleName, bool bCheckForDuplicates = true)
         {
             if (bCheckForDuplicates == true)
             {
@@ -265,7 +265,7 @@ namespace UnrealBuildTool
             PlatformSpecificDynamicallyLoadedModuleNames.Add(ModuleName);
         }
 		/** Adds a public delay load DLL */
-		public void AddPublicDelayLoadDLLs(string DelayLoadDLLName, bool bCheckForDuplicates = true)
+		public virtual void AddPublicDelayLoadDLLs(string DelayLoadDLLName, bool bCheckForDuplicates = true)
 		{
 			if (bCheckForDuplicates == true)
 			{
@@ -278,7 +278,7 @@ namespace UnrealBuildTool
 		}
 
 		/** Adds a private include path. */
-		public void AddPrivateIncludePath(string PathName, bool bCheckForDuplicates = true)
+		public virtual void AddPrivateIncludePath(string PathName, bool bCheckForDuplicates = true)
 		{
 			if (bCheckForDuplicates == true)
 			{
@@ -291,7 +291,7 @@ namespace UnrealBuildTool
 		}
 
 		/** Adds a private include path module. */
-		public void AddPrivateIncludePathModule(string ModuleName, bool bCheckForDuplicates = true)
+		public virtual void AddPrivateIncludePathModule(string ModuleName, bool bCheckForDuplicates = true)
 		{
 			if (bCheckForDuplicates == true)
 			{
@@ -304,7 +304,7 @@ namespace UnrealBuildTool
 		}
 
 		/** Adds a dependency module. */
-		public void AddPrivateDependencyModule(string ModuleName, bool bCheckForDuplicates = true)
+		public virtual void AddPrivateDependencyModule(string ModuleName, bool bCheckForDuplicates = true)
 		{
 			if (bCheckForDuplicates == true)
 			{
@@ -317,13 +317,13 @@ namespace UnrealBuildTool
 		}
 
 		/** Tells the module that a specific dependency is circularly reference */
-		public void AddCircularlyReferencedDependentModule( string ModuleName )
+		public virtual void AddCircularlyReferencedDependentModule(string ModuleName)
 		{
 			CircularlyReferencedDependentModules.Add( ModuleName );
 		}
 
 		/** Adds an extra module name */
-		public void AddExtraModuleName(string ModuleName)
+		public virtual void AddExtraModuleName(string ModuleName)
 		{
 			DynamicallyLoadedModuleNames.Add( ModuleName );
 		}
@@ -378,7 +378,7 @@ namespace UnrealBuildTool
 		}
 
 		/** Fix up the output path */
-		public string FixupOutputPath(string InOutputPath)
+		public virtual string FixupOutputPath(string InOutputPath)
 		{
 			string ModuleOutputPath = InOutputPath;
 			if ((OutputDirectory.Length != 0) && (UnrealBuildTool.BuildingRocket() == false))
