[GitHub] 72 : Extract instantiation of UEBuildModules to own method InstantiateModule
#ttp 333475

diff --git a/Engine/Source/Programs/UnrealBuildTool/Configuration/UEBuildTarget.cs b/Engine/Source/Programs/UnrealBuildTool/Configuration/UEBuildTarget.cs
--- a/Engine/Source/Programs/UnrealBuildTool/Configuration/UEBuildTarget.cs
+++ b/Engine/Source/Programs/UnrealBuildTool/Configuration/UEBuildTarget.cs
@@ -2400,87 +2400,110 @@ namespace UnrealBuildTool
 				UEBuildModuleType ModuleType = IsGameModule ? UEBuildModuleType.GameModule : UEBuildModuleType.EngineModule;
 
 				// Now, go ahead and create the module builder instance
-				switch (RulesObject.Type)
+				if((Module = InstantiateModule(RulesObject, ModuleName, ModuleType, ModuleDirectory, ApplicationOutputPath, IntelliSenseGatherer, ModuleSourceFiles)) == null)
 				{
-					case ModuleRules.ModuleType.CPlusPlus:
-						{
-							Module = new UEBuildModuleCPP(
-								InTarget: this,
-								InName: ModuleName,
-								InType: ModuleType,
-								InModuleDirectory: ModuleDirectory,
-								InOutputDirectory: ApplicationOutputPath,
-								InIntelliSenseGatherer: IntelliSenseGatherer,
-								InSourceFiles: ModuleSourceFiles,
-								InPublicSystemIncludePaths: RulesObject.PublicSystemIncludePaths,
-								InPublicIncludePaths: RulesObject.PublicIncludePaths,
-								InDefinitions: RulesObject.Definitions,
-								InPublicIncludePathModuleNames: RulesObject.PublicIncludePathModuleNames,
-								InPublicDependencyModuleNames: RulesObject.PublicDependencyModuleNames,
-								InPublicDelayLoadDLLs: RulesObject.PublicDelayLoadDLLs,
-								InPublicAdditionalLibraries: RulesObject.PublicAdditionalLibraries,
-								InPublicFrameworks: RulesObject.PublicFrameworks,
-								InPublicAdditionalFrameworks: RulesObject.PublicAdditionalFrameworks,
-								InPrivateIncludePaths: RulesObject.PrivateIncludePaths,
-								InPrivateIncludePathModuleNames: RulesObject.PrivateIncludePathModuleNames,
-								InPrivateDependencyModuleNames: RulesObject.PrivateDependencyModuleNames,
-								InCircularlyReferencedDependentModules: RulesObject.CircularlyReferencedDependentModules,
-								InDynamicallyLoadedModuleNames: RulesObject.DynamicallyLoadedModuleNames,
-								InPlatformSpecificDynamicallyLoadedModuleNames: RulesObject.PlatformSpecificDynamicallyLoadedModuleNames,
-								InOptimizeCode: RulesObject.OptimizeCode,
-								InAllowSharedPCH: (RulesObject.PCHUsage == ModuleRules.PCHUsageMode.NoSharedPCHs) ? false : true,
-								InSharedPCHHeaderFile: RulesObject.SharedPCHHeaderFile,
-								InUseRTTI: RulesObject.bUseRTTI,
-								InEnableBufferSecurityChecks: RulesObject.bEnableBufferSecurityChecks,
-								InFasterWithoutUnity: RulesObject.bFasterWithoutUnity,
-								InMinFilesUsingPrecompiledHeaderOverride: RulesObject.MinFilesUsingPrecompiledHeaderOverride,
-								InEnableExceptions: RulesObject.bEnableExceptions,
-								InEnableInlining: RulesObject.bEnableInlining
-								);
-						}
-						break;
+					throw new BuildException("Unrecognized module type specified by 'Rules' object {0}", RulesObject.ToString());
+				}
 
-					case ModuleRules.ModuleType.CPlusPlusCLR:
-						{
-							Module = new UEBuildModuleCPPCLR(
-								InTarget: this,
-								InName: ModuleName,
-								InType: ModuleType,
-								InModuleDirectory: ModuleDirectory,
-								InOutputDirectory: ApplicationOutputPath,
-								InIntelliSenseGatherer: IntelliSenseGatherer,
-								InSourceFiles: ModuleSourceFiles,
-								InDefinitions: RulesObject.Definitions,
-								InPublicSystemIncludePaths: RulesObject.PublicSystemIncludePaths,
-								InPublicIncludePaths: RulesObject.PublicIncludePaths,
-								InPublicIncludePathModuleNames: RulesObject.PublicIncludePathModuleNames,
-								InPublicDependencyModuleNames: RulesObject.PublicDependencyModuleNames,
-								InPublicDelayLoadDLLs: RulesObject.PublicDelayLoadDLLs,
-								InPublicAdditionalLibraries: RulesObject.PublicAdditionalLibraries,
-								InPublicFrameworks: RulesObject.PublicFrameworks,
-								InPublicAdditionalFrameworks: RulesObject.PublicAdditionalFrameworks,
-								InPrivateIncludePaths: RulesObject.PrivateIncludePaths,
-								InPrivateIncludePathModuleNames: RulesObject.PrivateIncludePathModuleNames,
-								InPrivateDependencyModuleNames: RulesObject.PrivateDependencyModuleNames,
-								InPrivateAssemblyReferences: RulesObject.PrivateAssemblyReferences,
-								InCircularlyReferencedDependentModules: RulesObject.CircularlyReferencedDependentModules,
-								InDynamicallyLoadedModuleNames: RulesObject.DynamicallyLoadedModuleNames,
-								InPlatformSpecificDynamicallyLoadedModuleNames: RulesObject.PlatformSpecificDynamicallyLoadedModuleNames,
-								InOptimizeCode: RulesObject.OptimizeCode,
-								InAllowSharedPCH: (RulesObject.PCHUsage == ModuleRules.PCHUsageMode.NoSharedPCHs) ? false : true,
-								InSharedPCHHeaderFile: RulesObject.SharedPCHHeaderFile,
-								InUseRTTI: RulesObject.bUseRTTI,
-								InEnableBufferSecurityChecks: RulesObject.bEnableBufferSecurityChecks,
-								InFasterWithoutUnity: RulesObject.bFasterWithoutUnity,
-								InMinFilesUsingPrecompiledHeaderOverride: RulesObject.MinFilesUsingPrecompiledHeaderOverride,
-								InEnableExceptions: RulesObject.bEnableExceptions,
-								InEnableInlining: RulesObject.bEnableInlining
-								);
-						}
-						break;
+				UnrealTargetPlatform Only = UnrealBuildTool.GetOnlyPlatformSpecificFor();
 
-					case ModuleRules.ModuleType.External:
-						Module = new UEBuildExternalModule(
+				if (Only == UnrealTargetPlatform.Unknown && UnrealBuildTool.SkipNonHostPlatforms())
+				{
+					Only = Platform;
+				}
+				// Allow all build platforms to 'adjust' the module setting. 
+				// This will allow undisclosed platforms to make changes without 
+				// exposing information about the platform in publicly accessible 
+				// locations.
+				UEBuildPlatform.PlatformModifyNewlyLoadedModule(Module, TargetInfo, Only);
+			}
+			return Module;
+		}
+
+		protected virtual UEBuildModule InstantiateModule(
+			ModuleRules RulesObject,
+			string ModuleName,
+			UEBuildModuleType ModuleType,
+			string ModuleDirectory,
+			string ApplicationOutputPath,
+			IntelliSenseGatherer IntelliSenseGatherer,
+			List<FileItem> ModuleSourceFiles)
+		{
+			switch (RulesObject.Type)
+			{
+				case ModuleRules.ModuleType.CPlusPlus:
+					return new UEBuildModuleCPP(
+							InTarget: this,
+							InName: ModuleName,
+							InType: ModuleType,
+							InModuleDirectory: ModuleDirectory,
+							InOutputDirectory: ApplicationOutputPath,
+							InIntelliSenseGatherer: IntelliSenseGatherer,
+							InSourceFiles: ModuleSourceFiles,
+							InPublicSystemIncludePaths: RulesObject.PublicSystemIncludePaths,
+							InPublicIncludePaths: RulesObject.PublicIncludePaths,
+							InDefinitions: RulesObject.Definitions,
+							InPublicIncludePathModuleNames: RulesObject.PublicIncludePathModuleNames,
+							InPublicDependencyModuleNames: RulesObject.PublicDependencyModuleNames,
+							InPublicDelayLoadDLLs: RulesObject.PublicDelayLoadDLLs,
+							InPublicAdditionalLibraries: RulesObject.PublicAdditionalLibraries,
+							InPublicFrameworks: RulesObject.PublicFrameworks,
+							InPublicAdditionalFrameworks: RulesObject.PublicAdditionalFrameworks,
+							InPrivateIncludePaths: RulesObject.PrivateIncludePaths,
+							InPrivateIncludePathModuleNames: RulesObject.PrivateIncludePathModuleNames,
+							InPrivateDependencyModuleNames: RulesObject.PrivateDependencyModuleNames,
+							InCircularlyReferencedDependentModules: RulesObject.CircularlyReferencedDependentModules,
+							InDynamicallyLoadedModuleNames: RulesObject.DynamicallyLoadedModuleNames,
+							InPlatformSpecificDynamicallyLoadedModuleNames: RulesObject.PlatformSpecificDynamicallyLoadedModuleNames,
+							InOptimizeCode: RulesObject.OptimizeCode,
+							InAllowSharedPCH: (RulesObject.PCHUsage == ModuleRules.PCHUsageMode.NoSharedPCHs) ? false : true,
+							InSharedPCHHeaderFile: RulesObject.SharedPCHHeaderFile,
+							InUseRTTI: RulesObject.bUseRTTI,
+							InEnableBufferSecurityChecks: RulesObject.bEnableBufferSecurityChecks,
+							InFasterWithoutUnity: RulesObject.bFasterWithoutUnity,
+							InMinFilesUsingPrecompiledHeaderOverride: RulesObject.MinFilesUsingPrecompiledHeaderOverride,
+							InEnableExceptions: RulesObject.bEnableExceptions,
+							InEnableInlining: RulesObject.bEnableInlining
+						);
+
+				case ModuleRules.ModuleType.CPlusPlusCLR:
+					return new UEBuildModuleCPPCLR(
+							InTarget: this,
+							InName: ModuleName,
+							InType: ModuleType,
+							InModuleDirectory: ModuleDirectory,
+							InOutputDirectory: ApplicationOutputPath,
+							InIntelliSenseGatherer: IntelliSenseGatherer,
+							InSourceFiles: ModuleSourceFiles,
+							InDefinitions: RulesObject.Definitions,
+							InPublicSystemIncludePaths: RulesObject.PublicSystemIncludePaths,
+							InPublicIncludePaths: RulesObject.PublicIncludePaths,
+							InPublicIncludePathModuleNames: RulesObject.PublicIncludePathModuleNames,
+							InPublicDependencyModuleNames: RulesObject.PublicDependencyModuleNames,
+							InPublicDelayLoadDLLs: RulesObject.PublicDelayLoadDLLs,
+							InPublicAdditionalLibraries: RulesObject.PublicAdditionalLibraries,
+							InPublicFrameworks: RulesObject.PublicFrameworks,
+							InPublicAdditionalFrameworks: RulesObject.PublicAdditionalFrameworks,
+							InPrivateIncludePaths: RulesObject.PrivateIncludePaths,
+							InPrivateIncludePathModuleNames: RulesObject.PrivateIncludePathModuleNames,
+							InPrivateDependencyModuleNames: RulesObject.PrivateDependencyModuleNames,
+							InPrivateAssemblyReferences: RulesObject.PrivateAssemblyReferences,
+							InCircularlyReferencedDependentModules: RulesObject.CircularlyReferencedDependentModules,
+							InDynamicallyLoadedModuleNames: RulesObject.DynamicallyLoadedModuleNames,
+							InPlatformSpecificDynamicallyLoadedModuleNames: RulesObject.PlatformSpecificDynamicallyLoadedModuleNames,
+							InOptimizeCode: RulesObject.OptimizeCode,
+							InAllowSharedPCH: (RulesObject.PCHUsage == ModuleRules.PCHUsageMode.NoSharedPCHs) ? false : true,
+							InSharedPCHHeaderFile: RulesObject.SharedPCHHeaderFile,
+							InUseRTTI: RulesObject.bUseRTTI,
+							InEnableBufferSecurityChecks: RulesObject.bEnableBufferSecurityChecks,
+							InFasterWithoutUnity: RulesObject.bFasterWithoutUnity,
+							InMinFilesUsingPrecompiledHeaderOverride: RulesObject.MinFilesUsingPrecompiledHeaderOverride,
+							InEnableExceptions: RulesObject.bEnableExceptions,
+							InEnableInlining: RulesObject.bEnableInlining
+						);
+
+				case ModuleRules.ModuleType.External:
+					return new UEBuildExternalModule(
 							InTarget: this,
 							InName: ModuleName,
 							InType: ModuleType,
@@ -2495,26 +2518,12 @@ namespace UnrealBuildTool
 							InPublicAdditionalFrameworks: RulesObject.PublicAdditionalFrameworks,
 							InPublicAdditionalShadowFiles: RulesObject.PublicAdditionalShadowFiles,
 							InPublicDependencyModuleNames: RulesObject.PublicDependencyModuleNames,
-							InPublicDelayLoadDLLs: RulesObject.PublicDelayLoadDLLs);
-						break;
-
-					default:
-						throw new BuildException("Unrecognized module type specified by 'Rules' object {0}", RulesObject.ToString());
-				}
-
-				UnrealTargetPlatform Only = UnrealBuildTool.GetOnlyPlatformSpecificFor();
+							InPublicDelayLoadDLLs: RulesObject.PublicDelayLoadDLLs
+						);
 
-				if (Only == UnrealTargetPlatform.Unknown && UnrealBuildTool.SkipNonHostPlatforms())
-				{
-					Only = Platform;
-				}
-				// Allow all build platforms to 'adjust' the module setting. 
-				// This will allow undisclosed platforms to make changes without 
-				// exposing information about the platform in publicly accessible 
-				// locations.
-				UEBuildPlatform.PlatformModifyNewlyLoadedModule(Module, TargetInfo, Only);
+				default:
+					return null;
 			}
-			return Module;
 		}
 
 		/** Finds a module given its name.  Throws an exception if the module couldn't be found. */
