[GitHub] 61 : Check CommandPath before searching in dictionary
#ttp 333300

diff --git a/Engine/Source/Programs/UnrealBuildTool/System/ActionGraph.cs b/Engine/Source/Programs/UnrealBuildTool/System/ActionGraph.cs
--- a/Engine/Source/Programs/UnrealBuildTool/System/ActionGraph.cs
+++ b/Engine/Source/Programs/UnrealBuildTool/System/ActionGraph.cs
@@ -208,7 +208,7 @@ namespace UnrealBuildTool
 			List<Action> ActionsToExecute = new List<Action>();
 			foreach (Action Action in AllActions)
 			{
-				if (ActionsNeededForThisTarget.ContainsKey(Action) && OutdatedActionDictionary[Action] && Action.CommandPath != null)
+				if (Action.CommandPath != null && ActionsNeededForThisTarget.ContainsKey(Action) && OutdatedActionDictionary[Action])
 				{
 					ActionsToExecute.Add(Action);
 				}
