Refactoring UBT not-supported-configuration error.

diff --git a/Engine/Source/Programs/UnrealBuildTool/System/RulesCompiler.cs b/Engine/Source/Programs/UnrealBuildTool/System/RulesCompiler.cs
--- a/Engine/Source/Programs/UnrealBuildTool/System/RulesCompiler.cs
+++ b/Engine/Source/Programs/UnrealBuildTool/System/RulesCompiler.cs
@@ -437,11 +437,15 @@ namespace UnrealBuildTool
 		/// <returns>true if successful, false if not</returns>
 		public virtual bool GetSupportedPlatforms(ref List<UnrealTargetPlatform> OutPlatforms)
 		{
-			if ((Type == TargetType.Program) || IsEditorType(Type))
+			if(Type == TargetType.Program)
 			{
 				// By default, all programs are desktop only.
 				return UnrealBuildTool.GetAllDesktopPlatforms(ref OutPlatforms, false);
 			}
+			else if(IsEditorType(Type))
+			{
+				return UnrealBuildTool.GetAllEditorPlatforms(ref OutPlatforms, false);
+			}
 			else if (TargetRules.IsGameType(Type))
 			{
 				// By default all games support all platforms
@@ -453,12 +457,6 @@ namespace UnrealBuildTool
 
 		public bool SupportsPlatform(UnrealTargetPlatform InPlatform)
 		{
-			// Only Win64 and Mac are supported by the editor.
-			if (IsEditorType(Type) && InPlatform != UnrealTargetPlatform.Win64 && InPlatform != UnrealTargetPlatform.Mac)
-			{
-				return false;
-			}
-
 			List<UnrealTargetPlatform> SupportedPlatforms = new List<UnrealTargetPlatform>();
 			if (GetSupportedPlatforms(ref SupportedPlatforms) == true)
 			{
diff --git a/Engine/Source/Programs/UnrealBuildTool/System/UnrealBuildTool.cs b/Engine/Source/Programs/UnrealBuildTool/System/UnrealBuildTool.cs
--- a/Engine/Source/Programs/UnrealBuildTool/System/UnrealBuildTool.cs
+++ b/Engine/Source/Programs/UnrealBuildTool/System/UnrealBuildTool.cs
@@ -165,6 +165,22 @@ namespace UnrealBuildTool
 				);
 		}
 
+		// @todo projectfiles: Move this into the ProjectPlatformGeneration class?
+		/// <summary>
+		/// IsEditorPlatform
+		/// </summary>
+		/// <param name="InPlatform">The platform of interest</param>
+		/// <returns>True if the given platform is a platform that supports the editor.</returns>
+		static public bool IsEditorPlatform(UnrealTargetPlatform InPlatform)
+		{
+			// Windows 64, Linux and Mac are editor platforms.
+			return (
+				(InPlatform == UnrealTargetPlatform.Win64) ||
+				(InPlatform == UnrealTargetPlatform.Linux) ||
+				(InPlatform == UnrealTargetPlatform.Mac)
+				);
+		}
+
         // @todo projectfiles: Move this into the ProjectPlatformGeneration class?
         /// <summary>
         /// IsServerPlatform
@@ -199,73 +215,73 @@ namespace UnrealBuildTool
                 );
         }
 
-        /// <summary>
-		/// Get all platforms
+		/// <summary>
+		/// Gets all platforms that satisfies the predicate.
 		/// </summary>
-		/// <param name="OutPlatforms">The list of platform to fill in</param>
-		/// <param name="bCheckValidity">If true, ensure it is a valid platform</param>
-		/// <returns>true if successful and platforms are found, false if not</returns>
-		static public bool GetAllPlatforms(ref List<UnrealTargetPlatform> OutPlatforms, bool bCheckValidity = true)
+		/// <param name="OutPlatforms">The list of platform to fill in.</param>
+		/// <param name="Predicate">The predicate to filter the platforms.</param>
+		/// <param name="bCheckValidity">If true, ensure it is a valid platform.</param>
+		/// <returns>True if successful and platforms are found, false if not.</returns>
+		static private bool GetPlatforms(ref List<UnrealTargetPlatform> OutPlatforms, Func<UnrealTargetPlatform, bool> Predicate, bool bCheckValidity = true)
 		{
 			OutPlatforms.Clear();
-			foreach (UnrealTargetPlatform platform in Enum.GetValues(typeof(UnrealTargetPlatform)))
+			foreach (UnrealTargetPlatform Platform in Enum.GetValues(typeof(UnrealTargetPlatform)))
 			{
-				UEBuildPlatform BuildPlatform = UEBuildPlatform.GetBuildPlatform(platform, true);
-				if ((bCheckValidity == false) || (BuildPlatform != null))
+				if (Predicate(Platform))
 				{
-					OutPlatforms.Add(platform);
+					UEBuildPlatform BuildPlatform = UEBuildPlatform.GetBuildPlatform(Platform, true);
+					if ((bCheckValidity == false) || (BuildPlatform != null))
+					{
+						OutPlatforms.Add(Platform);
+					}
 				}
 			}
 
 			return (OutPlatforms.Count > 0) ? true : false;
 		}
 
+        /// <summary>
+		/// Gets all platforms.
+		/// </summary>
+		/// <param name="OutPlatforms">The list of platform to fill in.</param>
+		/// <param name="bCheckValidity">If true, ensure it is a valid platform.</param>
+		/// <returns>True if successful and platforms are found, false if not.</returns>
+		static public bool GetAllPlatforms(ref List<UnrealTargetPlatform> OutPlatforms, bool bCheckValidity = true)
+		{
+			return GetPlatforms(ref OutPlatforms, (Platform) => { return true; }, bCheckValidity);
+		}
+
 		/// <summary>
-		/// Get all desktop platforms
+		/// Gets all desktop platforms.
 		/// </summary>
-		/// <param name="OutPlatforms">The list of platform to fill in</param>
-		/// <param name="bCheckValidity">If true, ensure it is a valid platform</param>
-		/// <returns>true if successful and platforms are found, false if not</returns>
+		/// <param name="OutPlatforms">The list of platform to fill in.</param>
+		/// <param name="bCheckValidity">If true, ensure it is a valid platform.</param>
+		/// <returns>True if successful and platforms are found, false if not.</returns>
 		static public bool GetAllDesktopPlatforms(ref List<UnrealTargetPlatform> OutPlatforms, bool bCheckValidity = true)
 		{
-			OutPlatforms.Clear();
-			foreach (UnrealTargetPlatform platform in Enum.GetValues(typeof(UnrealTargetPlatform)))
-			{
-				if (UnrealBuildTool.IsDesktopPlatform(platform))
-				{
-					UEBuildPlatform BuildPlatform = UEBuildPlatform.GetBuildPlatform(platform, true);
-					if ((bCheckValidity == false) || (BuildPlatform != null))
-					{
-						OutPlatforms.Add(platform);
-					}
-				}
-			}
+			return GetPlatforms(ref OutPlatforms, UnrealBuildTool.IsDesktopPlatform, bCheckValidity);
+		}
 
-			return (OutPlatforms.Count > 0) ? true : false;
+		/// <summary>
+		/// Gets all editor platforms.
+		/// </summary>
+		/// <param name="OutPlatforms">The list of platform to fill in.</param>
+		/// <param name="bCheckValidity">If true, ensure it is a valid platform.</param>
+		/// <returns>True if successful and platforms are found, false if not.</returns>
+		static public bool GetAllEditorPlatforms(ref List<UnrealTargetPlatform> OutPlatforms, bool bCheckValidity = true)
+		{
+			return GetPlatforms(ref OutPlatforms, UnrealBuildTool.IsEditorPlatform, bCheckValidity);
 		}
 
         /// <summary>
-        /// Get all server platforms
+        /// Gets all server platforms.
         /// </summary>
-        /// <param name="OutPlatforms">The list of platform to fill in</param>
-        /// <param name="bCheckValidity">If true, ensure it is a valid platform</param>
-        /// <returns>true if successful and platforms are found, false if not</returns>
+		/// <param name="OutPlatforms">The list of platform to fill in.</param>
+		/// <param name="bCheckValidity">If true, ensure it is a valid platform.</param>
+		/// <returns>True if successful and platforms are found, false if not.</returns>
         static public bool GetAllServerPlatforms(ref List<UnrealTargetPlatform> OutPlatforms, bool bCheckValidity = true)
         {
-            OutPlatforms.Clear();
-            foreach (UnrealTargetPlatform platform in Enum.GetValues(typeof(UnrealTargetPlatform)))
-            {
-                if (UnrealBuildTool.IsServerPlatform(platform))
-                {
-                    UEBuildPlatform BuildPlatform = UEBuildPlatform.GetBuildPlatform(platform, true);
-                    if ((bCheckValidity == false) || (BuildPlatform != null))
-                    {
-                        OutPlatforms.Add(platform);
-                    }
-                }
-            }
-
-            return (OutPlatforms.Count > 0) ? true : false;
+			return GetPlatforms(ref OutPlatforms, UnrealBuildTool.IsServerPlatform, bCheckValidity);
         }
 
 		/// <summary>
