Optimized section contents generation for Xcode projects by switching to StringBuilder.

diff --git a/Engine/Source/Programs/UnrealBuildTool/System/XcodeProject.cs b/Engine/Source/Programs/UnrealBuildTool/System/XcodeProject.cs
--- a/Engine/Source/Programs/UnrealBuildTool/System/XcodeProject.cs
+++ b/Engine/Source/Programs/UnrealBuildTool/System/XcodeProject.cs
@@ -59,12 +59,12 @@ namespace UnrealBuildTool
 		/// <summary>
 		/// Constructor
 		/// </summary>
-		public XcodeFileGroup(string InName, string InPath, bool InReference = false )
+		public XcodeFileGroup(string InName, string InPath, bool InReference = false)
 		{
 			GroupName = InName;
 			GroupPath = InPath;
 			GroupGuid = XcodeProjectFileGenerator.MakeXcodeGuid();
-            bReference = InReference;
+			bReference = InReference;
 		}
 
 		/// <summary>
@@ -72,8 +72,10 @@ namespace UnrealBuildTool
 		/// </summary>
 		public void Append(ref StringBuilder Contents, bool bFilesOnly = false)
 		{
-            if (bReference)
-                return;
+			if (bReference)
+			{
+				return;
+			}
 
 			if (!bFilesOnly)
 			{
@@ -112,7 +114,7 @@ namespace UnrealBuildTool
 		public string GroupPath;
 		public Dictionary<string, XcodeFileGroup> Children = new Dictionary<string, XcodeFileGroup>();
 		public List<XcodeSourceFile> Files = new List<XcodeSourceFile>();
-        public bool bReference;
+		public bool bReference;
 	}
 
 	class XcodeProjectFile : ProjectFile
@@ -276,8 +278,12 @@ namespace UnrealBuildTool
 		 * Generates bodies of all sections that contain a list of source files plus a dictionary of project navigator groups.
 		 */
 		public void GenerateSectionsContents(ref string PBXBuildFileSection, ref string PBXFileReferenceSection,
-		                                     ref string PBXSourcesBuildPhaseSection, ref Dictionary<string, XcodeFileGroup> Groups)
+											 ref string PBXSourcesBuildPhaseSection, ref Dictionary<string, XcodeFileGroup> Groups)
 		{
+			StringBuilder FileSection = new StringBuilder();
+			StringBuilder ReferenceSection = new StringBuilder();
+			StringBuilder BuildPhaseSection = new StringBuilder();
+
 			foreach (var CurSourceFile in SourceFiles)
 			{
 				XcodeSourceFile SourceFile = CurSourceFile as XcodeSourceFile;
@@ -298,24 +304,24 @@ namespace UnrealBuildTool
 					SourceFile.ReplaceGUIDS(UE4CmdLineRunFileGuid, UE4CmdLineRunFileRefGuid);
 				}
 
-				if(IsGeneratedProject)
+				if (IsGeneratedProject)
 				{
-					PBXBuildFileSection += string.Format("\t\t{0} /* {1} in {2} */ = {{isa = PBXBuildFile; fileRef = {3} /* {1} */; }};" + ProjectFileGenerator.NewLine,
-					                                     SourceFile.FileGUID,
-					                                     FileName,
-					                                     GetFileCategory(FileExtension),
-					                                     SourceFile.FileRefGUID);
+					FileSection.Append(string.Format("\t\t{0} /* {1} in {2} */ = {{isa = PBXBuildFile; fileRef = {3} /* {1} */; }};" + ProjectFileGenerator.NewLine,
+													 SourceFile.FileGUID,
+													 FileName,
+													 GetFileCategory(FileExtension),
+													 SourceFile.FileRefGUID));
 				}
 
-				PBXFileReferenceSection += string.Format("\t\t{0} /* {1} */ = {{isa = PBXFileReference; lastKnownFileType = {2}; name = \"{1}\"; path = \"{3}\"; sourceTree = SOURCE_ROOT; }};" + ProjectFileGenerator.NewLine,
-				                            SourceFile.FileRefGUID,
-				                            FileName,
-				                            GetFileType(FileExtension),
-                                            FilePathMac);
+				ReferenceSection.Append(string.Format("\t\t{0} /* {1} */ = {{isa = PBXFileReference; lastKnownFileType = {2}; name = \"{1}\"; path = \"{3}\"; sourceTree = SOURCE_ROOT; }};" + ProjectFileGenerator.NewLine,
+													  SourceFile.FileRefGUID,
+													  FileName,
+													  GetFileType(FileExtension),
+													  FilePathMac));
 
 				if (IsPartOfUE4XcodeHelperTarget(SourceFile))
 				{
-					PBXSourcesBuildPhaseSection += "\t\t\t\t" + SourceFile.FileGUID + " /* " + FileName + " in Sources */," + ProjectFileGenerator.NewLine;
+					BuildPhaseSection.Append("\t\t\t\t" + SourceFile.FileGUID + " /* " + FileName + " in Sources */," + ProjectFileGenerator.NewLine);
 				}
 
 				string GroupPath = Path.GetFullPath(Path.GetDirectoryName(SourceFile.FilePath));
@@ -325,6 +331,10 @@ namespace UnrealBuildTool
 					Group.Files.Add(SourceFile);
 				}
 			}
+
+			PBXBuildFileSection += FileSection.ToString();
+			PBXFileReferenceSection += ReferenceSection.ToString();
+			PBXSourcesBuildPhaseSection += BuildPhaseSection.ToString();
 		}
 	}
 }
