Fix setup script failures for scripts that require admin access (mwassmer)

diff --git a/Engine/Source/Programs/UnrealBuildTool/Configuration/UEBuildPlatform.cs b/Engine/Source/Programs/UnrealBuildTool/Configuration/UEBuildPlatform.cs
--- a/Engine/Source/Programs/UnrealBuildTool/Configuration/UEBuildPlatform.cs
+++ b/Engine/Source/Programs/UnrealBuildTool/Configuration/UEBuildPlatform.cs
@@ -395,6 +395,8 @@ namespace UnrealBuildTool
                     HookProcess.StartInfo.FileName = HookExe;
                     HookProcess.StartInfo.Arguments = "";
 
+                    //installers may require administrator access to succeed. so run as an admmin.
+					HookProcess.StartInfo.Verb = "runas";
                     HookProcess.Start();
                     HookProcess.WaitForExit();
 
