fix for aligned memory read of script code for iOS
#ue4
#ios

diff --git a/Engine/Source/Programs/UnrealBuildTool/IOS/UEBuildIOS.cs b/Engine/Source/Programs/UnrealBuildTool/IOS/UEBuildIOS.cs
--- a/Engine/Source/Programs/UnrealBuildTool/IOS/UEBuildIOS.cs
+++ b/Engine/Source/Programs/UnrealBuildTool/IOS/UEBuildIOS.cs
@@ -109,6 +109,7 @@ namespace UnrealBuildTool
 			InBuildTarget.GlobalCompileEnvironment.Config.Definitions.Add("WITH_DATABASE_SUPPORT=0");
 			InBuildTarget.GlobalCompileEnvironment.Config.Definitions.Add("WITH_EDITOR=0");
 			InBuildTarget.GlobalCompileEnvironment.Config.Definitions.Add("USE_NULL_RHI=0");
+            InBuildTarget.GlobalCompileEnvironment.Config.Definitions.Add("REQUIRES_ALIGNED_INT_ACCESS");
 
 			if (GetActiveArchitecture() == "-simulator")
 			{
diff --git a/Engine/Source/Runtime/CoreUObject/Public/UObject/ScriptSerialization.h b/Engine/Source/Runtime/CoreUObject/Public/UObject/ScriptSerialization.h
--- a/Engine/Source/Runtime/CoreUObject/Public/UObject/ScriptSerialization.h
+++ b/Engine/Source/Runtime/CoreUObject/Public/UObject/ScriptSerialization.h
@@ -13,7 +13,7 @@
 #endif
 
 #ifndef XFER
-#if PLATFORM_HTML5
+#ifdef REQUIRES_ALIGNED_INT_ACCESS
 	#define XFER(T) \
 		{ \
 		T Temp; \
