Removed references to .generated.inl.

diff --git a/Engine/Source/Programs/UnrealBuildTool/System/CPPEnvironment.cs b/Engine/Source/Programs/UnrealBuildTool/System/CPPEnvironment.cs
--- a/Engine/Source/Programs/UnrealBuildTool/System/CPPEnvironment.cs
+++ b/Engine/Source/Programs/UnrealBuildTool/System/CPPEnvironment.cs
@@ -253,7 +253,7 @@ namespace UnrealBuildTool
 		/** List of "shared" precompiled header environments, potentially shared between modules */
 		public readonly List<PrecompileHeaderEnvironment> SharedPCHEnvironments = new List<PrecompileHeaderEnvironment>();
 
-		/** @Hack: UnrealHeaderTool should not list CoreUObject.generated.inl as a dependency */
+		// Whether or not UHT is being built
 		public bool bHackHeaderGenerator;
 
 		/** Default constructor. */
diff --git a/Engine/Source/Programs/UnrealBuildTool/System/CPPHeaders.cs b/Engine/Source/Programs/UnrealBuildTool/System/CPPHeaders.cs
--- a/Engine/Source/Programs/UnrealBuildTool/System/CPPHeaders.cs
+++ b/Engine/Source/Programs/UnrealBuildTool/System/CPPHeaders.cs
@@ -429,26 +429,22 @@ namespace UnrealBuildTool
 					HasUObjects = true;
 				}
 
-				// @Hack: UHT should not list CoreUObject.generated.inl as its dependency.
-				if (!bHackHeaderGenerator || String.Compare(HeaderValue, "CoreUObject.generated.inl", true) != 0)
+				//@TODO: The intermediate exclusion is to work around autogenerated absolute paths in Module.SomeGame.cpp style files
+				bool bIsIntermediateOrThirdParty = FileToRead.Contains("Intermediate") || FileToRead.Contains("ThirdParty");
+				bool bCheckForBackwardSlashes = FileToRead.StartsWith(InstalledFolder);
+				if (UnrealBuildTool.HasUProjectFile())
 				{
-					//@TODO: The intermediate exclusion is to work around autogenerated absolute paths in Module.SomeGame.cpp style files
-					bool bIsIntermediateOrThirdParty = FileToRead.Contains("Intermediate") || FileToRead.Contains("ThirdParty");
-					bool bCheckForBackwardSlashes = FileToRead.StartsWith(InstalledFolder);
-					if (UnrealBuildTool.HasUProjectFile())
-					{
-						bCheckForBackwardSlashes |= Utils.IsFileUnderDirectory(FileToRead, UnrealBuildTool.GetUProjectPath());
-					}
-					if (bCheckForBackwardSlashes && !bIsIntermediateOrThirdParty)
+					bCheckForBackwardSlashes |= Utils.IsFileUnderDirectory(FileToRead, UnrealBuildTool.GetUProjectPath());
+				}
+				if (bCheckForBackwardSlashes && !bIsIntermediateOrThirdParty)
+				{
+					if (HeaderValue.IndexOf('\\', 0) >= 0)
 					{
-						if (HeaderValue.IndexOf('\\', 0) >= 0)
-						{
-							throw new BuildException("In {0}: #include \"{1}\" contains backslashes ('\\'), please use forward slashes ('/') instead.", FileToRead, C.Value);
-						}
+						throw new BuildException("In {0}: #include \"{1}\" contains backslashes ('\\'), please use forward slashes ('/') instead.", FileToRead, C.Value);
 					}
-					HeaderValue = Utils.CleanDirectorySeparators(HeaderValue);
-					Result.Add(new DependencyInclude(HeaderValue));
 				}
+				HeaderValue = Utils.CleanDirectorySeparators(HeaderValue);
+				Result.Add(new DependencyInclude(HeaderValue));
 			}
 
 			// also look for #import in objective C files
