Modifications to build/cook/run pipeline to support third party frameworks on iOS

diff --git a/Engine/Source/Programs/AutomationTool/IOS/IOSPlatform.Automation.cs b/Engine/Source/Programs/AutomationTool/IOS/IOSPlatform.Automation.cs
--- a/Engine/Source/Programs/AutomationTool/IOS/IOSPlatform.Automation.cs
+++ b/Engine/Source/Programs/AutomationTool/IOS/IOSPlatform.Automation.cs
@@ -475,6 +475,12 @@ public class IOSPlatform : Platform
 				SC.StageFiles(StagedFileType.NonUFS, SourcePath, "*.png", false, null, "", true, false);
 			}
 
+			// copy any additional framework assets that will be needed at runtime
+			{
+				string SourcePath = CombinePaths( SC.LocalRoot, "Engine", "Build", "IOS", "Resources", "FrameworkAssets" );
+				SC.StageFiles( StagedFileType.NonUFS, SourcePath, "*.*", true, null, "", true, false );
+			}
+
 			// copy the icons/launch screens from the game (may stomp the engine copies)
 			{
 				string SourcePath = CombinePaths(SC.ProjectRoot, "Build", "IOS", "Resources", "Graphics");
diff --git a/Engine/Source/Programs/UnrealBuildTool/Configuration/UEBuildModule.cs b/Engine/Source/Programs/UnrealBuildTool/Configuration/UEBuildModule.cs
--- a/Engine/Source/Programs/UnrealBuildTool/Configuration/UEBuildModule.cs
+++ b/Engine/Source/Programs/UnrealBuildTool/Configuration/UEBuildModule.cs
@@ -818,6 +818,7 @@ namespace UnrealBuildTool
 			IEnumerable<string> InPublicAdditionalLibraries,
 			IEnumerable<string> InPublicFrameworks,
 			IEnumerable<string> InPublicAdditionalFrameworks,
+			IEnumerable<string> InPublicAdditionalShadowFiles,
 			IEnumerable<string> InPrivateIncludePaths,
 			IEnumerable<string> InPrivateIncludePathModuleNames,
 			IEnumerable<string> InPrivateDependencyModuleNames,
@@ -846,7 +847,7 @@ namespace UnrealBuildTool
 					InPublicAdditionalLibraries,
 					InPublicFrameworks,
 					InPublicAdditionalFrameworks,
-					null,
+					InPublicAdditionalShadowFiles,
 					InPublicIncludePathModuleNames,
 					InPublicDependencyModuleNames,
 					InPublicDelayLoadDLLs,
@@ -1839,6 +1840,7 @@ namespace UnrealBuildTool
 			IEnumerable<string> InPublicAdditionalLibraries,
 			IEnumerable<string> InPublicFrameworks,
 			IEnumerable<string> InPublicAdditionalFrameworks,
+			IEnumerable<string> InPublicAdditionalShadowFiles,
 			IEnumerable<string> InPrivateIncludePaths,
 			IEnumerable<string> InPrivateIncludePathModuleNames,
 			IEnumerable<string> InPrivateDependencyModuleNames,
@@ -1857,7 +1859,7 @@ namespace UnrealBuildTool
 			)
 		: base(InTarget,InName,InType,InModuleDirectory,InOutputDirectory,InIntelliSenseGatherer,
 			InSourceFiles,InPublicIncludePaths,InPublicSystemIncludePaths,InDefinitions,
-			InPublicIncludePathModuleNames,InPublicDependencyModuleNames,InPublicDelayLoadDLLs,InPublicAdditionalLibraries,InPublicFrameworks,InPublicAdditionalFrameworks,
+			InPublicIncludePathModuleNames,InPublicDependencyModuleNames,InPublicDelayLoadDLLs,InPublicAdditionalLibraries,InPublicFrameworks,InPublicAdditionalFrameworks,InPublicAdditionalShadowFiles,
 			InPrivateIncludePaths,InPrivateIncludePathModuleNames,InPrivateDependencyModuleNames,
             InCircularlyReferencedDependentModules, InDynamicallyLoadedModuleNames, InPlatformSpecificDynamicallyLoadedModuleNames, InOptimizeCode,
 			InAllowSharedPCH, InSharedPCHHeaderFile, InUseRTTI, InEnableBufferSecurityChecks, InFasterWithoutUnity, InMinFilesUsingPrecompiledHeaderOverride,
diff --git a/Engine/Source/Programs/UnrealBuildTool/Configuration/UEBuildTarget.cs b/Engine/Source/Programs/UnrealBuildTool/Configuration/UEBuildTarget.cs
--- a/Engine/Source/Programs/UnrealBuildTool/Configuration/UEBuildTarget.cs
+++ b/Engine/Source/Programs/UnrealBuildTool/Configuration/UEBuildTarget.cs
@@ -1566,6 +1566,7 @@ namespace UnrealBuildTool
 					InPublicDelayLoadDLLs: new List<string>(),
 					InPublicAdditionalLibraries: new List<string>(),
 					InPublicFrameworks: new List<string>(),
+					InPublicAdditionalShadowFiles: new List<string>(),
 					InPublicAdditionalFrameworks: new List<string>(),
 					InPrivateIncludePaths: new List<string>(),
 					InPrivateIncludePathModuleNames: new List<string>(),
@@ -2449,6 +2450,7 @@ namespace UnrealBuildTool
 							InPublicAdditionalLibraries: RulesObject.PublicAdditionalLibraries,
 							InPublicFrameworks: RulesObject.PublicFrameworks,
 							InPublicAdditionalFrameworks: RulesObject.PublicAdditionalFrameworks,
+							InPublicAdditionalShadowFiles: RulesObject.PublicAdditionalShadowFiles,
 							InPrivateIncludePaths: RulesObject.PrivateIncludePaths,
 							InPrivateIncludePathModuleNames: RulesObject.PrivateIncludePathModuleNames,
 							InPrivateDependencyModuleNames: RulesObject.PrivateDependencyModuleNames,
@@ -2484,6 +2486,7 @@ namespace UnrealBuildTool
 							InPublicAdditionalLibraries: RulesObject.PublicAdditionalLibraries,
 							InPublicFrameworks: RulesObject.PublicFrameworks,
 							InPublicAdditionalFrameworks: RulesObject.PublicAdditionalFrameworks,
+							InPublicAdditionalShadowFiles: RulesObject.PublicAdditionalShadowFiles,
 							InPrivateIncludePaths: RulesObject.PrivateIncludePaths,
 							InPrivateIncludePathModuleNames: RulesObject.PrivateIncludePathModuleNames,
 							InPrivateDependencyModuleNames: RulesObject.PrivateDependencyModuleNames,
diff --git a/Engine/Source/Programs/UnrealBuildTool/IOS/IOSToolChain.cs b/Engine/Source/Programs/UnrealBuildTool/IOS/IOSToolChain.cs
--- a/Engine/Source/Programs/UnrealBuildTool/IOS/IOSToolChain.cs
+++ b/Engine/Source/Programs/UnrealBuildTool/IOS/IOSToolChain.cs
@@ -66,6 +66,9 @@ namespace UnrealBuildTool
 
 		public List<string> BuiltBinaries = new List<string>();
 
+		/** List of third party framework files that need to be unzipped */
+		public List<string> ThirdPartyFrameworkZips = new List<string>();
+
 		/// <summary>
 		/// Function to call to reset default data.
 		/// </summary>
@@ -319,6 +322,16 @@ namespace UnrealBuildTool
 			}
 			foreach (string Framework in LinkEnvironment.Config.AdditionalFrameworks)
 			{
+				// If we see ThirdPartyFrameworks in the string, assume we need to treat this as a path + name
+				if ( Framework.Contains( "ThirdPartyFrameworks/" ) )
+				{
+					string FrameworkName = Framework.Substring( Framework.LastIndexOf( '/' ) + 1 );
+					string FrameworkPath = Framework.Replace( FrameworkName, "" );
+					Result += " -F " + FrameworkPath;
+					Result += " -framework " + FrameworkName;
+					continue;
+				}
+
 				Result += " -framework " + Framework;
 			}
 			foreach (string Framework in LinkEnvironment.Config.WeakFrameworks)
@@ -560,6 +573,11 @@ namespace UnrealBuildTool
 				// Add any additional files that we'll need in order to link the app
 				foreach (string AdditionalShadowFile in LinkEnvironment.Config.AdditionalShadowFiles)
 				{
+					if ( AdditionalShadowFile.Contains( "ThirdPartyFrameworks/" ) && AdditionalShadowFile.Contains( ".zip" ) )
+					{
+						ThirdPartyFrameworkZips.Add( AdditionalShadowFile );
+					}
+
 					FileItem ShadowFile = FileItem.GetExistingItemByPath(AdditionalShadowFile);
 					if (ShadowFile != null)
 					{
@@ -791,6 +809,41 @@ namespace UnrealBuildTool
 			}
 
 			base.PreBuildSync();
+
+			// Unzip any third party frameworks that are stored as zips
+			foreach ( string FrameworkZip in ThirdPartyFrameworkZips )
+			{
+				FileItem FrameworkZipItem = FileItem.GetExistingItemByPath( FrameworkZip );
+
+				if ( FrameworkZipItem == null )
+				{
+					Log.TraceInformation( "FrameworkZipItem not found for {0}", FrameworkZip );
+					continue;
+				}
+
+				if ( ExternalExecution.GetRuntimePlatform() == UnrealTargetPlatform.Mac )
+				{
+					// If we're on the mac, just unzip using the shell
+					string ResultsText;
+					string LocalZipPath = FrameworkZipItem.AbsolutePath.Substring( 0, FrameworkZipItem.AbsolutePath.LastIndexOf( '/' ) );
+					RunExecutableAndWait( "unzip", String.Format( "-o {0} -d {1}", FrameworkZipItem.AbsolutePath, LocalZipPath ), out ResultsText );
+					continue;
+				}
+
+				// We copied using RPC utility, we need to unzip using RPC utility as well
+				FileItem RemoteShadowFile = LocalToRemoteFileItem( FrameworkZipItem, false );
+
+				string ZipPath = RemoteShadowFile.AbsolutePath.Substring( 0, RemoteShadowFile.AbsolutePath.LastIndexOf( '/' ) );
+
+				Log.TraceInformation( "Unzipping: {0}", RemoteShadowFile.AbsolutePath );
+
+				Hashtable Result = RPCUtilHelper.Command( "/", String.Format( "unzip -o {0} -d {1}", RemoteShadowFile.AbsolutePath, ZipPath ), "", null );
+
+				foreach ( DictionaryEntry Entry in Result )
+				{
+					Log.TraceInformation( "{0}", Entry.Value );
+				}
+			}
 		}
 
 		public override void PostBuildSync(UEBuildTarget Target)
@@ -1034,5 +1087,20 @@ namespace UnrealBuildTool
 				}
 			}
 		}
+
+		public static int RunExecutableAndWait( string ExeName, string ArgumentList, out string StdOutResults )
+		{
+			// Create the process
+			ProcessStartInfo PSI = new ProcessStartInfo( ExeName, ArgumentList );
+			PSI.RedirectStandardOutput = true;
+			PSI.UseShellExecute = false;
+			PSI.CreateNoWindow = true;
+			Process NewProcess = Process.Start( PSI );
+
+			// Wait for the process to exit and grab it's output
+			StdOutResults = NewProcess.StandardOutput.ReadToEnd();
+			NewProcess.WaitForExit();
+			return NewProcess.ExitCode;
+		}
 	};
 }
diff --git a/Engine/Source/Programs/UnrealBuildTool/IOS/UEDeployIOS.cs b/Engine/Source/Programs/UnrealBuildTool/IOS/UEDeployIOS.cs
--- a/Engine/Source/Programs/UnrealBuildTool/IOS/UEDeployIOS.cs
+++ b/Engine/Source/Programs/UnrealBuildTool/IOS/UEDeployIOS.cs
@@ -167,6 +167,9 @@ namespace UnrealBuildTool.IOS
 			{
 				CopyFiles(BuildDirectory + "/Resources/Graphics", AppDirectory, "*.png", true);
 			}
+			
+			// copy additional engine framework assets in
+			CopyFiles( InEngineDir + "/Build/IOS/Resources/FrameworkAssets", AppDirectory, "*.*", true );
 
 			//CopyFiles(BuildDirectory, PayloadDirectory, null, "iTunesArtwork", null);
 			CopyFolder(InEngineDir + "/Content/Stats", AppDirectory + "/cookeddata/engine/content/stats", true);
diff --git a/Engine/Source/Programs/UnrealBuildTool/System/LinkEnvironment.cs b/Engine/Source/Programs/UnrealBuildTool/System/LinkEnvironment.cs
--- a/Engine/Source/Programs/UnrealBuildTool/System/LinkEnvironment.cs
+++ b/Engine/Source/Programs/UnrealBuildTool/System/LinkEnvironment.cs
@@ -117,6 +117,7 @@ namespace UnrealBuildTool
 			ExcludedLibraries.AddRange(InCopyEnvironment.ExcludedLibraries);
 			AdditionalLibraries.AddRange(InCopyEnvironment.AdditionalLibraries);
 			Frameworks.AddRange(InCopyEnvironment.Frameworks);
+			AdditionalShadowFiles.AddRange( InCopyEnvironment.AdditionalShadowFiles );
 			AdditionalFrameworks.AddRange(InCopyEnvironment.AdditionalFrameworks);
 			WeakFrameworks.AddRange(InCopyEnvironment.WeakFrameworks);
 			DelayLoadDLLs.AddRange(InCopyEnvironment.DelayLoadDLLs);
