Fix framework name extraction logic

diff --git a/Engine/Source/Programs/UnrealBuildTool/IOS/IOSToolChain.cs b/Engine/Source/Programs/UnrealBuildTool/IOS/IOSToolChain.cs
--- a/Engine/Source/Programs/UnrealBuildTool/IOS/IOSToolChain.cs
+++ b/Engine/Source/Programs/UnrealBuildTool/IOS/IOSToolChain.cs
@@ -325,8 +325,14 @@ namespace UnrealBuildTool
 				// If we see ThirdPartyFrameworks in the string, assume we need to treat this as a path + name
 				if ( Framework.Contains( "ThirdPartyFrameworks/" ) )
 				{
-					string FrameworkName = Framework.Substring( Framework.LastIndexOf( '/' ) + 1 );
-					string FrameworkPath = Framework.Replace( FrameworkName, "" );
+					Int32 SourceIndex = Framework.LastIndexOf( '/' );
+					if ( SourceIndex == -1 )
+					{
+						throw new BuildException( "Invalid third party framework path" );
+					}
+
+					string FrameworkName = Framework.Substring( SourceIndex + 1 );
+					string FrameworkPath = Framework.Substring( 0, SourceIndex );
 					Result += " -F " + FrameworkPath;
 					Result += " -framework " + FrameworkName;
 					continue;
