Remove bogus compare in SetupEnvironmentFromSDK

diff --git a/Engine/Source/Programs/UnrealBuildTool/Configuration/UEBuildPlatform.cs b/Engine/Source/Programs/UnrealBuildTool/Configuration/UEBuildPlatform.cs
--- a/Engine/Source/Programs/UnrealBuildTool/Configuration/UEBuildPlatform.cs
+++ b/Engine/Source/Programs/UnrealBuildTool/Configuration/UEBuildPlatform.cs
@@ -490,8 +490,7 @@ namespace UnrealBuildTool
                     for (int i = 0; i < EnvVarNames.Count; ++i)
                     {
                         string EnvVarName = EnvVarNames[i];
-                        string EnvVarValue = EnvVarValues[i];
-                        if (String.Compare(EnvVarName, "strippath", true) == 0)
+                        string EnvVarValue = EnvVarValues[i];                        
                         if (BuildConfiguration.bPrintDebugInfo)
                         {
                             Console.WriteLine("Setting variable '{0}' to '{1}'", EnvVarName, EnvVarValue);
