Fixed iOS stub building

diff --git a/Engine/Source/Programs/UnrealBuildTool/System/XcodeProjectFileGenerator.cs b/Engine/Source/Programs/UnrealBuildTool/System/XcodeProjectFileGenerator.cs
--- a/Engine/Source/Programs/UnrealBuildTool/System/XcodeProjectFileGenerator.cs
+++ b/Engine/Source/Programs/UnrealBuildTool/System/XcodeProjectFileGenerator.cs
@@ -1297,7 +1297,8 @@ namespace UnrealBuildTool
 
 		private void WriteProject(StringBuilder XcodeProjectFileContent, List<XcodeProjectTarget> ProjectTargets, XcodeProjectTarget UE4ProjectTarget, XcodeProjectTarget UE4XcodeHelperTarget)
 		{
-			string PathBranch = bGeneratingRunIOSProject ? "/Engine/Intermediate/IOS" : "";
+			// @todo: Windows and Mac should store the project in the same folder
+			string PathBranch = (ExternalExecution.GetRuntimePlatform() != UnrealTargetPlatform.Mac) ? "/Engine/Intermediate/IOS" : "";
 			var XcodeProjectPath = MasterProjectRelativePath + PathBranch + "/" + MasterProjectName + ".xcodeproj";
 			if (bGeneratingGameProjectFiles && GameProjectName == "UE4Game")
 			{
