put in some output to catch the iOS build error

diff --git a/Engine/Source/Programs/UnrealBuildTool/IOS/IOSToolChain.cs b/Engine/Source/Programs/UnrealBuildTool/IOS/IOSToolChain.cs
--- a/Engine/Source/Programs/UnrealBuildTool/IOS/IOSToolChain.cs
+++ b/Engine/Source/Programs/UnrealBuildTool/IOS/IOSToolChain.cs
@@ -921,6 +921,8 @@ namespace UnrealBuildTool
 						                " -sdk iphoneos" +
 						                " CODE_SIGN_IDENTITY=\"iPhone Developer\"";
 
+                        Console.WriteLine("Code signing with command line: " + CmdLine);
+
 						Process SignProcess = new Process ();
 						SignProcess.StartInfo.WorkingDirectory = RemoteShadowDirectoryMac;
 						SignProcess.StartInfo.FileName = "/usr/bin/xcrun";
@@ -935,7 +937,7 @@ namespace UnrealBuildTool
 						// delete the temp project
 						if (Project.Contains ("_IOS.xcodeproj"))
 						{
-							Directory.Delete (Project, true);
+//							Directory.Delete (Project, true);
 						}
 
 						if (OutputReceivedDataEventHandlerEncounteredError)
