Changed Xcode project generator so it creates schemes for targets also for iOS deployment-only projects

diff --git a/Engine/Source/Programs/UnrealBuildTool/System/XcodeProjectFileGenerator.cs b/Engine/Source/Programs/UnrealBuildTool/System/XcodeProjectFileGenerator.cs
--- a/Engine/Source/Programs/UnrealBuildTool/System/XcodeProjectFileGenerator.cs
+++ b/Engine/Source/Programs/UnrealBuildTool/System/XcodeProjectFileGenerator.cs
@@ -1329,10 +1329,7 @@ namespace UnrealBuildTool
 				Log.TraceInformation("Skipping project file write, as it didn't change...");
 			}
 
-			if (!bGeneratingRunIOSProject)
-			{
-				WriteSchemes(XcodeProjectPath, ProjectTargets, UE4ProjectTarget, UE4XcodeHelperTarget);
-			}
+			WriteSchemes(XcodeProjectPath, ProjectTargets, UE4ProjectTarget, UE4XcodeHelperTarget);
 		}
 
 		/// <summary>
