delete temporary xcodeproj after signing

diff --git a/Engine/Source/Programs/UnrealBuildTool/IOS/IOSToolChain.cs b/Engine/Source/Programs/UnrealBuildTool/IOS/IOSToolChain.cs
--- a/Engine/Source/Programs/UnrealBuildTool/IOS/IOSToolChain.cs
+++ b/Engine/Source/Programs/UnrealBuildTool/IOS/IOSToolChain.cs
@@ -937,7 +937,7 @@ namespace UnrealBuildTool
 						// delete the temp project
 						if (Project.Contains ("_IOS.xcodeproj"))
 						{
-//							Directory.Delete (Project, true);
+							Directory.Delete (Project, true);
 						}
 
 						if (OutputReceivedDataEventHandlerEncounteredError)
