Fix for long paths in a Mac build.

#add Utils.CollapseRelativeDirectories function, translated from FPaths.
#change MacPath now uses CollapseRelativeDirectories instead of GetFullPath.

diff --git a/Engine/Source/Programs/UnrealBuildTool/System/Utils.cs b/Engine/Source/Programs/UnrealBuildTool/System/Utils.cs
--- a/Engine/Source/Programs/UnrealBuildTool/System/Utils.cs
+++ b/Engine/Source/Programs/UnrealBuildTool/System/Utils.cs
@@ -782,6 +782,60 @@ namespace UnrealBuildTool
 		}
 
 
+		/// <summary>
+		/// Correctly collapses any ../ or ./ entries in a path.
+		/// </summary>
+		/// <param name="InPath">The path to be collapsed</param>
+		/// <returns>true if the path could be collapsed, false otherwise.</returns>
+		public static bool CollapseRelativeDirectories(ref string InPath)
+		{
+			string ParentDir       = "/..";
+			int    ParentDirLength = ParentDir.Length;
+
+			for (;;)
+			{
+				// An empty path is finished
+				if (string.IsNullOrEmpty(InPath))
+					break;
+
+				// Consider empty paths or paths which start with .. or /.. as invalid
+				if (InPath.StartsWith("..") || InPath.StartsWith(ParentDir))
+					return false;
+
+				// If there are no "/.."s left then we're done
+				int Index = InPath.IndexOf(ParentDir);
+				if (Index == -1)
+					break;
+
+				int PreviousSeparatorIndex = Index;
+				for (;;)
+				{
+					// Find the previous slash
+					PreviousSeparatorIndex = Math.Max(0, InPath.LastIndexOf("/", PreviousSeparatorIndex - 1));
+
+					// Stop if we've hit the start of the string
+					if (PreviousSeparatorIndex == 0)
+						break;
+
+					// Stop if we've found a directory that isn't "/./"
+					if ((Index - PreviousSeparatorIndex) > 1 && (InPath[PreviousSeparatorIndex + 1] != '.' || InPath[PreviousSeparatorIndex + 2] != '/'))
+						break;
+				}
+
+				// If we're attempting to remove the drive letter, that's illegal
+				int Colon = InPath.IndexOf(":", PreviousSeparatorIndex);
+				if (Colon >= 0 && Colon < Index)
+					return false;
+
+				InPath = InPath.Substring(0, PreviousSeparatorIndex) + InPath.Substring(Index + ParentDirLength);
+			}
+
+			InPath = InPath.Replace("./", "");
+
+			return true;
+		}
+
+
 		/// <summary>
 		/// Finds the Engine Version from ObjVersion.cpp.
 		/// </summary>
diff --git a/Engine/Source/Programs/UnrealBuildTool/ToolChain/RemoteToolChain.cs b/Engine/Source/Programs/UnrealBuildTool/ToolChain/RemoteToolChain.cs
--- a/Engine/Source/Programs/UnrealBuildTool/ToolChain/RemoteToolChain.cs
+++ b/Engine/Source/Programs/UnrealBuildTool/ToolChain/RemoteToolChain.cs
@@ -147,7 +147,7 @@ namespace UnrealBuildTool
 
 				try
 				{
-					MacPath = Path.GetFullPath(MacPath);
+					Utils.CollapseRelativeDirectories(ref MacPath);
 				}
 				catch (Exception Ex)
 				{
