Moved telemetry provider for ET to NotForLicensees

diff --git a/Engine/Source/Programs/UnrealBuildTool/UnrealBuildTool.csproj b/Engine/Source/Programs/UnrealBuildTool/UnrealBuildTool.csproj
--- a/Engine/Source/Programs/UnrealBuildTool/UnrealBuildTool.csproj
+++ b/Engine/Source/Programs/UnrealBuildTool/UnrealBuildTool.csproj
@@ -176,7 +176,7 @@
   <Choose>
     <When Condition="Exists('System\NoRedist\TelemetryProviderET.cs')">
       <ItemGroup>
-        <Compile Include="System\NoRedist\TelemetryProviderET.cs" />
+        <Compile Include="System\NotForLicensees\TelemetryProviderET.cs" />
       </ItemGroup>
     </When>
   </Choose>
@@ -289,6 +289,9 @@
   <ItemGroup>
     <None Include="app.config" />
   </ItemGroup>
+  <ItemGroup>
+    <Folder Include="System\NoRedist\" />
+  </ItemGroup>
   <Import Project="$(MSBuildBinPath)\Microsoft.CSharp.targets" />
   <!-- To modify your build process, add your task inside one of the targets below and uncomment it. 
        Other similar extension points exist, see Microsoft.Common.targets.
