- Fixed "latest" IOS sdk version with weirdly formatted SDK dirs

diff --git a/Engine/Source/Programs/UnrealBuildTool/IOS/IOSToolChain.cs b/Engine/Source/Programs/UnrealBuildTool/IOS/IOSToolChain.cs
--- a/Engine/Source/Programs/UnrealBuildTool/IOS/IOSToolChain.cs
+++ b/Engine/Source/Programs/UnrealBuildTool/IOS/IOSToolChain.cs
@@ -117,7 +117,16 @@ namespace UnrealBuildTool
 						{
 							// get the SDK version from the directory name
 							string SDKString = SubDirName.Replace("iPhoneOS", "");
-                            float SDKVersion = float.Parse(SDKString, System.Globalization.CultureInfo.InvariantCulture);
+							float SDKVersion = 0.0f;
+							try
+							{
+								SDKVersion = float.Parse(SDKString, System.Globalization.CultureInfo.InvariantCulture);
+							}
+							catch (Exception)
+							{
+								// weirdly formatted SDKs
+								continue;
+							}
 
 							// update largest SDK version number
 							if (SDKVersion > MaxSDKVersion)
