fix for stub file not including the signed executable and plist file
#ue4
#ios

diff --git a/Engine/Build/IOS/XcodeSupportFiles/prepackage.sh b/Engine/Build/IOS/XcodeSupportFiles/prepackage.sh
--- a/Engine/Build/IOS/XcodeSupportFiles/prepackage.sh
+++ b/Engine/Build/IOS/XcodeSupportFiles/prepackage.sh
@@ -56,7 +56,13 @@ then
   then
     cp -R $1-Info.plist ../../../../$1/Intermediate/IOS/$1-Info.plist
   else
-    mkdir -p ../../../../Engine/Binaries/IOS/Engine/Intermediate/IOS/
-    cp -R $1-Info.plist ../../../../Engine/Binaries/IOS/Engine/Intermediate/IOS/$1-Info.plist
+	if [ "$1" == "UE4Game" ]
+	then
+      mkdir -p ../../../../Engine/Intermediate/IOS/
+      cp -R $1-Info.plist ../../../../Engine/Intermediate/IOS/$1-Info.plist
+	else 
+      mkdir -p ../../../../$1/Intermediate/IOS/
+      cp -R $1-Info.plist ../../../../$1/Intermediate/IOS/$1-Info.plist
+	fi
   fi
 fi
\ No newline at end of file
diff --git a/Engine/Source/Programs/UnrealBuildTool/IOS/IOSToolChain.cs b/Engine/Source/Programs/UnrealBuildTool/IOS/IOSToolChain.cs
--- a/Engine/Source/Programs/UnrealBuildTool/IOS/IOSToolChain.cs
+++ b/Engine/Source/Programs/UnrealBuildTool/IOS/IOSToolChain.cs
@@ -1026,7 +1026,7 @@ namespace UnrealBuildTool
 					if (!bUseDangerouslyFastMode)
 					{
 						// generate the dummy project so signing works
-						UnrealBuildTool.GenerateProjectFiles(new XcodeProjectFileGenerator(), new string[] { "-NoIntellIsense", "-iosdeployonly" });
+						UnrealBuildTool.GenerateProjectFiles(new XcodeProjectFileGenerator(), new string[] { "-NoIntellisense", "-iosdeployonly" });
 					}
 
 					Process StubGenerateProcess = new Process();
diff --git a/Engine/Source/Programs/UnrealBuildTool/System/XcodeProjectFileGenerator.cs b/Engine/Source/Programs/UnrealBuildTool/System/XcodeProjectFileGenerator.cs
--- a/Engine/Source/Programs/UnrealBuildTool/System/XcodeProjectFileGenerator.cs
+++ b/Engine/Source/Programs/UnrealBuildTool/System/XcodeProjectFileGenerator.cs
@@ -682,6 +682,19 @@ namespace UnrealBuildTool
 			{
 				EngineRelative = Path.GetFullPath(EngineRelativePath + "/../");
 			}
+			else if (ExternalExecution.GetRuntimePlatform() != UnrealTargetPlatform.Mac)
+			{
+				EngineRelative = Path.GetFullPath(EngineRelativePath + "/../../../../");
+				if (EngineRelative.IndexOf(':') == 1)
+				{
+					EngineRelative = EngineRelative.Substring(2);
+				}
+				if (Target.TargetPlatform == UnrealTargetPlatform.IOS)
+				{
+					UEToolChain Toolchain = UEToolChain.GetPlatformToolChain(CPPTargetPlatform.IOS);
+					EngineRelative = Toolchain.ConvertPath(EngineRelative);
+				}
+			}
 
 			if (!bGeneratingRocketProjectFiles)
 			{
