Restored creating of references to source code folders in deploy-only iOS Xcode projects

diff --git a/Engine/Source/Programs/UnrealBuildTool/System/XcodeProjectFileGenerator.cs b/Engine/Source/Programs/UnrealBuildTool/System/XcodeProjectFileGenerator.cs
--- a/Engine/Source/Programs/UnrealBuildTool/System/XcodeProjectFileGenerator.cs
+++ b/Engine/Source/Programs/UnrealBuildTool/System/XcodeProjectFileGenerator.cs
@@ -885,7 +885,23 @@ namespace UnrealBuildTool
 				XcodeProject.UE4CmdLineRunFileGuid = UE4CmdLineRunMFileGuid;
 				XcodeProject.UE4CmdLineRunFileRefGuid = UE4CmdLineRunMFileRefGuid;
 
-				XcodeProject.GenerateSectionsContents(ref PBXBuildFileSection, ref PBXFileReferenceSection, ref PBXSourcesBuildPhaseSection, ref Groups);
+				if (bGeneratingRunIOSProject)
+				{
+					foreach (var CurSourceFile in XcodeProject.SourceFiles)
+					{
+						XcodeSourceFile SourceFile = CurSourceFile as XcodeSourceFile;
+						string GroupPath = Path.GetFullPath(Path.GetDirectoryName(SourceFile.FilePath));
+						XcodeFileGroup Group = XcodeProject.FindGroupByFullPath(ref Groups, GroupPath);
+						if (Group != null)
+						{
+							Group.bReference = true;
+						}
+					}
+				}
+				else
+				{
+					XcodeProject.GenerateSectionsContents(ref PBXBuildFileSection, ref PBXFileReferenceSection, ref PBXSourcesBuildPhaseSection, ref Groups);
+				}
 
 				foreach (var CurPath in XcodeProject.IntelliSenseIncludeSearchPaths)
 				{
@@ -923,6 +939,23 @@ namespace UnrealBuildTool
 			}
 		}
 
+		private void AppendReferenceGroups(ref string PBXFileReferenceSection, ref Dictionary<string, XcodeFileGroup> Groups)
+		{
+			string RelativePath = "../../../";
+
+			foreach (XcodeFileGroup Group in Groups.Values)
+			{
+				Group.bReference = true;
+				string GroupPath = RelativePath + Group.GroupPath;
+
+				// Add File reference.
+				PBXFileReferenceSection += "\t\t" + Group.GroupGuid + " /* " + Group.GroupName + " */ = {"
+											+ "isa = PBXFileReference; lastKnownFileType = folder; "
+											+ "name = " + Group.GroupName + "; "
+											+ "path = " + Utils.CleanDirectorySeparators(GroupPath, '/') + "; sourceTree = \"<group>\"; };" + ProjectFileGenerator.NewLine;
+			}
+		}
+
 		/// <summary>
 		/// Creates the container item proxy section to be added to the PBX Project.
 		/// </summary>
@@ -1429,11 +1462,8 @@ namespace UnrealBuildTool
 			AppendBuildPhases(ProjectTargets, ref PBXBuildFileSection, ref PBXFileReferenceSection, ref PBXSourcesBuildPhaseSection, ref PBXResourcesBuildPhaseSection, ref PBXFrameworksBuildPhaseSection,
 				ref PBXShellScriptBuildPhaseSection, TestFrameworkFiles, UE4CmdLineRunMFileGuid);
 
-			if (!bGeneratingRunIOSProject)
-			{
-				AppendSourceFiles(ProjectTargets, ref PBXBuildFileSection, ref PBXFileReferenceSection, ref PBXSourcesBuildPhaseSection, ref Groups, UE4XcodeHelperTarget, UE4CmdLineRunMFileGuid,
-					ref IncludeDirectories, ref SystemIncludeDirectories, ref PreprocessorDefinitions);
-			}
+			AppendSourceFiles(ProjectTargets, ref PBXBuildFileSection, ref PBXFileReferenceSection, ref PBXSourcesBuildPhaseSection, ref Groups, UE4XcodeHelperTarget, UE4CmdLineRunMFileGuid,
+				ref IncludeDirectories, ref SystemIncludeDirectories, ref PreprocessorDefinitions);
 
 			foreach (XcodeProjectTarget Target in ProjectTargets)
 			{
@@ -1454,6 +1484,11 @@ namespace UnrealBuildTool
 				}
 			}
 
+			if (bGeneratingRunIOSProject)
+			{
+				AppendReferenceGroups(ref PBXFileReferenceSection, ref Groups);
+			}
+
 			PBXBuildFileSection += "/* End PBXBuildFile section */" + ProjectFileGenerator.NewLine + ProjectFileGenerator.NewLine;
 			PBXFileReferenceSection += "/* End PBXFileReference section */" + ProjectFileGenerator.NewLine + ProjectFileGenerator.NewLine;
 
