Fixed UE_ENGINE_DIRECTORY on Mac

diff --git a/Engine/Source/Programs/UnrealBuildTool/Configuration/UEBuildTarget.cs b/Engine/Source/Programs/UnrealBuildTool/Configuration/UEBuildTarget.cs
--- a/Engine/Source/Programs/UnrealBuildTool/Configuration/UEBuildTarget.cs
+++ b/Engine/Source/Programs/UnrealBuildTool/Configuration/UEBuildTarget.cs
@@ -1304,7 +1304,12 @@ namespace UnrealBuildTool
                 if (TargetRules.IsAGame(Rules.Type) && IsCurrentPlatform)
 				{
 					// The hardcoded engine directory needs to be a relative path to match the normal EngineDir format. Not doing so breaks the network file system (TTP#315861).
-					string EnginePath = Utils.CleanDirectorySeparators(Utils.MakePathRelativeTo(ProjectFileGenerator.EngineRelativePath, Path.GetDirectoryName(ExecutableBinary.Config.OutputFilePath)), '/');
+					string OutputFilePath = ExecutableBinary.Config.OutputFilePath;
+					if (Platform == UnrealTargetPlatform.Mac && OutputFilePath.Contains(".app/Contents/MacOS"))
+					{
+						OutputFilePath = OutputFilePath.Substring(0, OutputFilePath.LastIndexOf(".app/Contents/MacOS") + 4);
+					}
+					string EnginePath = Utils.CleanDirectorySeparators(Utils.MakePathRelativeTo(ProjectFileGenerator.EngineRelativePath, Path.GetDirectoryName(OutputFilePath)), '/');
 					if (EnginePath.EndsWith("/") == false)
 					{
 						EnginePath += "/";
