Merging the bulk of Darkside changes.

- New target platforms (LinuxNoEditor for client+server, LinuxClient for client, Linux for editor, LinuxServer for dedicated server - mirroring other platforms).
- Preliminary support for gcc in the toolchain (and some safe fixes in code).
- No more lowercasing when accessing files on Linux, tread carefully.
- Changes needed for native compilation and other minor improvements.

diff --git a/Engine/Build/BatchFiles/Linux/Build.sh b/Engine/Build/BatchFiles/Linux/Build.sh
--- a/Engine/Build/BatchFiles/Linux/Build.sh
+++ b/Engine/Build/BatchFiles/Linux/Build.sh
@@ -2,6 +2,10 @@
 
 # This script gets can be used to build clean individual projects using UnrealBuildTool
 
+set -e
+
+cd "`dirname "$0"`/../../../.." 
+
 # First make sure that the UnrealBuildTool is up-to-date
 if ! xbuild /property:Configuration=Development /verbosity:quiet /nologo Engine/Source/Programs/UnrealBuildTool/UnrealBuildTool_Mono.csproj; then
   echo "Failed to build to build tool (UnrealBuildTool)"
diff --git a/Engine/Build/BatchFiles/Linux/GenerateProjectFiles.sh b/Engine/Build/BatchFiles/Linux/GenerateProjectFiles.sh
--- a/Engine/Build/BatchFiles/Linux/GenerateProjectFiles.sh
+++ b/Engine/Build/BatchFiles/Linux/GenerateProjectFiles.sh
@@ -18,9 +18,10 @@ located inside the Engine/Build/BatchFiles/Linux directory."
 fi
 
 
-if [ "$(lsb_release --id)" = "Distributor ID:	Ubuntu" ]; then
+if [ "$(lsb_release --id)" = "Distributor ID:	Ubuntu" -o "$(lsb_release --id)" = "Distributor ID:	Debian" ]; then
   # Install all necessary dependencies
   DEPS="mono-xbuild \
+    mono-dmcs \
     libmono-microsoft-build-tasks-v4.0-4.0-cil \
     libmono-system-data-datasetextensions4.0-cil
     libmono-system-web-extensions4.0-cil"
diff --git a/Engine/Source/Developer/Linux/LinuxNoEditorTargetPlatform/LinuxNoEditorTargetPlatform.Build.cs b/Engine/Source/Developer/Linux/LinuxNoEditorTargetPlatform/LinuxNoEditorTargetPlatform.Build.cs
new file mode 100644
--- /dev/null
+++ b/Engine/Source/Developer/Linux/LinuxNoEditorTargetPlatform/LinuxNoEditorTargetPlatform.Build.cs
@@ -0,0 +1,34 @@
+// Copyright 1998-2014 Epic Games, Inc. All Rights Reserved.
+
+using UnrealBuildTool;
+
+public class LinuxNoEditorTargetPlatform : ModuleRules
+{
+    public LinuxNoEditorTargetPlatform(TargetInfo Target)
+    {
+        BinariesSubFolder = "Linux";
+
+        PrivateDependencyModuleNames.AddRange(
+            new string[] {
+				"Core",
+				"TargetPlatform"
+			}
+        );
+
+        if (UEBuildConfiguration.bCompileAgainstEngine)
+        {
+            PrivateDependencyModuleNames.AddRange(new string[] {
+				"Engine"
+				}
+            );
+
+            PrivateIncludePathModuleNames.Add("TextureCompressor");
+        }
+
+        PrivateIncludePaths.AddRange(
+            new string[] {
+				"Developer/Linux/LinuxTargetPlatform/Private"
+			}
+        );
+    }
+}
diff --git a/Engine/Source/Developer/Linux/LinuxNoEditorTargetPlatform/Private/LinuxNoEditorTargetPlatformModule.cpp b/Engine/Source/Developer/Linux/LinuxNoEditorTargetPlatform/Private/LinuxNoEditorTargetPlatformModule.cpp
new file mode 100644
--- /dev/null
+++ b/Engine/Source/Developer/Linux/LinuxNoEditorTargetPlatform/Private/LinuxNoEditorTargetPlatformModule.cpp
@@ -0,0 +1,41 @@
+// Copyright 1998-2014 Epic Games, Inc. All Rights Reserved.
+
+/*=============================================================================
+	LinuxNoEditorTargetPlatformModule.cpp: Implements the FLinuxNoEditorTargetPlatformModule class.
+=============================================================================*/
+
+#include "LinuxNoEditorTargetPlatformPrivatePCH.h"
+
+
+/**
+ * Holds the target platform singleton.
+ */
+static ITargetPlatform* Singleton = NULL;
+
+
+/**
+ * Module for the Linux target platform (without editor).
+ */
+class FLinuxNoEditorTargetPlatformModule
+	: public ITargetPlatformModule
+{
+public:
+
+	virtual ~FLinuxNoEditorTargetPlatformModule( )
+	{
+		Singleton = NULL;
+	}
+
+	virtual ITargetPlatform* GetTargetPlatform( )
+	{
+		if (Singleton == NULL)
+		{
+			Singleton = new TLinuxTargetPlatform<false, false, false>();
+		}
+
+		return Singleton;
+	}
+};
+
+
+IMPLEMENT_MODULE(FLinuxNoEditorTargetPlatformModule, LinuxNoEditorTargetPlatform);
\ No newline at end of file
diff --git a/Engine/Source/Developer/Linux/LinuxNoEditorTargetPlatform/Private/LinuxNoEditorTargetPlatformPrivatePCH.h b/Engine/Source/Developer/Linux/LinuxNoEditorTargetPlatform/Private/LinuxNoEditorTargetPlatformPrivatePCH.h
new file mode 100644
--- /dev/null
+++ b/Engine/Source/Developer/Linux/LinuxNoEditorTargetPlatform/Private/LinuxNoEditorTargetPlatformPrivatePCH.h
@@ -0,0 +1,29 @@
+// Copyright 1998-2014 Epic Games, Inc. All Rights Reserved.
+
+/*=============================================================================
+	LinuxNoEditorTargetPlatformPrivatePCH.h: Pre-compiled header file for the LinuxNoEditorTargetPlatform module.
+=============================================================================*/
+
+#pragma once
+
+
+/* Dependencies
+ *****************************************************************************/
+
+#include "Core.h"
+#include "ModuleManager.h"
+#include "Runtime/Core/Public/Linux/LinuxPlatformProperties.h"
+
+#if WITH_ENGINE
+	#include "Engine.h"
+	#include "TextureCompressorModule.h"
+#endif
+
+#include "TargetPlatform.h"
+
+
+/* Private includes
+ *****************************************************************************/
+
+#include "LinuxTargetDevice.h"
+#include "LinuxTargetPlatform.h"
diff --git a/Engine/Source/Developer/Linux/LinuxNoEditorTargetPlatform/Private/LinuxTargetDevice.cpp b/Engine/Source/Developer/Linux/LinuxNoEditorTargetPlatform/Private/LinuxTargetDevice.cpp
new file mode 100644
--- /dev/null
+++ b/Engine/Source/Developer/Linux/LinuxNoEditorTargetPlatform/Private/LinuxTargetDevice.cpp
@@ -0,0 +1,68 @@
+// Copyright 1998-2013 Epic Games, Inc. All Rights Reserved.
+
+/*=============================================================================
+	LinuxTargetDevice.cpp: Implements the LinuxTargetDevice class.
+=============================================================================*/
+
+#include "LinuxNoEditorTargetPlatformPrivatePCH.h"
+
+#include "LinuxTargetDevice.h"
+
+
+
+/* ITargetDevice interface
+ *****************************************************************************/
+
+bool FLinuxTargetDevice::Deploy( const FString& SourceFolder, FString& OutAppId )
+{
+	return false;
+}
+
+int32 FLinuxTargetDevice::GetProcessSnapshot( TArray<FTargetDeviceProcessInfo>& OutProcessInfos )
+{
+	return 0;
+}
+
+bool FLinuxTargetDevice::Launch( const FString& AppId, EBuildConfigurations::Type BuildConfiguration, EBuildTargets::Type BuildTarget, const FString& Params, uint32* OutProcessId )
+{
+	return false;
+}
+
+bool FLinuxTargetDevice::Reboot( bool bReconnect)
+{
+	return false;
+}
+
+bool FLinuxTargetDevice::Run( const FString& ExecutablePath, const FString& Params, uint32* OutProcessId )
+{
+	return false;
+}
+
+bool FLinuxTargetDevice::SupportsFeature( ETargetDeviceFeatures::Type Feature ) const
+{
+	return false;
+}
+
+bool FLinuxTargetDevice::SupportsSdkVersion( const FString& VersionString ) const
+{
+	return true;
+}
+
+void FLinuxTargetDevice::SetUserCredentials( const FString & InUserName, const FString & InUserPassword )
+{
+	UserName = InUserName;
+	UserPassword = InUserPassword;
+}
+
+bool FLinuxTargetDevice::GetUserCredentials( FString & OutUserName, FString & OutUserPassword )
+{
+	OutUserName = UserName;
+	OutUserPassword = UserPassword;
+	return true;
+}
+
+bool FLinuxTargetDevice::TerminateProcess( const int32 ProcessId )
+{
+	return false;
+}
+
diff --git a/Engine/Source/Developer/Linux/LinuxServerTargetPlatform/Private/LinuxServerTargetPlatformModule.cpp b/Engine/Source/Developer/Linux/LinuxServerTargetPlatform/Private/LinuxServerTargetPlatformModule.cpp
--- a/Engine/Source/Developer/Linux/LinuxServerTargetPlatform/Private/LinuxServerTargetPlatformModule.cpp
+++ b/Engine/Source/Developer/Linux/LinuxServerTargetPlatform/Private/LinuxServerTargetPlatformModule.cpp
@@ -30,7 +30,7 @@ public:
 	{
 		if (Singleton == NULL)
 		{
-			Singleton = new TLinuxTargetPlatform<true, false>();
+			Singleton = new TLinuxTargetPlatform<false, true, false>();
 		}
 		
 		return Singleton;
diff --git a/Engine/Source/Developer/Linux/LinuxTargetPlatform/Private/LinuxTargetPlatform.h b/Engine/Source/Developer/Linux/LinuxTargetPlatform/Private/LinuxTargetPlatform.h
--- a/Engine/Source/Developer/Linux/LinuxTargetPlatform/Private/LinuxTargetPlatform.h
+++ b/Engine/Source/Developer/Linux/LinuxTargetPlatform/Private/LinuxTargetPlatform.h
@@ -10,13 +10,13 @@
 /**
  * Template for Linux target platforms
  */
-template<bool IS_DEDICATED_SERVER, bool IS_CLIENT_ONLY>
+template<bool HAS_EDITOR_DATA, bool IS_DEDICATED_SERVER, bool IS_CLIENT_ONLY>
 class TLinuxTargetPlatform
-	: public TTargetPlatformBase<FLinuxPlatformProperties<IS_DEDICATED_SERVER, IS_CLIENT_ONLY> >
+	: public TTargetPlatformBase<FLinuxPlatformProperties<HAS_EDITOR_DATA, IS_DEDICATED_SERVER, IS_CLIENT_ONLY> >
 {
 public:
 
-	typedef public TTargetPlatformBase<FLinuxPlatformProperties<IS_DEDICATED_SERVER, IS_CLIENT_ONLY> > TSuper;
+	typedef public TTargetPlatformBase<FLinuxPlatformProperties<HAS_EDITOR_DATA, IS_DEDICATED_SERVER, IS_CLIENT_ONLY> > TSuper;
 
 	/**
 	 * Default constructor.
@@ -110,7 +110,7 @@ public:
 			return true;
 		}
 
-		return TTargetPlatformBase<FLinuxPlatformProperties<IS_DEDICATED_SERVER, IS_CLIENT_ONLY>>::SupportsFeature(Feature);
+		return TTargetPlatformBase<FLinuxPlatformProperties<HAS_EDITOR_DATA, IS_DEDICATED_SERVER, IS_CLIENT_ONLY>>::SupportsFeature(Feature);
 	}
 
 #if WITH_ENGINE
diff --git a/Engine/Source/Developer/Linux/LinuxTargetPlatform/Private/LinuxTargetPlatformModule.cpp b/Engine/Source/Developer/Linux/LinuxTargetPlatform/Private/LinuxTargetPlatformModule.cpp
--- a/Engine/Source/Developer/Linux/LinuxTargetPlatform/Private/LinuxTargetPlatformModule.cpp
+++ b/Engine/Source/Developer/Linux/LinuxTargetPlatform/Private/LinuxTargetPlatformModule.cpp
@@ -41,7 +41,7 @@ public:
 	{
 		if (Singleton == NULL)
 		{
-			Singleton = new TLinuxTargetPlatform<false, false>();
+			Singleton = new TLinuxTargetPlatform<true, false, false>();
 		}
 		
 		return Singleton;
diff --git a/Engine/Source/Programs/AutomationTool/Linux/LinuxPlatform.Automation.cs b/Engine/Source/Programs/AutomationTool/Linux/LinuxPlatform.Automation.cs
--- a/Engine/Source/Programs/AutomationTool/Linux/LinuxPlatform.Automation.cs
+++ b/Engine/Source/Programs/AutomationTool/Linux/LinuxPlatform.Automation.cs
@@ -88,10 +88,23 @@ public abstract class BaseLinuxPlatform : Platform
 
 	public override string GetCookPlatform(bool bDedicatedServer, bool bIsClientOnly, string CookFlavor)
 	{
-        const string ClientCookPlatform = "Linux";
+        const string NoEditorCookPlatform = "LinuxNoEditor";
         const string ServerCookPlatform = "LinuxServer";
-        return bDedicatedServer ? ServerCookPlatform : ClientCookPlatform;
-	}
+        const string ClientCookPlatform = "LinuxClient";
+
+        if (bDedicatedServer)
+        {
+            return ServerCookPlatform;
+        }
+        else if (bIsClientOnly)
+        {
+            return ClientCookPlatform;
+        }
+        else
+        {
+            return NoEditorCookPlatform;
+        }
+    }
 
     /// <summary>
     /// return true if we need to change the case of filenames outside of pak files
@@ -99,7 +112,7 @@ public abstract class BaseLinuxPlatform : Platform
     /// <returns></returns>
     public override bool DeployLowerCaseFilenames(bool bUFSFile)
     {
-        return true;
+        return false;
     }
 
     /// <summary>
@@ -117,8 +130,8 @@ public abstract class BaseLinuxPlatform : Platform
 
             List<string> Exes = GetExecutableNames(SC);
 
-            string binPath = CombinePaths(GetCookPlatform(Params.DedicatedServer, false, ""), SC.RelativeProjectRootForStage.ToLower(), "binaries", SC.PlatformDir.ToLower(), Path.GetFileName(Exes[0]).ToLower()).Replace("\\", "/");
-            string iconPath = CombinePaths(GetCookPlatform(Params.DedicatedServer, false, ""), SC.RelativeProjectRootForStage.ToLower(), SC.ShortProjectName.ToLower() + ".png").Replace("\\", "/");
+            string binPath = CombinePaths(GetCookPlatform(Params.DedicatedServer, false, ""), SC.RelativeProjectRootForStage, "Binaries", SC.PlatformDir, Path.GetFileName(Exes[0])).Replace("\\", "/");
+            string iconPath = CombinePaths(GetCookPlatform(Params.DedicatedServer, false, ""), SC.RelativeProjectRootForStage, SC.ShortProjectName + ".png").Replace("\\", "/");
 
             string DesiredGLVersion = "4.3";
 
@@ -147,7 +160,7 @@ EOF
 
 # Set permissions
 chmod 755 $HOME/{3}
-chmod 700 $HOME/Desktop/{1}.desktop", DesiredGLVersion, SC.ShortProjectName.ToLower(), SC.ShortProjectName, binPath, (Params.Pak ? " -pak" : ""), iconPath);
+chmod 700 $HOME/Desktop/{1}.desktop", DesiredGLVersion, SC.ShortProjectName, SC.ShortProjectName, binPath, (Params.Pak ? " -pak" : ""), iconPath);
             // End Bash Shell Script
 
             string scriptFile = Path.GetTempFileName();
diff --git a/Engine/Source/Programs/UnrealBuildTool/Linux/LinuxToolChain.cs b/Engine/Source/Programs/UnrealBuildTool/Linux/LinuxToolChain.cs
--- a/Engine/Source/Programs/UnrealBuildTool/Linux/LinuxToolChain.cs
+++ b/Engine/Source/Programs/UnrealBuildTool/Linux/LinuxToolChain.cs
@@ -16,14 +16,38 @@ namespace UnrealBuildTool.Linux
             return ExternalExecution.GetRuntimePlatform() != UnrealTargetPlatform.Linux;
         }
 
+        private string Which(string name)
+        {
+            Process proc = new Process();
+            proc.StartInfo.FileName = "/bin/sh";
+            proc.StartInfo.Arguments = String.Format("-c 'which {0}'", name);
+            proc.StartInfo.UseShellExecute = false;
+            proc.StartInfo.CreateNoWindow = true;
+            proc.StartInfo.RedirectStandardOutput = true;
+            proc.StartInfo.RedirectStandardError = true;
+
+            proc.Start();
+            proc.WaitForExit();
+
+            string path = proc.StandardOutput.ReadLine();
+            Console.WriteLine(String.Format("which {0} result: ({1}) {2}", name, proc.ExitCode, path));
+
+            if (proc.ExitCode == 0 && String.IsNullOrEmpty(proc.StandardError.ReadToEnd()))
+            {
+                return path;
+            }
+            return null;
+        }
+
         public override void RegisterToolChain()
         {
             if (!CrossCompiling())
             {
                 // use native linux toolchain
-                ClangPath = "clang++";
-                ArPath = "ar";
-                RanlibPath = "ranlib";
+                ClangPath = Which("clang++");
+                GCCPath = Which("g++");
+                ArPath = Which("ar");
+                RanlibPath = Which("ranlib");
             }
             else
             {
@@ -71,21 +95,39 @@ namespace UnrealBuildTool.Linux
             //Result += " -Wshadow";                      // additional warning not normally included in Wall: warns if there variable/typedef shadows some other variable - not helpful because we have gobs of code that shadows variables
             Result += " -mmmx -msse -msse2";            // allows use of SIMD intrinsics
             Result += " -fno-math-errno";               // do not assume that math ops have side effects
-            Result += " -fdiagnostics-format=msvc";     // make diagnostics compatible with MSVC
 
-            Result += " -Wno-unused-variable";
-            if (!CrossCompiling())
+            if (String.IsNullOrEmpty(ClangPath))
             {
-                Result += " -Wno-logical-op-parentheses";   // needed for external headers we can't change
+                // GCC only option
+                Result += " -fno-strict-aliasing";
+                Result += " -Wno-sign-compare"; // needed to suppress: comparison between signed and unsigned integer expressions
+                Result += " -Wno-enum-compare"; // Stats2.h triggers this (ALIGNOF(int64) <= DATA_ALIGN)
+                Result += " -Wno-return-type"; // Variant.h triggers this
+                Result += " -Wno-unused-local-typedefs";
+                Result += " -Wno-multichar";
+                Result += " -Wno-unused-but-set-variable";
+                Result += " -Wno-sequence-point"; // TaskGraph.cpp:755 FTaskThread* Target = Target = &Thread(ThreadToExecuteOn);
+                Result += " -Wno-strict-overflow"; // Array.h:518
             }
+            else
+            {
+                // Clang only options
+                Result += " -fdiagnostics-format=msvc";     // make diagnostics compatible with MSVC
+                Result += " -Wno-unused-private-field";     // MultichannelTcpSocket.h triggers this, possibly more
+                // this hides the "warning : comparison of unsigned expression < 0 is always false" type warnings due to constant comparisons, which are possible with template arguments
+                Result += " -Wno-tautological-compare";
+                if (!CrossCompiling())
+                {
+                    Result += " -Wno-logical-op-parentheses";   // needed for external headers we can't change
+                }
+            }
+
+            Result += " -Wno-unused-variable";
             // this will hide the warnings about static functions in headers that aren't used in every single .cpp file
             Result += " -Wno-unused-function";
             // this hides the "enumeration value 'XXXXX' not handled in switch [-Wswitch]" warnings - we should maybe remove this at some point and add UE_LOG(, Fatal, ) to default cases
             Result += " -Wno-switch";
-            // this hides the "warning : comparison of unsigned expression < 0 is always false" type warnings due to constant comparisons, which are possible with template arguments
-            Result += " -Wno-tautological-compare";
             Result += " -Wno-unknown-pragmas";			// Slate triggers this (with its optimize on/off pragmas)
-            Result += " -Wno-unused-private-field";     // MultichannelTcpSocket.h triggers this, possibly more
 			Result += " -Wno-invalid-offsetof"; // needed to suppress warnings about using offsetof on non-POD types.
 
             //Result += " -DOPERATOR_NEW_INLINE=FORCENOINLINE";
@@ -206,7 +248,7 @@ namespace UnrealBuildTool.Linux
             {
                 Result += " -shared";
             }
-            Result += " -v";
+            //Result += " -v";
 
             if (CrossCompiling())
             {
@@ -219,7 +261,7 @@ namespace UnrealBuildTool.Linux
             }
 
             // RPATH for third party libs
-            Result += " -Wl,-rpath=${ORIGIN}/../../../engine/binaries/linux";
+            Result += " -Wl,-rpath=${ORIGIN}/../../../Engine/Binaries/Linux";
 
             if (CrossCompiling())
             {
@@ -280,6 +322,7 @@ namespace UnrealBuildTool.Linux
         // cache the location of NDK tools
         static string BaseLinuxPath;
         static string ClangPath;
+        static string GCCPath;
         static string ArPath;
         static string RanlibPath;
 
@@ -400,7 +443,14 @@ namespace UnrealBuildTool.Linux
                 FileArguments += string.Format(" \"{0}\"", SourceFile.AbsolutePath);
 
                 CompileAction.WorkingDirectory = Path.GetFullPath(".");
-                CompileAction.CommandPath = ClangPath;
+                if (String.IsNullOrEmpty(ClangPath))
+                {
+                    CompileAction.CommandPath = GCCPath;
+                }
+                else
+                {
+                    CompileAction.CommandPath = ClangPath;
+                }
                 CompileAction.CommandArguments = Arguments + FileArguments + CompileEnvironment.Config.AdditionalArguments;
                 CompileAction.CommandDescription = "Compile";
                 CompileAction.StatusDescription = Path.GetFileName(SourceFile.AbsolutePath);
@@ -484,7 +534,14 @@ namespace UnrealBuildTool.Linux
             // Create an action that invokes the linker.
             Action LinkAction = new Action(ActionType.Link);
             LinkAction.WorkingDirectory = Path.GetFullPath(".");
-            LinkAction.CommandPath = ClangPath;
+            if (String.IsNullOrEmpty(ClangPath))
+            {
+                LinkAction.CommandPath = GCCPath;
+            }
+            else
+            {
+                LinkAction.CommandPath = ClangPath;
+            }
 
             // Get link arguments.
             LinkAction.CommandArguments = GetLinkArguments(LinkEnvironment);
@@ -516,7 +573,7 @@ namespace UnrealBuildTool.Linux
 			// Add the library paths to the argument list.
 			foreach (string LibraryPath in LinkEnvironment.Config.LibraryPaths)
 			{
-				LinkAction.CommandArguments += string.Format(" -L\"{0}\"", LibraryPath);
+				LinkAction.CommandArguments += string.Format(" -Wl,-L\"{0}\"", LibraryPath);
 			}
 
 			// add libraries in a library group
@@ -552,5 +609,18 @@ namespace UnrealBuildTool.Linux
         {
             throw new BuildException("Linux cannot compile C# files");
         }
+
+        /** Converts the passed in path from UBT host to compiler native format. */
+        public override String ConvertPath(String OriginalPath)
+        {
+            if (ExternalExecution.GetRuntimePlatform() == UnrealTargetPlatform.Linux)
+            {
+                return OriginalPath.Replace("\\", "/");
+            }
+            else
+            {
+                return OriginalPath;
+            }
+        }
     }
 }
diff --git a/Engine/Source/Programs/UnrealBuildTool/Linux/UEBuildLinux.cs b/Engine/Source/Programs/UnrealBuildTool/Linux/UEBuildLinux.cs
--- a/Engine/Source/Programs/UnrealBuildTool/Linux/UEBuildLinux.cs
+++ b/Engine/Source/Programs/UnrealBuildTool/Linux/UEBuildLinux.cs
@@ -190,6 +190,7 @@ namespace UnrealBuildTool.Linux
         {
             // increase Unity size to avoid too long command lines
             BuildConfiguration.NumIncludedBytesPerUnityCPP = 1024 * 1024;
+            UEBuildConfiguration.bCompileICU = true;
         }
 
         /**
@@ -285,6 +286,7 @@ namespace UnrealBuildTool.Linux
                         if (UEBuildConfiguration.bBuildDeveloperTools)
                         {
                             InModule.AddPlatformSpecificDynamicallyLoadedModule("LinuxTargetPlatform");
+                            InModule.AddPlatformSpecificDynamicallyLoadedModule("LinuxNoEditorTargetPlatform");
                             InModule.AddPlatformSpecificDynamicallyLoadedModule("LinuxServerTargetPlatform");
                         }
                     }
@@ -294,6 +296,7 @@ namespace UnrealBuildTool.Linux
                 if (UEBuildConfiguration.bForceBuildTargetPlatforms)
                 {
                     InModule.AddPlatformSpecificDynamicallyLoadedModule("LinuxTargetPlatform");
+                    InModule.AddPlatformSpecificDynamicallyLoadedModule("LinuxNoEditorTargetPlatform");
                     InModule.AddPlatformSpecificDynamicallyLoadedModule("LinuxServerTargetPlatform");
                 }
             }
diff --git a/Engine/Source/Programs/UnrealBuildTool/System/MakefileGenerator.cs b/Engine/Source/Programs/UnrealBuildTool/System/MakefileGenerator.cs
--- a/Engine/Source/Programs/UnrealBuildTool/System/MakefileGenerator.cs
+++ b/Engine/Source/Programs/UnrealBuildTool/System/MakefileGenerator.cs
@@ -49,35 +49,138 @@ namespace UnrealBuildTool
 			return bSuccess;
 		}
 
+        private bool WriteMakefile()
+        {
+            var FileName = "Makefile"; // MasterProjectName + ".mk";
+            var MakefileContent = new StringBuilder();
+            MakefileContent.Append(
+                "# Makefile generated by MakefileGenerator.cs\n" +
+                "# *DO NOT EDIT*\n\n" +
+                "TARGETS ="
+                );
+            foreach (string Target in DiscoverTargets())
+            {
+                var Basename = Path.GetFileNameWithoutExtension(Target);
+                Basename = Path.GetFileNameWithoutExtension(Basename);
+                foreach (UnrealTargetConfiguration CurConfiguration in Enum.GetValues(typeof(UnrealTargetConfiguration)))
+                {
+                    if (CurConfiguration != UnrealTargetConfiguration.Unknown && CurConfiguration != UnrealTargetConfiguration.Development)
+                    {
+                        if (UnrealBuildTool.IsValidConfiguration(CurConfiguration))
+                        {
+                            var Confname = Enum.GetName(typeof(UnrealTargetConfiguration), CurConfiguration);
+                            MakefileContent.Append(String.Format(" \\\n    {0}-Linux-{1}", Basename, Confname));
+                        }
+                    }
+                }
+                MakefileContent.Append(" \\\n    " + Basename);
+            }
+            MakefileContent.Append("\n\nBUILD = Engine/Build/BatchFiles/Linux/Build.sh\n\n" +
+                "all: $(TARGETS)\n"
+                );
+            foreach (string Target in DiscoverTargets())
+            {
+                var Basename = Path.GetFileNameWithoutExtension(Target);
+                Basename = Path.GetFileNameWithoutExtension(Basename);
+                foreach (UnrealTargetConfiguration CurConfiguration in Enum.GetValues(typeof(UnrealTargetConfiguration)))
+                {
+                    if (CurConfiguration != UnrealTargetConfiguration.Unknown && CurConfiguration != UnrealTargetConfiguration.Development)
+                    {
+                        if (UnrealBuildTool.IsValidConfiguration(CurConfiguration))
+                        {
+                            var Confname = Enum.GetName(typeof(UnrealTargetConfiguration), CurConfiguration);
+                            MakefileContent.Append(String.Format("\n{0}-Linux-{1}:\n\t$(BUILD) {0} Linux {1} $(ARGS)\n", Basename, Confname));
+                        }
+                    }
+                }
+
+                MakefileContent.Append(String.Format("\n{0}:\n\t$(BUILD) {0} Linux Development $(ARGS)\n", Basename));
+            }
+            MakefileContent.Append("\n.PHONY: $(TARGETS)\n");
+            var FullFileName = Path.Combine(MasterProjectRelativePath, FileName);
+            return WriteFileIfChanged(FullFileName, MakefileContent.ToString());
+        }
+
+        private bool WriteCMakeLists()
+        {
+            var FileName = "CMakeLists.txt";
+            var CMakefileContent = new StringBuilder();
+
+            CMakefileContent.Append(
+                "# Makefile generated by MakefileGenerator.cs\n" +
+                "# *DO NOT EDIT*\n\n" +
+                "cmake_minimum_required (VERSION 2.6)\n" +
+                "project (UE4)\n\n" +
+                "set(BUILD ${CMAKE_SOURCE_DIR}/Engine/Build/BatchFiles/Linux/Build.sh)\n\n" +
+                "set(SOURCE_FILES"
+                );
+
+            var AllModuleFiles = DiscoverModules();
+            foreach (string CurModuleFile in AllModuleFiles)
+            {
+                CMakefileContent.Append("\n\t");
+
+                var FoundFiles = SourceFileSearch.FindModuleSourceFiles(CurModuleFile, ExcludeNoRedistFiles: bExcludeNoRedistFiles);
+                foreach (string CurSourceFile in FoundFiles)
+                {
+                    string SourceFileRelativeToRoot = Utils.MakePathRelativeTo(CurSourceFile, Path.Combine(EngineRelativePath));
+                    if (!SourceFileRelativeToRoot.StartsWith("..") && !Path.IsPathRooted(SourceFileRelativeToRoot))
+                    {
+                        SourceFileRelativeToRoot = "Engine/" + SourceFileRelativeToRoot;
+                    }
+                    else
+                    {
+                        SourceFileRelativeToRoot = SourceFileRelativeToRoot.Substring(3);
+                    }
+
+                    CMakefileContent.Append(String.Format("\"{0}\" ", SourceFileRelativeToRoot));
+                }
+            }
+            CMakefileContent.Append("\n)\n\n");
+            foreach (string TargetFilePath in DiscoverTargets())
+            {
+                var TargetName = Utils.GetFilenameWithoutAnyExtensions(TargetFilePath);		// Remove both ".cs" and ".Target"
+
+                // Check to see if this is an Engine target.  That is, the target is located under the "Engine" folder
+                bool IsEngineTarget = false;
+                string TargetFileRelativeToEngineDirectory = Utils.MakePathRelativeTo(TargetFilePath, Path.Combine(EngineRelativePath), AlwaysTreatSourceAsDirectory: false);
+                if (!TargetFileRelativeToEngineDirectory.StartsWith("..") && !Path.IsPathRooted(TargetFileRelativeToEngineDirectory))
+                {
+                    // This is an engine target
+                    IsEngineTarget = true;
+                }
+                string RootPath = IsEngineTarget ? "Engine"  : TargetFileRelativeToEngineDirectory.Split(Path.DirectorySeparatorChar)[1];
+
+                foreach (UnrealTargetConfiguration CurConfiguration in Enum.GetValues(typeof(UnrealTargetConfiguration)))
+                {
+                    if (CurConfiguration != UnrealTargetConfiguration.Unknown && CurConfiguration != UnrealTargetConfiguration.Development)
+                    {
+                        if (UnrealBuildTool.IsValidConfiguration(CurConfiguration))
+                        {
+                            var ConfName = Enum.GetName(typeof(UnrealTargetConfiguration), CurConfiguration);
+                            string OutFile = Path.Combine(RootPath, String.Format("Binaries/Linux/{0}-Linux-{1}", TargetName, ConfName));
+                            //CMakefileContent.Append(String.Format("add_custom_command(OUTPUT {0} COMMAND ${{BUILD}} {1} Linux {2} DEPENDS ${{SOURCE_FILES}})\n", OutFile, TargetName, ConfName));
+                            //CMakefileContent.Append(String.Format("add_executable({0}-Linux-{1} IMPORTED)\n", TargetName, ConfName, OutFile));
+                            //CMakefileContent.Append(String.Format("set_property(TARGET {0}-Linux-{1} PROPERTY IMPORTED_LOCATION {2})\n", TargetName, ConfName, OutFile));
+                            //CMakefileContent.Append(String.Format("add_dependencies({0}-Linux-{1} {2})\n", TargetName, ConfName, OutFile));
+                            //CMakefileContent.Append(String.Format("add_custom_target({0}-Linux-{1} DEPENDS {2} SOURCES ${{SOURCE_FILES}})\n", TargetName, ConfName, OutFile));
+                            CMakefileContent.Append(String.Format("add_custom_target({0}-Linux-{1} ${{BUILD}} {0} Linux {1})\n", TargetName, ConfName));
+                            //CMakefileContent.Append(String.Format("set_target_properties({0}-Linux-{1} PROPERTIES OUTPUT_NAME {2})\n", TargetName, ConfName, OutFile));
+                        }
+                    }
+                }
+
+                CMakefileContent.Append(String.Format("add_custom_target({0} ${{BUILD}} {0} Linux Development SOURCES ${{SOURCE_FILES}})\n\n", TargetName));
+            }
+            var FullFileName = Path.Combine(MasterProjectRelativePath, FileName);
+            return WriteFileIfChanged(FullFileName, CMakefileContent.ToString());
+        }
+
 		/// ProjectFileGenerator interface
 		//protected override bool WriteMasterProjectFile( ProjectFile UBTProject )
 		protected override bool WriteProjectFiles ()
 		{
-			var FileName = "Makefile"; // MasterProjectName + ".mk";
-			var MakefileContent = new StringBuilder();
-			MakefileContent.Append (
-				"# Makefile generated by MakefileGenerator.cs\n" +
-				"# *DO NOT EDIT*\n\n" +
-				"PLATFORM = Linux\n" +
-				"CONFIG = Debug\n\n" +
-				"TARGETS ="
-				);
-			foreach( string Target in DiscoverTargets() )
-			{
-				var Basename = Path.GetFileNameWithoutExtension(Target);
-				Basename = Path.GetFileNameWithoutExtension(Basename);
-				MakefileContent.Append( " \\\n    " + Basename );
-			}
-			MakefileContent.Append( "\n\nall: $(TARGETS)\n" );
-			foreach( string Target in DiscoverTargets() )
-			{
-				var Basename = Path.GetFileNameWithoutExtension(Target);
-				Basename = Path.GetFileNameWithoutExtension(Basename);
-				MakefileContent.Append(String.Format("\n{0}:\n\tEngine/Build/BatchFiles/Linux/Build.sh {1} $(PLATFORM) $(CONFIG) $(ARGS)\n", Basename, Basename));
-			}
-			MakefileContent.Append( "\n.PHONY: $(TARGETS)\n" );
-			var FullFileName = Path.Combine( MasterProjectRelativePath, FileName );
-			return WriteFileIfChanged( FullFileName, MakefileContent.ToString() );
+            return WriteMakefile() & WriteCMakeLists();
 		}
 
 		/// ProjectFileGenerator interface
diff --git a/Engine/Source/Programs/UnrealBuildTool/System/ProjectFileGenerator.cs b/Engine/Source/Programs/UnrealBuildTool/System/ProjectFileGenerator.cs
--- a/Engine/Source/Programs/UnrealBuildTool/System/ProjectFileGenerator.cs
+++ b/Engine/Source/Programs/UnrealBuildTool/System/ProjectFileGenerator.cs
@@ -832,7 +832,7 @@ namespace UnrealBuildTool
 		/// Finds all module files (filtering by platform)
 		/// </summary>
 		/// <returns>Filtered list of module files</returns>
-		private List<string> DiscoverModules()
+		protected List<string> DiscoverModules()
 		{
 			var AllModuleFiles = new List<string>();
 
diff --git a/Engine/Source/Runtime/Core/Private/Internationalization/Internationalization.cpp b/Engine/Source/Runtime/Core/Private/Internationalization/Internationalization.cpp
--- a/Engine/Source/Runtime/Core/Private/Internationalization/Internationalization.cpp
+++ b/Engine/Source/Runtime/Core/Private/Internationalization/Internationalization.cpp
@@ -361,7 +361,7 @@ FString& FInternationalization::Leetify(FString& SourceString)
 		if( !Succeeded )
 		{
 			// Failed to LEETify, add something to beginning and end just to help mark the string.
-			SourceString = FString::FString(TEXT("\x2021")) + SourceString + FString::FString(TEXT("\x2021"));
+			SourceString = FString(TEXT("\x2021")) + SourceString + FString(TEXT("\x2021"));
 		}
 	}
 
diff --git a/Engine/Source/Runtime/Core/Private/Linux/LinuxPlatformMisc.cpp b/Engine/Source/Runtime/Core/Private/Linux/LinuxPlatformMisc.cpp
--- a/Engine/Source/Runtime/Core/Private/Linux/LinuxPlatformMisc.cpp
+++ b/Engine/Source/Runtime/Core/Private/Linux/LinuxPlatformMisc.cpp
@@ -72,29 +72,14 @@ bool FLinuxMisc::ControlScreensaver(EScreenSaverAction Action)
 	return true;
 }
 
-/** Global variable to resolve interdependency between RootDir(), which calls NormalizePath(), and NormalizePath(), which needs RootDir() once.*/
-bool volatile GHaveRootDir = false;
-
 const TCHAR* FLinuxMisc::RootDir()
 {
 	const TCHAR * TrueRootDir = FGenericPlatformMisc::RootDir();
-	GHaveRootDir = true;
 	return TrueRootDir;
 }
 
 void FLinuxMisc::NormalizePath(FString& InPath)
 {
-	// only lowercase part of the path that is under root (if we know it)
-	if (GHaveRootDir)
-	{
-		static FString Root = RootDir();
-		// if absolute path begins at root
-		if (InPath.Find(Root, ESearchCase::IgnoreCase) == 0)
-		{
-			InPath = FPaths::Combine(*Root, *InPath.RightChop(Root.Len()).ToLower());
-		}
-	}
-
 	if (InPath.Contains(TEXT("~"), ESearchCase::CaseSensitive))	// case sensitive is quicker, and our substring doesn't care
 	{
 		static bool bHaveHome = false;
@@ -103,7 +88,7 @@ void FLinuxMisc::NormalizePath(FString& InPath)
 		if (!bHaveHome)
 		{
 			//  get user $HOME var first
-			const char * VarValue = __secure_getenv("HOME");
+			const char * VarValue = secure_getenv("HOME");
 			if (NULL != VarValue)
 			{
 				FCString::Strcpy(CachedResult, ARRAY_COUNT(CachedResult) - 1, ANSI_TO_TCHAR(VarValue));
diff --git a/Engine/Source/Runtime/Core/Private/Tests/Misc/CircularQueueTest.cpp b/Engine/Source/Runtime/Core/Private/Tests/Misc/CircularQueueTest.cpp
--- a/Engine/Source/Runtime/Core/Private/Tests/Misc/CircularQueueTest.cpp
+++ b/Engine/Source/Runtime/Core/Private/Tests/Misc/CircularQueueTest.cpp
@@ -26,7 +26,7 @@ bool FCircularQueueTest::RunTest( const FString& Parameters )
 	// partially filled
 	{
 		TCircularQueue<int32> Queue(QueueSize);
-		int32 Value;
+		int32 Value = 0;
 	
 		TestTrue(TEXT("Adding to an empty queue must succeed"), Queue.Enqueue(666));
 		TestFalse(TEXT("Partially filled queues must not be empty"), Queue.IsEmpty());
@@ -47,7 +47,7 @@ bool FCircularQueueTest::RunTest( const FString& Parameters )
 		TestTrue(TEXT("Full queues must be full"), Queue.IsFull());
 		TestFalse(TEXT("Adding to full queue must fail"), Queue.Enqueue(666));
 
-		int32 Value;
+		int32 Value = 0;
 
 		for (int32 Index = 0; Index < QueueSize - 1; ++Index)
 		{
diff --git a/Engine/Source/Runtime/Core/Public/GenericPlatform/StandardPlatformString.h b/Engine/Source/Runtime/Core/Public/GenericPlatform/StandardPlatformString.h
--- a/Engine/Source/Runtime/Core/Public/GenericPlatform/StandardPlatformString.h
+++ b/Engine/Source/Runtime/Core/Public/GenericPlatform/StandardPlatformString.h
@@ -151,7 +151,7 @@ struct FStandardPlatformString : public FGenericPlatformString
 	}
 
 #if PLATFORM_USE_SYSTEM_VSWPRINTF
-	static FORCEINLINE int32 GetVarArgs( WIDECHAR* Dest, SIZE_T DestSize, int32 Count, const WIDECHAR*& Fmt, va_list ArgPtr )
+	static int32 GetVarArgs( WIDECHAR* Dest, SIZE_T DestSize, int32 Count, const WIDECHAR*& Fmt, va_list ArgPtr )
 	{
 #if PLATFORM_USE_LS_SPEC_FOR_WIDECHAR
 		// fix up the Fmt string, as fast as possible, without using an FString
@@ -304,7 +304,7 @@ struct FStandardPlatformString : public FGenericPlatformString
 		return strtok(StrToken, Delim);
 	}
 
-	static FORCEINLINE int32 GetVarArgs( ANSICHAR* Dest, SIZE_T DestSize, int32 Count, const ANSICHAR*& Fmt, va_list ArgPtr )
+	static int32 GetVarArgs( ANSICHAR* Dest, SIZE_T DestSize, int32 Count, const ANSICHAR*& Fmt, va_list ArgPtr )
 	{
 		int32 Result = vsnprintf(Dest,Count,Fmt,ArgPtr);
 		va_end( ArgPtr );
diff --git a/Engine/Source/Runtime/Core/Public/Linux/LinuxPlatform.h b/Engine/Source/Runtime/Core/Public/Linux/LinuxPlatform.h
--- a/Engine/Source/Runtime/Core/Public/Linux/LinuxPlatform.h
+++ b/Engine/Source/Runtime/Core/Public/Linux/LinuxPlatform.h
@@ -22,12 +22,12 @@ typedef FLinuxPlatformTypes FPlatformTypes;
 
 // Base defines, must define these for the platform, there are no defaults
 #define PLATFORM_DESKTOP						1
-#ifdef _LINUX64
+#if defined(_LINUX64) || defined(_LP64)
 	#define PLATFORM_64BITS						1
 #else
 	#define PLATFORM_64BITS						0
 #endif
-#define PLATFORM_CAN_SUPPORT_EDITORONLY_DATA	0
+#define PLATFORM_CAN_SUPPORT_EDITORONLY_DATA	1
 
 // Base defines, defaults are commented out
 
diff --git a/Engine/Source/Runtime/Core/Public/Linux/LinuxPlatformIncludes.h b/Engine/Source/Runtime/Core/Public/Linux/LinuxPlatformIncludes.h
--- a/Engine/Source/Runtime/Core/Public/Linux/LinuxPlatformIncludes.h
+++ b/Engine/Source/Runtime/Core/Public/Linux/LinuxPlatformIncludes.h
@@ -72,5 +72,18 @@ struct RECT
 
 // include platform properties and typedef it for the runtime
 #include "Linux/LinuxPlatformProperties.h"
-typedef FLinuxPlatformProperties<UE_SERVER, !WITH_SERVER_CODE> FPlatformProperties;
 
+#ifndef WITH_EDITORONLY_DATA
+#error "WITH_EDITORONLY_DATA must be defined"
+#endif
+
+#ifndef UE_SERVER
+#error "WITH_EDITORONLY_DATA must be defined"
+#endif
+
+typedef FLinuxPlatformProperties<WITH_EDITORONLY_DATA, UE_SERVER, !WITH_SERVER_CODE> FPlatformProperties;
+
+// In glibc 2.17, __secure_getenv was renamed to secure_getenv.
+#if defined(__GLIBC__) && (__GLIBC__ == 2) && (__GLIBC_MINOR__ < 17)
+#  define secure_getenv __secure_getenv
+#endif
\ No newline at end of file
diff --git a/Engine/Source/Runtime/Core/Public/Linux/LinuxPlatformProperties.h b/Engine/Source/Runtime/Core/Public/Linux/LinuxPlatformProperties.h
--- a/Engine/Source/Runtime/Core/Public/Linux/LinuxPlatformProperties.h
+++ b/Engine/Source/Runtime/Core/Public/Linux/LinuxPlatformProperties.h
@@ -15,18 +15,33 @@
 /**
  * Implements Linux platform properties.
  */
-template<bool IS_DEDICATED_SERVER, bool IS_CLIENT_ONLY>
+template<bool HAS_EDITOR_DATA, bool IS_DEDICATED_SERVER, bool IS_CLIENT_ONLY>
 struct FLinuxPlatformProperties
 	: public FGenericPlatformProperties
 {
 	static FORCEINLINE const char* DisplayName()
 	{
+		if (IS_DEDICATED_SERVER)
+		{
+			return "Linux (Dedicated Server)";
+		}
+
+		if (HAS_EDITOR_DATA)
+		{
+			return "Linux (Editor)";
+		}
+
+		if (IS_CLIENT_ONLY)
+		{
+			return "Linux (Client-only)";
+		}
+
 		return "Linux";
 	}
 
 	static FORCEINLINE bool HasEditorOnlyData( )
 	{
-		return false;
+		return HAS_EDITOR_DATA;
 	}
 
 	static FORCEINLINE const char* IniPlatformName( )
@@ -55,15 +70,23 @@ struct FLinuxPlatformProperties
 		{
 			return "LinuxServer";
 		}
-		else
+
+		if (HAS_EDITOR_DATA)
 		{
 			return "Linux";
 		}
+
+		if (IS_CLIENT_ONLY)
+		{
+			return "LinuxClient";
+		}
+
+		return "LinuxNoEditor";
 	}
 
 	static FORCEINLINE bool RequiresCookedData( )
 	{
-		return true;
+		return !HAS_EDITOR_DATA;
 	}
 
 	static FORCEINLINE bool RequiresUserCredentials()
@@ -84,6 +107,7 @@ struct FLinuxPlatformProperties
 	static FORCEINLINE bool SupportsWindowedMode()
 	{
 		return !IS_DEDICATED_SERVER;
+		return true;
 	}
 
 	static FORCEINLINE bool SupportsQuit()
diff --git a/Engine/Source/Runtime/Core/Public/Linux/LinuxPlatformString.h b/Engine/Source/Runtime/Core/Public/Linux/LinuxPlatformString.h
--- a/Engine/Source/Runtime/Core/Public/Linux/LinuxPlatformString.h
+++ b/Engine/Source/Runtime/Core/Public/Linux/LinuxPlatformString.h
@@ -23,6 +23,11 @@ struct FLinuxPlatformString : public FStandardPlatformString
 		int ret = mbstowcs(Dest, Source, LengthM1);
 		if(ret == LengthM1) Dest[LengthM1] = '\0';
 	}
+
+	static const ANSICHAR* GetEncodingName()
+	{
+		return "UTF-32LE";
+	}
 };
 
 typedef FLinuxPlatformString FPlatformString;
diff --git a/Engine/Source/Runtime/Engine/Private/AIController.cpp b/Engine/Source/Runtime/Engine/Private/AIController.cpp
--- a/Engine/Source/Runtime/Engine/Private/AIController.cpp
+++ b/Engine/Source/Runtime/Engine/Private/AIController.cpp
@@ -472,7 +472,7 @@ EPathFollowingRequestResult::Type AAIController::MoveToLocation(const FVector& D
 	FVector GoalLocation = Dest;
 
 	// fail if projection to navigation is required but it either failed or there's not navigation component
-	if (bCanRequestMove && bProjectDestinationToNavigation && (NavComponent == NULL || !NavComponent->ProjectPointToNavigation(Dest, GoalLocation)))
+	if (bCanRequestMove && bProjectDestinationToNavigation && (NavComponent.Get() == NULL || !NavComponent->ProjectPointToNavigation(Dest, GoalLocation)))
 	{
 		UE_VLOG_LOCATION(this, Dest, 30.f, FLinearColor::Red, TEXT("AAIController::MoveToLocation failed to project destination location to navmesh"));
 		bCanRequestMove = false;
@@ -585,7 +585,7 @@ EPathFollowingStatus::Type AAIController::GetMoveStatus() const
 
 bool AAIController::HasPartialPath() const
 {
-	return (PathFollowingComponent != NULL) && (PathFollowingComponent->HasPartialPath());
+	return (PathFollowingComponent.Get() != NULL) && (PathFollowingComponent->HasPartialPath());
 }
 
 FVector AAIController::GetImmediateMoveDestination() const
diff --git a/Engine/Source/Runtime/Engine/Private/Character.cpp b/Engine/Source/Runtime/Engine/Private/Character.cpp
--- a/Engine/Source/Runtime/Engine/Private/Character.cpp
+++ b/Engine/Source/Runtime/Engine/Private/Character.cpp
@@ -319,7 +319,7 @@ void ACharacter::ApplyDamageMomentum(float DamageTaken, FDamageEvent const& Dama
 	UDamageType const* const DmgTypeCDO = DamageEvent.DamageTypeClass->GetDefaultObject<UDamageType>();
 	float const ImpulseScale = DmgTypeCDO->DamageImpulse;
 
-	if ( (ImpulseScale > 3.f) && (CharacterMovement != NULL) )
+	if ( (ImpulseScale > 3.f) && (CharacterMovement.Get() != NULL) )
 	{
 		FHitResult HitInfo;
 		FVector ImpulseDir;
@@ -511,13 +511,13 @@ void ACharacter::SetBase( UPrimitiveComponent* NewBaseComponent, bool bNotifyPaw
 
 void ACharacter::TurnOff()
 {
-	if (CharacterMovement != NULL)
+	if (CharacterMovement.Get() != NULL)
 	{
 		CharacterMovement->StopMovementImmediately();
 		CharacterMovement->DisableMovement();
 	}
 
-	if (GetNetMode() != NM_DedicatedServer && Mesh != NULL)
+	if (GetNetMode() != NM_DedicatedServer && Mesh.Get() != NULL)
 	{
 		Mesh->bPauseAnims = true;
 		if (Mesh->IsSimulatingPhysics())
@@ -538,7 +538,7 @@ void ACharacter::Restart()
 
 void ACharacter::PawnClientRestart()
 {
-	if (CharacterMovement != NULL)
+	if (CharacterMovement.Get() != NULL)
 	{
 		CharacterMovement->StopMovementImmediately();
 		CharacterMovement->ResetPredictionData_Client();
@@ -585,7 +585,7 @@ void ACharacter::DisplayDebug(UCanvas* Canvas, const FDebugDisplayInfo& DebugDis
 		FIndenter PhysicsIndent(Indent);
 
 		FString BaseString;
-		if ( CharacterMovement == NULL || MovementBase == NULL )
+		if ( CharacterMovement.Get() == NULL || MovementBase == NULL )
 		{
 			BaseString = "Not Based";
 		}
@@ -598,7 +598,7 @@ void ACharacter::DisplayDebug(UCanvas* Canvas, const FDebugDisplayInfo& DebugDis
 		Canvas->DrawText(RenderFont, FString::Printf(TEXT("RelativeLoc: %s RelativeRot: %s %s"), *RelativeMovement.Location.ToString(), *RelativeMovement.Rotation.ToString(), *BaseString), Indent, YPos);
 		YPos += YL;
 
-		if ( CharacterMovement != NULL )
+		if ( CharacterMovement.Get() != NULL )
 		{
 			CharacterMovement->DisplayDebug(Canvas, DebugDisplay, YL, YPos);
 		}
diff --git a/Engine/Source/Runtime/Engine/Private/Interpolation.cpp b/Engine/Source/Runtime/Engine/Private/Interpolation.cpp
--- a/Engine/Source/Runtime/Engine/Private/Interpolation.cpp
+++ b/Engine/Source/Runtime/Engine/Private/Interpolation.cpp
@@ -4658,7 +4658,7 @@ bool UInterpTrackMove::GetLocationAtTime(UInterpTrackInst* TrInst, float Time, F
 {
 	UInterpTrackInstMove* MoveTrackInst = CastChecked<UInterpTrackInstMove>(TrInst);
 
-	check( SubTracks.Num() > 0 || (EulerTrack.Points.Num() == PosTrack.Points.Num()) && (EulerTrack.Points.Num() == LookupTrack.Points.Num()));
+	check( (SubTracks.Num() > 0) || ((EulerTrack.Points.Num() == PosTrack.Points.Num()) && (EulerTrack.Points.Num() == LookupTrack.Points.Num())));
 
 	// Do nothing if no data on this track.
 	if( SubTracks.Num() == 0 && EulerTrack.Points.Num() == 0)
diff --git a/Engine/Source/Runtime/Engine/Private/PhysicsEngine/PhysicsConstraintActor.cpp b/Engine/Source/Runtime/Engine/Private/PhysicsEngine/PhysicsConstraintActor.cpp
--- a/Engine/Source/Runtime/Engine/Private/PhysicsEngine/PhysicsConstraintActor.cpp
+++ b/Engine/Source/Runtime/Engine/Private/PhysicsEngine/PhysicsConstraintActor.cpp
@@ -29,7 +29,7 @@ void APhysicsConstraintActor::PostLoad()
 	Super::PostLoad();
 
 	// Copy 'actors to constrain' into component
-	if (GetLinkerUE4Version() < VER_UE4_ALL_PROPS_TO_CONSTRAINTINSTANCE && ConstraintComp != NULL)
+	if (GetLinkerUE4Version() < VER_UE4_ALL_PROPS_TO_CONSTRAINTINSTANCE && ConstraintComp.Get() != NULL)
 	{
 		ConstraintComp->ConstraintActor1 = ConstraintActor1_DEPRECATED;
 		ConstraintComp->ConstraintActor2 = ConstraintActor2_DEPRECATED;
diff --git a/Engine/Source/Runtime/Engine/Private/ShaderCompiler/ShaderCompiler.cpp b/Engine/Source/Runtime/Engine/Private/ShaderCompiler/ShaderCompiler.cpp
--- a/Engine/Source/Runtime/Engine/Private/ShaderCompiler/ShaderCompiler.cpp
+++ b/Engine/Source/Runtime/Engine/Private/ShaderCompiler/ShaderCompiler.cpp
@@ -1319,7 +1319,7 @@ void FShaderCompilingManager::ProcessCompiledShaderMaps(
 			}
 		}
 
-		check(ShaderMap && Materials || ProcessIt.Key() == GlobalShaderMapId);
+		check((ShaderMap && Materials) || ProcessIt.Key() == GlobalShaderMapId);
 
 		if (ShaderMap && Materials)
 		{
diff --git a/Engine/Source/Runtime/Engine/Private/UserInterface/PlayerInput.cpp b/Engine/Source/Runtime/Engine/Private/UserInterface/PlayerInput.cpp
--- a/Engine/Source/Runtime/Engine/Private/UserInterface/PlayerInput.cpp
+++ b/Engine/Source/Runtime/Engine/Private/UserInterface/PlayerInput.cpp
@@ -323,7 +323,7 @@ void UPlayerInput::InvertAxis(const FName AxisName)
 		FAxisKeyDetails* KeyDetails = AxisKeyMap.Find(AxisName);
 		if (KeyDetails)
 		{
-			bInverted = KeyDetails->bInverted = !KeyDetails->bInverted;
+			bInverted = (KeyDetails->bInverted = !KeyDetails->bInverted);
 		}
 		if (bInverted)
 		{
diff --git a/Engine/Source/Runtime/Launch/Private/Linux/LaunchLinux.cpp b/Engine/Source/Runtime/Launch/Private/Linux/LaunchLinux.cpp
--- a/Engine/Source/Runtime/Launch/Private/Linux/LaunchLinux.cpp
+++ b/Engine/Source/Runtime/Launch/Private/Linux/LaunchLinux.cpp
@@ -106,7 +106,7 @@ static bool IncreasePerProcessLimits()
 #endif // !UE_BUILD_SHIPPING
 		if (!IncreaseLimit(RLIMIT_NOFILE, FileHandlesToReserve))
 		{
-			fprintf(stderr, "Could not adjust number of file handles, consider changing \"nofile\" in /etc/security/limits.conf and relogin.\n", errno, strerror(errno));
+			fprintf(stderr, "Could not adjust number of file handles, consider changing \"nofile\" in /etc/security/limits.conf and relogin.\nerror(%d): %s\n", errno, strerror(errno));
 			return false;
 		}
 	}
@@ -117,7 +117,7 @@ static bool IncreasePerProcessLimits()
 		printf("Increasing per-process limit of core file size to infinity.\n");
 		if (!IncreaseLimit(RLIMIT_CORE, RLIM_INFINITY))
 		{
-			fprintf(stderr, "Could not adjust core file size, consider changing \"core\" in /etc/security/limits.conf and relogin.\n", errno, strerror(errno));
+			fprintf(stderr, "Could not adjust core file size, consider changing \"core\" in /etc/security/limits.conf and relogin.\nerror(%d): %s\n", errno, strerror(errno));
 			fprintf(stderr, "Alternatively, pass -nocore if you are unable or unwilling to do that.\n");
 			return false;
 		}
diff --git a/Engine/Source/Runtime/OpenGLDrv/Private/OpenGL3.cpp b/Engine/Source/Runtime/OpenGLDrv/Private/OpenGL3.cpp
--- a/Engine/Source/Runtime/OpenGLDrv/Private/OpenGL3.cpp
+++ b/Engine/Source/Runtime/OpenGLDrv/Private/OpenGL3.cpp
@@ -24,7 +24,7 @@ void FOpenGL3::ProcessQueryGLInt()
 
 	GET_GL_INT(GL_MAX_GEOMETRY_TEXTURE_IMAGE_UNITS, 0, MaxGeometryTextureImageUnits);
 
-#ifndef __clang__
+#ifndef __GNUC__
 #define LOG_AND_GET_GL_QUERY_INT(IntEnum, Default, Dest) do { if (IntEnum) {glGetQueryiv(IntEnum, GL_QUERY_COUNTER_BITS, &Dest);} else {Dest = Default;} /*FPlatformMisc::LowLevelOutputDebugStringf(TEXT(" GL_QUERY_COUNTER_BITS: ") ## TEXT(#IntEnum) ## TEXT(": %d"), Dest);*/ } while(0)
 #else
 #define LOG_AND_GET_GL_QUERY_INT(IntEnum, Default, Dest) do { if (IntEnum) {glGetQueryiv(IntEnum, GL_QUERY_COUNTER_BITS, &Dest);} else {Dest = Default;} /*FPlatformMisc::LowLevelOutputDebugStringf(TEXT(" GL_QUERY_COUNTER_BITS: " #IntEnum ": %d"), Dest);*/ } while(0)
diff --git a/Engine/Source/Runtime/OpenGLDrv/Private/OpenGLDevice.cpp b/Engine/Source/Runtime/OpenGLDrv/Private/OpenGLDevice.cpp
--- a/Engine/Source/Runtime/OpenGLDrv/Private/OpenGLDevice.cpp
+++ b/Engine/Source/Runtime/OpenGLDrv/Private/OpenGLDevice.cpp
@@ -437,7 +437,7 @@ static void InitRHICapabilitiesForGL()
 #endif
 
 	// GL vendor and version information.
-#ifndef __clang__
+#ifndef __GNUC__
 	#define LOG_GL_STRING(StringEnum) UE_LOG(LogRHI, Log, TEXT("  ") ## TEXT(#StringEnum) ## TEXT(": %s"), ANSI_TO_TCHAR((const ANSICHAR*)glGetString(StringEnum)))
 #else
 	#define LOG_GL_STRING(StringEnum) UE_LOG(LogRHI, Log, TEXT("  " #StringEnum ": %s"), ANSI_TO_TCHAR((const ANSICHAR*)glGetString(StringEnum)))
@@ -522,7 +522,7 @@ static void InitRHICapabilitiesForGL()
 	FOpenGL::InitDebugContext();
 
 	// Log and get various limits.
-#ifndef __clang__
+#ifndef __GNUC__
 	#define LOG_AND_GET_GL_INT_TEMP(IntEnum,Default) GLint Value_##IntEnum; if (IntEnum) {glGetIntegerv(IntEnum, &Value_##IntEnum);} else {Value_##IntEnum = Default;} UE_LOG(LogRHI, Log, TEXT("  ") ## TEXT(#IntEnum) ## TEXT(": %d"), Value_##IntEnum)
 #else
 	#define LOG_AND_GET_GL_INT_TEMP(IntEnum,Default) GLint Value_##IntEnum; if (IntEnum) {glGetIntegerv(IntEnum, &Value_##IntEnum);} else {Value_##IntEnum = Default;} UE_LOG(LogRHI, Log, TEXT("  " #IntEnum ": %d"), Value_##IntEnum)
diff --git a/Engine/Source/Runtime/RenderCore/Public/UniformBuffer.h b/Engine/Source/Runtime/RenderCore/Public/UniformBuffer.h
--- a/Engine/Source/Runtime/RenderCore/Public/UniformBuffer.h
+++ b/Engine/Source/Runtime/RenderCore/Public/UniformBuffer.h
@@ -475,7 +475,7 @@ private:
 		typedef MemberType zzA##MemberName ArrayDecl; \
 		typedef TUniformBufferTypeInfo<zzA##MemberName>::TAlignedType zzT##MemberName; \
 		zzT##MemberName MemberName; \
-		checkAtCompileTime(TUniformBufferTypeInfo<zzT##MemberName>::BaseType != UBMT_INVALID,InvalidMemberType##MemberName); \
+		checkAtCompileTime(TUniformBufferTypeInfo<zzA##MemberName>::BaseType != UBMT_INVALID,InvalidMemberType##MemberName); \
 	private: \
 		struct zzNextMemberId##MemberName {}; \
 		static TArray<FUniformBufferStruct::FMember> zzGetMembersBefore(zzNextMemberId##MemberName) \
@@ -486,15 +486,15 @@ private:
 			OutMembers.Add(FUniformBufferStruct::FMember( \
 				TEXT(#MemberName), \
 				STRUCT_OFFSET(zzTThisStruct,MemberName), \
-				(EUniformBufferBaseType)TUniformBufferTypeInfo<zzT##MemberName>::BaseType, \
+				(EUniformBufferBaseType)TUniformBufferTypeInfo<zzA##MemberName>::BaseType, \
 				Precision, \
-				TUniformBufferTypeInfo<zzT##MemberName>::NumRows, \
-				TUniformBufferTypeInfo<zzT##MemberName>::NumColumns, \
-				TUniformBufferTypeInfo<zzT##MemberName>::NumElements, \
-				TUniformBufferTypeInfo<zzT##MemberName>::GetStruct() \
+				TUniformBufferTypeInfo<zzA##MemberName>::NumRows, \
+				TUniformBufferTypeInfo<zzA##MemberName>::NumColumns, \
+				TUniformBufferTypeInfo<zzA##MemberName>::NumElements, \
+				TUniformBufferTypeInfo<zzA##MemberName>::GetStruct() \
 				)); \
 			checkAtCompileTime( \
-				(STRUCT_OFFSET(zzTThisStruct,MemberName) & (TUniformBufferTypeInfo<zzT##MemberName>::Alignment - 1)) == 0, \
+				(STRUCT_OFFSET(zzTThisStruct,MemberName) & (TUniformBufferTypeInfo<zzA##MemberName>::Alignment - 1)) == 0, \
 				MisalignedUniformBufferStructMember##MemberName \
 				); \
 			return OutMembers; \
diff --git a/Engine/Source/Runtime/Renderer/Private/CompositionLighting/PostProcessDeferredDecals.cpp b/Engine/Source/Runtime/Renderer/Private/CompositionLighting/PostProcessDeferredDecals.cpp
--- a/Engine/Source/Runtime/Renderer/Private/CompositionLighting/PostProcessDeferredDecals.cpp
+++ b/Engine/Source/Runtime/Renderer/Private/CompositionLighting/PostProcessDeferredDecals.cpp
@@ -741,13 +741,13 @@ void FRCPassPostProcessDeferredDecals::Process(FRenderingCompositePassContext& C
 		{
 			// could be optimized
 			RHISetRenderTarget(GSceneRenderTargets.DBufferA->GetRenderTargetItem().TargetableTexture, FTextureRHIParamRef());
-			RHIClear(true, FLinearColor::FLinearColor(0, 0, 0, 1), false, 0, false, 0, FIntRect());
+			RHIClear(true, FLinearColor(0, 0, 0, 1), false, 0, false, 0, FIntRect());
 			RHISetRenderTarget(GSceneRenderTargets.DBufferB->GetRenderTargetItem().TargetableTexture, FTextureRHIParamRef());
 			// todo: some hardware would like to have 0 or 1 for faster clear, we chose 128/255 to represent 0 (8 bit cannot represent 0.5f)
-			RHIClear(true, FLinearColor::FLinearColor(128.0f/255.0f, 128.0f/255.0f, 128.0f/255.0f, 1), false, 0, false, 0, FIntRect());
+			RHIClear(true, FLinearColor(128.0f/255.0f, 128.0f/255.0f, 128.0f/255.0f, 1), false, 0, false, 0, FIntRect());
 			RHISetRenderTarget(GSceneRenderTargets.DBufferC->GetRenderTargetItem().TargetableTexture, FTextureRHIParamRef());
 			// R:roughness, G:roughness opacity
-			RHIClear(true, FLinearColor::FLinearColor(0, 1, 0, 1), false, 0, false, 0, FIntRect());
+			RHIClear(true, FLinearColor(0, 1, 0, 1), false, 0, false, 0, FIntRect());
 		}
 	}
 
diff --git a/Engine/Source/Runtime/Slate/Private/Framework/Docking/SDockingCross.cpp b/Engine/Source/Runtime/Slate/Private/Framework/Docking/SDockingCross.cpp
--- a/Engine/Source/Runtime/Slate/Private/Framework/Docking/SDockingCross.cpp
+++ b/Engine/Source/Runtime/Slate/Private/Framework/Docking/SDockingCross.cpp
@@ -178,7 +178,7 @@ static FDockingDragOperation::FDockTarget GetDropTarget( const TSharedPtr<SDocki
 	//       |  /  \  |
 	//       | /    \ |
 	//       y-------\ (1,1)
-	//       |        \
+	//       |        \.
 	//       v         \ slope = 1			
 
 	const FVector2D DeadZoneCenter = FVector2D(0.5f, 0.5f);
