Fixed a problem with Xcode project generator adding unnamed group to the list of file references

diff --git a/Engine/Source/Programs/UnrealBuildTool/System/XcodeProjectFileGenerator.cs b/Engine/Source/Programs/UnrealBuildTool/System/XcodeProjectFileGenerator.cs
--- a/Engine/Source/Programs/UnrealBuildTool/System/XcodeProjectFileGenerator.cs
+++ b/Engine/Source/Programs/UnrealBuildTool/System/XcodeProjectFileGenerator.cs
@@ -252,7 +252,10 @@ namespace UnrealBuildTool
 
 			foreach (XcodeFileGroup Group in Groups.Values)
 			{
-				Contents.Append(string.Format("\t\t\t\t{0} /* {1} */,{2}", Group.GroupGuid, Group.GroupName, ProjectFileGenerator.NewLine));
+				if (!string.IsNullOrEmpty(Group.GroupName))
+				{
+					Contents.Append(string.Format("\t\t\t\t{0} /* {1} */,{2}", Group.GroupGuid, Group.GroupName, ProjectFileGenerator.NewLine));
+				}
 			}
 
 			foreach (XcodeProjectTarget Target in Targets)
@@ -945,14 +948,17 @@ namespace UnrealBuildTool
 
 			foreach (XcodeFileGroup Group in Groups.Values)
 			{
-				Group.bReference = true;
-				string GroupPath = RelativePath + Group.GroupPath;
+				if (!string.IsNullOrEmpty(Group.GroupName))
+				{
+					Group.bReference = true;
+					string GroupPath = RelativePath + Group.GroupPath;
 
-				// Add File reference.
-				PBXFileReferenceSection += "\t\t" + Group.GroupGuid + " /* " + Group.GroupName + " */ = {"
+					// Add File reference.
+					PBXFileReferenceSection += "\t\t" + Group.GroupGuid + " /* " + Group.GroupName + " */ = {"
 											+ "isa = PBXFileReference; lastKnownFileType = folder; "
 											+ "name = " + Group.GroupName + "; "
 											+ "path = " + Utils.CleanDirectorySeparators(GroupPath, '/') + "; sourceTree = \"<group>\"; };" + ProjectFileGenerator.NewLine;
+				}
 			}
 		}
 
