Put UnrealBuildTool.Utils.DisplayProgress back in; it's still called from UAT. Will fixup properly later.

diff --git a/Engine/Source/Programs/UnrealBuildTool/System/Utils.cs b/Engine/Source/Programs/UnrealBuildTool/System/Utils.cs
--- a/Engine/Source/Programs/UnrealBuildTool/System/Utils.cs
+++ b/Engine/Source/Programs/UnrealBuildTool/System/Utils.cs
@@ -562,6 +562,28 @@ namespace UnrealBuildTool
 		}
 
 
+		/// <summary>
+		/// Backspaces the specified number of characters, then displays a progress percentage value to the console
+		/// </summary>
+		/// <param name="Numerator">Progress numerator</param>
+		/// <param name="Denominator">Progress denominator</param>
+		/// <param name="NumCharsToBackspaceOver">Number of characters to backspace before writing the text.  This value will be updated with the length of the new progress string.  The first time progress is displayed, you should pass 0 for this value.</param>
+		public static void DisplayProgress( int Numerator, int Denominator, ref int NumCharsToBackspaceOver )
+		{
+			// Backspace over previous progress value
+			while( NumCharsToBackspaceOver-- > 0 )
+			{
+				Console.Write( "\b" );
+			}
+
+			// Display updated progress string and keep track of how long it was
+			float ProgressValue = Denominator > 0 ? ( (float)Numerator / (float)Denominator ) : 1.0f;
+			var ProgressString = String.Format( "{0}%", Math.Round( ProgressValue * 100.0f ) );
+			NumCharsToBackspaceOver = ProgressString.Length;
+			Console.Write( ProgressString );
+		}
+
+
 		/*
 		 * Read and write classes with xml specifiers
 		 */
