Mac support for copying bundled iOS framework assets

diff --git a/Engine/Source/Programs/UnrealBuildTool/IOS/IOSToolChain.cs b/Engine/Source/Programs/UnrealBuildTool/IOS/IOSToolChain.cs
--- a/Engine/Source/Programs/UnrealBuildTool/IOS/IOSToolChain.cs
+++ b/Engine/Source/Programs/UnrealBuildTool/IOS/IOSToolChain.cs
@@ -882,6 +882,8 @@ namespace UnrealBuildTool
 
 				if ( ExternalExecution.GetRuntimePlatform() == UnrealTargetPlatform.Mac )
 				{
+					Log.TraceInformation( "Unzipping {0}", FrameworkZipItem.AbsolutePath );
+
 					// If we're on the mac, just unzip using the shell
 					string ResultsText;
 					string LocalUnzipZipPath = FrameworkZipItem.AbsolutePath.Substring( 0, FrameworkZipItem.AbsolutePath.LastIndexOf( '/' ) );
@@ -1002,6 +1004,42 @@ namespace UnrealBuildTool
 						PackageStub (RemoteShadowDirectoryMac, Target.GameName, Path.GetFileNameWithoutExtension (Target.OutputPath));
 					}
 				}
+
+				{
+					// Copy bundled assets from additional frameworks to the intermediate assets directory (so they can get picked up during staging)
+					String LocalFrameworkAssets = Path.GetFullPath( Target.ProjectDirectory + "/Intermediate/IOS/FrameworkAssets" );
+
+					// Delete the local dest directory if it exists
+					if ( Directory.Exists( LocalFrameworkAssets ) )
+					{
+						Directory.Delete( LocalFrameworkAssets, true );
+					}
+
+					// Create the intermediate local directory
+					string ResultsText;
+					RunExecutableAndWait( "mkdir", String.Format( "-p {0}", LocalFrameworkAssets ), out ResultsText );
+
+					foreach ( UEBuildFramework Framework in RememberedAdditionalFrameworks )
+					{
+						if ( Framework.OwningModule == null || Framework.CopyBundledAssets == null || Framework.CopyBundledAssets == "" )
+						{
+							continue;		// Only care if we need to copy bundle assets
+						}
+
+						string LocalZipPath = GetLocalFrameworkZipPath( Framework );
+
+						LocalZipPath = LocalZipPath.Replace( ".zip", "" );
+
+						// For now, this is hard coded, but we need to loop over all modules, and copy bundled assets that need it
+						string LocalSource	= LocalZipPath + "/" + Framework.CopyBundledAssets;
+						string BundleName	= Framework.CopyBundledAssets.Substring( Framework.CopyBundledAssets.LastIndexOf( '/' ) + 1 );
+						string LocalDest	= LocalFrameworkAssets + "/" + BundleName;
+
+						Log.TraceInformation( "Copying bundled asset... LocalSource: {0}, LocalDest: {1}", LocalSource, LocalDest );
+
+						RunExecutableAndWait( "cp", String.Format( "-R -L {0} {1}", LocalSource, LocalDest ), out ResultsText );
+					}
+				}
 			}
 			else
 			{
@@ -1147,51 +1185,53 @@ namespace UnrealBuildTool
 					}
 				}
 
-				// Copy bundled assets from additional frameworks to the intermediate assets directory (so they can get picked up during staging)
-				String LocalFrameworkAssets = Path.GetFullPath( Target.ProjectDirectory + "/Intermediate/IOS/FrameworkAssets" );
-				String RemoteFrameworkAssets = ConvertPath( LocalFrameworkAssets );
-
-				// Delete the intermediate directory on the mac
-				RPCUtilHelper.Command( "/", String.Format( "rm -rf {0}", RemoteFrameworkAssets ), "", null );
+				{
+					// Copy bundled assets from additional frameworks to the intermediate assets directory (so they can get picked up during staging)
+					String LocalFrameworkAssets = Path.GetFullPath( Target.ProjectDirectory + "/Intermediate/IOS/FrameworkAssets" );
+					String RemoteFrameworkAssets = ConvertPath( LocalFrameworkAssets );
 
-				// Create a fresh intermediate after we delete it
-				RPCUtilHelper.Command( "/", String.Format( "mkdir -p {0}", RemoteFrameworkAssets ), "", null );
+					// Delete the intermediate directory on the mac
+					RPCUtilHelper.Command( "/", String.Format( "rm -rf {0}", RemoteFrameworkAssets ), "", null );
 
-				// Delete the local dest directory if it exists
-				if ( Directory.Exists( LocalFrameworkAssets ) )
-				{
-					Directory.Delete( LocalFrameworkAssets, true );
-				}
+					// Create a fresh intermediate after we delete it
+					RPCUtilHelper.Command( "/", String.Format( "mkdir -p {0}", RemoteFrameworkAssets ), "", null );
 
-				foreach ( UEBuildFramework Framework in RememberedAdditionalFrameworks )
-				{
-					if ( Framework.OwningModule == null || Framework.CopyBundledAssets == null || Framework.CopyBundledAssets == "" )
+					// Delete the local dest directory if it exists
+					if ( Directory.Exists( LocalFrameworkAssets ) )
 					{
-						continue;		// Only care if we need to copy bundle assets
+						Directory.Delete( LocalFrameworkAssets, true );
 					}
 
-					string RemoteZipPath = GetRemoteFrameworkZipPath( Framework );
+					foreach ( UEBuildFramework Framework in RememberedAdditionalFrameworks )
+					{
+						if ( Framework.OwningModule == null || Framework.CopyBundledAssets == null || Framework.CopyBundledAssets == "" )
+						{
+							continue;		// Only care if we need to copy bundle assets
+						}
+
+						string RemoteZipPath = GetRemoteFrameworkZipPath( Framework );
 
-					RemoteZipPath = RemoteZipPath.Replace( ".zip", "" );
+						RemoteZipPath = RemoteZipPath.Replace( ".zip", "" );
 
-					// For now, this is hard coded, but we need to loop over all modules, and copy bundled assets that need it
-					string RemoteSource = RemoteZipPath + "/" + Framework.CopyBundledAssets;
-					string BundleName	= Framework.CopyBundledAssets.Substring( Framework.CopyBundledAssets.LastIndexOf( '/' ) + 1 );
+						// For now, this is hard coded, but we need to loop over all modules, and copy bundled assets that need it
+						string RemoteSource = RemoteZipPath + "/" + Framework.CopyBundledAssets;
+						string BundleName	= Framework.CopyBundledAssets.Substring( Framework.CopyBundledAssets.LastIndexOf( '/' ) + 1 );
 
-					String RemoteDest	= RemoteFrameworkAssets + "/" + BundleName;
-					String LocalDest	= LocalFrameworkAssets + "\\" + BundleName;
+						String RemoteDest	= RemoteFrameworkAssets + "/" + BundleName;
+						String LocalDest	= LocalFrameworkAssets + "\\" + BundleName;
 
-					Log.TraceInformation( "Copying bundled asset... RemoteSource: {0}, RemoteDest: {1}, LocalDest: {2}", RemoteSource, RemoteDest, LocalDest );
+						Log.TraceInformation( "Copying bundled asset... RemoteSource: {0}, RemoteDest: {1}, LocalDest: {2}", RemoteSource, RemoteDest, LocalDest );
 
-					Hashtable Results = RPCUtilHelper.Command( "/", String.Format( "cp -R -L {0} {1}", RemoteSource, RemoteDest ), "", null );
+						Hashtable Results = RPCUtilHelper.Command( "/", String.Format( "cp -R -L {0} {1}", RemoteSource, RemoteDest ), "", null );
 
-					foreach ( DictionaryEntry Entry in Results )
-					{
-						Log.TraceInformation( "{0}", Entry.Value );
-					}
+						foreach ( DictionaryEntry Entry in Results )
+						{
+							Log.TraceInformation( "{0}", Entry.Value );
+						}
 
-					// Copy the bundled resource from the remote mac to the local dest
-					RPCUtilHelper.CopyDirectory( RemoteDest, LocalDest, RPCUtilHelper.ECopyOptions.None );
+						// Copy the bundled resource from the remote mac to the local dest
+						RPCUtilHelper.CopyDirectory( RemoteDest, LocalDest, RPCUtilHelper.ECopyOptions.None );
+					}
 				}
 			}
 		}
