Removed some bHackHeaderGenerator usage.

diff --git a/Engine/Source/Programs/UnrealBuildTool/Configuration/UEBuildModule.cs b/Engine/Source/Programs/UnrealBuildTool/Configuration/UEBuildModule.cs
--- a/Engine/Source/Programs/UnrealBuildTool/Configuration/UEBuildModule.cs
+++ b/Engine/Source/Programs/UnrealBuildTool/Configuration/UEBuildModule.cs
@@ -1418,7 +1418,7 @@ namespace UnrealBuildTool
 
 		private FileItem ProcessDependencies(FileItem CPPFile, CPPEnvironment ModuleCompileEnvironment)
 		{
-			List<DependencyInclude> DirectIncludeFilenames = CPPEnvironment.GetDirectIncludeDependencies(CPPFile, ModuleCompileEnvironment.Config.TargetPlatform, ModuleCompileEnvironment.bHackHeaderGenerator);
+			List<DependencyInclude> DirectIncludeFilenames = CPPEnvironment.GetDirectIncludeDependencies(CPPFile, ModuleCompileEnvironment.Config.TargetPlatform);
 			if (BuildConfiguration.bPrintDebugInfo)
 			{
 				Log.TraceVerbose("Found direct includes for {0}: {1}", Path.GetFileName(CPPFile.AbsolutePath), string.Join(", ", DirectIncludeFilenames.Select(F => F.IncludeName)));
diff --git a/Engine/Source/Programs/UnrealBuildTool/System/ActionGraph.cs b/Engine/Source/Programs/UnrealBuildTool/System/ActionGraph.cs
--- a/Engine/Source/Programs/UnrealBuildTool/System/ActionGraph.cs
+++ b/Engine/Source/Programs/UnrealBuildTool/System/ActionGraph.cs
@@ -1030,7 +1030,7 @@ namespace UnrealBuildTool
 
 				if( ExpandCPPHeaderDependencies && bIsCPPFile )
 				{
-					List<DependencyInclude> DirectlyIncludedFilenames = CPPEnvironment.GetDirectIncludeDependencies( FileItem, CPPTargetPlatform.Win64, bHackHeaderGenerator:false );
+					List<DependencyInclude> DirectlyIncludedFilenames = CPPEnvironment.GetDirectIncludeDependencies( FileItem, CPPTargetPlatform.Win64 );
 
 					// Resolve the included file name to an actual file.
 					var DirectlyIncludedFiles =
diff --git a/Engine/Source/Programs/UnrealBuildTool/System/CPPHeaders.cs b/Engine/Source/Programs/UnrealBuildTool/System/CPPHeaders.cs
--- a/Engine/Source/Programs/UnrealBuildTool/System/CPPHeaders.cs
+++ b/Engine/Source/Programs/UnrealBuildTool/System/CPPHeaders.cs
@@ -170,7 +170,7 @@ namespace UnrealBuildTool
 				IncludedFilesMap.Add(CPPFile, IncludedFileList);
 
 				// Gather a list of names of files directly included by this C++ file.
-				List<DependencyInclude> DirectIncludes = GetDirectIncludeDependencies(CPPFile, Config.TargetPlatform, bHackHeaderGenerator);
+				List<DependencyInclude> DirectIncludes = GetDirectIncludeDependencies(CPPFile, Config.TargetPlatform);
 
 				// Build a list of the unique set of files that are included by this file.
 				string SourceFilesDirectory = Path.GetDirectoryName( CPPFile.AbsolutePath );
@@ -323,7 +323,7 @@ namespace UnrealBuildTool
 													RegexOptions.Compiled | RegexOptions.Singleline | RegexOptions.ExplicitCapture );
 
 		/** Finds the names of files directly included by the given C++ file. */
-		public static List<DependencyInclude> GetDirectIncludeDependencies( FileItem CPPFile, CPPTargetPlatform TargetPlatform, bool bHackHeaderGenerator )
+		public static List<DependencyInclude> GetDirectIncludeDependencies( FileItem CPPFile, CPPTargetPlatform TargetPlatform )
 		{
 			// Try to fulfill request from cache first.
 			List<DependencyInclude> Result = null;
@@ -379,14 +379,14 @@ namespace UnrealBuildTool
 						EndIndex = FileContents.Length;
 					}
 
-					CollectHeaders(CPPFile, bHackHeaderGenerator, Result, ref HasUObjects, FileToRead, FileContents, InstalledFolder, StartIndex, EndIndex);
+					CollectHeaders(CPPFile, Result, ref HasUObjects, FileToRead, FileContents, InstalledFolder, StartIndex, EndIndex);
 
 					StartIndex = EndIndex + 1;
 				}
 			}
 			else
 			{
-				CollectHeaders(CPPFile, bHackHeaderGenerator, Result, ref HasUObjects, FileToRead, FileContents, InstalledFolder, 0, FileContents.Length);
+				CollectHeaders(CPPFile, Result, ref HasUObjects, FileToRead, FileContents, InstalledFolder, 0, FileContents.Length);
 			}
 
 			// Store whether this C++ file has UObjects right inside the FileItem for this file
@@ -404,7 +404,6 @@ namespace UnrealBuildTool
 		/// Collects all header files included in a CPPFile
 		/// </summary>
 		/// <param name="CPPFile"></param>
-		/// <param name="bHackHeaderGenerator"></param>
 		/// <param name="Result"></param>
 		/// <param name="HasUObjects"></param>
 		/// <param name="FileToRead"></param>
@@ -412,7 +411,7 @@ namespace UnrealBuildTool
 		/// <param name="InstalledFolder"></param>
 		/// <param name="StartIndex"></param>
 		/// <param name="EndIndex"></param>
-		private static void CollectHeaders(FileItem CPPFile, bool bHackHeaderGenerator, List<DependencyInclude> Result, ref bool HasUObjects, string FileToRead, string FileContents, string InstalledFolder, int StartIndex, int EndIndex)
+		private static void CollectHeaders(FileItem CPPFile, List<DependencyInclude> Result, ref bool HasUObjects, string FileToRead, string FileContents, string InstalledFolder, int StartIndex, int EndIndex)
 		{
 			Match M = CPPHeaderRegex.Match(FileContents, StartIndex, EndIndex - StartIndex);
 			CaptureCollection CC = M.Groups["HeaderFile"].Captures;
