Linux native compilation improvements (incorporates PR 116 by Hydrael).

- detect Gentoo (and in generally distros without lsb_release installed)
- explicitly disable inlining in Debug

diff --git a/Engine/Source/Programs/UnrealBuildTool/Linux/LinuxToolChain.cs b/Engine/Source/Programs/UnrealBuildTool/Linux/LinuxToolChain.cs
--- a/Engine/Source/Programs/UnrealBuildTool/Linux/LinuxToolChain.cs
+++ b/Engine/Source/Programs/UnrealBuildTool/Linux/LinuxToolChain.cs
@@ -136,6 +136,15 @@ namespace UnrealBuildTool.Linux
             if (CompileEnvironment.Config.TargetConfiguration == CPPTargetConfiguration.Shipping)
             {
                 Result += " -Wno-unused-value";
+
+                // in shipping, strip as much info as possible
+                Result += " -g0";
+                Result += " -fomit-frame-pointer";
+                Result += " -fvisibility=hidden";           // prevents from exporting all symbols (reduces the size of the binary)
+            }
+            else if (CompileEnvironment.Config.TargetConfiguration == CPPTargetConfiguration.Debug)
+            {
+		Result += " -fno-inline";                   // disable inlining for better debuggability (e.g. callstacks, "skip file" in gdb)
             }
 
             // debug info
@@ -151,13 +160,6 @@ namespace UnrealBuildTool.Linux
             {
                 Result += " -gline-tables-only"; // include debug info for meaningful callstacks
             }
-            else
-            {
-                // in shipping, strip as much info as possible
-                Result += " -g0";
-                Result += " -fomit-frame-pointer";
-                Result += " -fvisibility=hidden";           // prevents from exporting all symbols (reduces the size of the binary)
-            }
 
             // optimization level
             if (CompileEnvironment.Config.TargetConfiguration == CPPTargetConfiguration.Debug)
@@ -176,10 +178,12 @@ namespace UnrealBuildTool.Linux
 
             //Result += " -v";                            // for better error diagnosis
 
+            // assume we will not perform 32 bit builds on Linux, so define
+            // _LINUX64 in any case
+            Result += " -D _LINUX64";
             if (CrossCompiling())
             {
                 Result += " -target x86_64-unknown-linux-gnu";        // Set target triple
-                Result += " -D _LINUX64";
                 Result += String.Format(" --sysroot={0}", BaseLinuxPath);
             }
 
diff --git a/Engine/Source/Programs/UnrealBuildTool/System/ExternalExecution.cs b/Engine/Source/Programs/UnrealBuildTool/System/ExternalExecution.cs
--- a/Engine/Source/Programs/UnrealBuildTool/System/ExternalExecution.cs
+++ b/Engine/Source/Programs/UnrealBuildTool/System/ExternalExecution.cs
@@ -136,7 +136,7 @@ namespace UnrealBuildTool
                 // TODO(sbc): Find a better way to do this.  Shelling out to uname would seem
                 // to make the most sense but I'm not sure there is an elegant way to do that from
                 // C#.
-                if (File.Exists("/etc/lsb-release") || File.Exists("/etc/debian_version"))
+                if (File.Exists("/etc/lsb-release") || File.Exists("/etc/debian_version") || (Directory.Exists("/etc/portage") && Directory.Exists("/usr/portage")))
                 {
                     return UnrealTargetPlatform.Linux;
                 }
