Cleaned up UBT's CMake generator.

- Removed leftover attempts at better CMake support (caused Mono warnings).
- Fixed formatting (tabs->spaces in C# files).

diff --git a/Engine/Source/Programs/UnrealBuildTool/System/MakefileGenerator.cs b/Engine/Source/Programs/UnrealBuildTool/System/MakefileGenerator.cs
--- a/Engine/Source/Programs/UnrealBuildTool/System/MakefileGenerator.cs
+++ b/Engine/Source/Programs/UnrealBuildTool/System/MakefileGenerator.cs
@@ -139,17 +139,7 @@ namespace UnrealBuildTool
             CMakefileContent.Append("\n)\n\n");
             foreach (string TargetFilePath in DiscoverTargets())
             {
-                var TargetName = Utils.GetFilenameWithoutAnyExtensions(TargetFilePath);		// Remove both ".cs" and ".Target"
-
-                // Check to see if this is an Engine target.  That is, the target is located under the "Engine" folder
-                bool IsEngineTarget = false;
-                string TargetFileRelativeToEngineDirectory = Utils.MakePathRelativeTo(TargetFilePath, Path.Combine(EngineRelativePath), AlwaysTreatSourceAsDirectory: false);
-                if (!TargetFileRelativeToEngineDirectory.StartsWith("..") && !Path.IsPathRooted(TargetFileRelativeToEngineDirectory))
-                {
-                    // This is an engine target
-                    IsEngineTarget = true;
-                }
-                string RootPath = IsEngineTarget ? "Engine"  : TargetFileRelativeToEngineDirectory.Split(Path.DirectorySeparatorChar)[1];
+                var TargetName = Utils.GetFilenameWithoutAnyExtensions(TargetFilePath);		// Remove both ".cs" and ".
 
                 foreach (UnrealTargetConfiguration CurConfiguration in Enum.GetValues(typeof(UnrealTargetConfiguration)))
                 {
@@ -158,14 +148,7 @@ namespace UnrealBuildTool
                         if (UnrealBuildTool.IsValidConfiguration(CurConfiguration))
                         {
                             var ConfName = Enum.GetName(typeof(UnrealTargetConfiguration), CurConfiguration);
-                            string OutFile = Path.Combine(RootPath, String.Format("Binaries/Linux/{0}-Linux-{1}", TargetName, ConfName));
-                            //CMakefileContent.Append(String.Format("add_custom_command(OUTPUT {0} COMMAND ${{BUILD}} {1} Linux {2} DEPENDS ${{SOURCE_FILES}})\n", OutFile, TargetName, ConfName));
-                            //CMakefileContent.Append(String.Format("add_executable({0}-Linux-{1} IMPORTED)\n", TargetName, ConfName, OutFile));
-                            //CMakefileContent.Append(String.Format("set_property(TARGET {0}-Linux-{1} PROPERTY IMPORTED_LOCATION {2})\n", TargetName, ConfName, OutFile));
-                            //CMakefileContent.Append(String.Format("add_dependencies({0}-Linux-{1} {2})\n", TargetName, ConfName, OutFile));
-                            //CMakefileContent.Append(String.Format("add_custom_target({0}-Linux-{1} DEPENDS {2} SOURCES ${{SOURCE_FILES}})\n", TargetName, ConfName, OutFile));
                             CMakefileContent.Append(String.Format("add_custom_target({0}-Linux-{1} ${{BUILD}} {0} Linux {1})\n", TargetName, ConfName));
-                            //CMakefileContent.Append(String.Format("set_target_properties({0}-Linux-{1} PROPERTIES OUTPUT_NAME {2})\n", TargetName, ConfName, OutFile));
                         }
                     }
                 }
@@ -176,33 +159,33 @@ namespace UnrealBuildTool
             return WriteFileIfChanged(FullFileName, CMakefileContent.ToString());
         }
 
-		/// ProjectFileGenerator interface
-		//protected override bool WriteMasterProjectFile( ProjectFile UBTProject )
-		protected override bool WriteProjectFiles ()
-		{
+        /// ProjectFileGenerator interface
+        //protected override bool WriteMasterProjectFile( ProjectFile UBTProject )
+        protected override bool WriteProjectFiles ()
+        {
             return WriteMakefile() & WriteCMakeLists();
-		}
+        }
 
-		/// ProjectFileGenerator interface
-		public override MasterProjectFolder AllocateMasterProjectFolder( ProjectFileGenerator InitOwnerProjectFileGenerator, string InitFolderName )
-		{
-			return new MakefileFolder( InitOwnerProjectFileGenerator, InitFolderName );
-		}
+        /// ProjectFileGenerator interface
+        public override MasterProjectFolder AllocateMasterProjectFolder( ProjectFileGenerator InitOwnerProjectFileGenerator, string InitFolderName )
+        {
+            return new MakefileFolder( InitOwnerProjectFileGenerator, InitFolderName );
+        }
 
-		/// ProjectFileGenerator interface
-		/// <summary>
-		/// Allocates a generator-specific project file object
-		/// </summary>
-		/// <param name="InitFilePath">Path to the project file</param>
-		/// <returns>The newly allocated project file object</returns>
-		protected override ProjectFile AllocateProjectFile( string InitFilePath )
-		{
-			return new MakefileProjectFile( InitFilePath );
-		}
+        /// ProjectFileGenerator interface
+        /// <summary>
+        /// Allocates a generator-specific project file object
+        /// </summary>
+        /// <param name="InitFilePath">Path to the project file</param>
+        /// <returns>The newly allocated project file object</returns>
+        protected override ProjectFile AllocateProjectFile( string InitFilePath )
+        {
+            return new MakefileProjectFile( InitFilePath );
+        }
 
-		/// ProjectFileGenerator interface
-		public override void CleanProjectFiles(string InMasterProjectRelativePath, string InMasterProjectName, string InIntermediateProjectFilesPath)
-		{
-		}
-	}
+        /// ProjectFileGenerator interface
+        public override void CleanProjectFiles(string InMasterProjectRelativePath, string InMasterProjectName, string InIntermediateProjectFilesPath)
+        {
+        }
+    }
 }
