Nuked dead code in UBT that caused Mono warnings.

Seems like these have been merged over from Pretest and never saw any use in Main.

diff --git a/Engine/Source/Programs/UnrealBuildTool/Configuration/UEBuildModule.cs b/Engine/Source/Programs/UnrealBuildTool/Configuration/UEBuildModule.cs
--- a/Engine/Source/Programs/UnrealBuildTool/Configuration/UEBuildModule.cs
+++ b/Engine/Source/Programs/UnrealBuildTool/Configuration/UEBuildModule.cs
@@ -1310,12 +1310,6 @@ namespace UnrealBuildTool
 				LinkInputFiles.AddRange( CPPCompileEnvironment.CompileFiles( CPPFilesToCompile, Name ).ObjectFiles );
 			}
 
-			string PCHPath = null;
-			if (ProcessedDependencies.UniquePCHHeaderFile != null)
-			{
-				PCHPath = UEToolChain.GetPlatformToolChain(ModuleCompileEnvironment.Config.TargetPlatform).ConvertPath( ProcessedDependencies.UniquePCHHeaderFile.AbsolutePath );
-			}
-
 			if (AutoGenerateCppInfo != null && !ModuleCompileEnvironment.bHackHeaderGenerator)
 			{
 				var GeneratedCppFileItem = FileItem.GetItemByPath( AutoGenerateCppInfo.Filename );
diff --git a/Engine/Source/Programs/UnrealBuildTool/Configuration/UEBuildTarget.cs b/Engine/Source/Programs/UnrealBuildTool/Configuration/UEBuildTarget.cs
--- a/Engine/Source/Programs/UnrealBuildTool/Configuration/UEBuildTarget.cs
+++ b/Engine/Source/Programs/UnrealBuildTool/Configuration/UEBuildTarget.cs
@@ -1371,8 +1371,6 @@ namespace UnrealBuildTool
 
 			if( (BuildConfiguration.bXGEExport && UEBuildConfiguration.bGenerateManifest) || (!UEBuildConfiguration.bGenerateManifest && !UEBuildConfiguration.bCleanProject && !ProjectFileGenerator.bGenerateProjectFiles) )
 			{
-				var ToolChain = UEToolChain.GetPlatformToolChain(GlobalLinkEnvironment.Config.TargetPlatform);
-
 				var UObjectModules = new List<UHTModuleInfo>();
 
 				// Figure out which modules have UObjects that we may need to generate headers for
