Cleaned up UBT from compiler warnings.

Removed code that was unused (or not yet fully in?).

diff --git a/Engine/Source/Programs/UnrealBuildTool/IOS/UEDeployIOS.cs b/Engine/Source/Programs/UnrealBuildTool/IOS/UEDeployIOS.cs
--- a/Engine/Source/Programs/UnrealBuildTool/IOS/UEDeployIOS.cs
+++ b/Engine/Source/Programs/UnrealBuildTool/IOS/UEDeployIOS.cs
@@ -231,7 +231,6 @@ namespace UnrealBuildTool.IOS
 					Log.TraceInformation("Copying binaries back to this device...");
 
 					IOSToolChain Toolchain = UEToolChain.GetPlatformToolChain(CPPTargetPlatform.IOS) as IOSToolChain;
-					bool bIsStaticLibrary = InTarget.OutputPath.EndsWith(".a");
 
 					try
 					{
@@ -241,9 +240,6 @@ namespace UnrealBuildTool.IOS
 							BinaryDir = BinaryDir.Replace(InTarget.Rules.Type.ToString(), "Game");
 						}
 
-						string RemoteBinariesDir = Toolchain.ConvertPath(BinaryDir);
-						string LocalBinariesDir = BinaryDir;
-
 						// Get the app bundle's name
 						string AppFullName = InTarget.AppName;
 						if (InTarget.Configuration != UnrealTargetConfiguration.Development)
@@ -252,7 +248,6 @@ namespace UnrealBuildTool.IOS
 							AppFullName += "-" + InTarget.Configuration.ToString();
 						}
 
-						List<string> NotBundledBinaries = new List<string>();
 						foreach (string BinaryPath in Toolchain.BuiltBinaries)
 						{
 							if (!BinaryPath.Contains("Dummy"))
diff --git a/Engine/Source/Programs/UnrealBuildTool/Windows/UEBuildWindows.cs b/Engine/Source/Programs/UnrealBuildTool/Windows/UEBuildWindows.cs
--- a/Engine/Source/Programs/UnrealBuildTool/Windows/UEBuildWindows.cs
+++ b/Engine/Source/Programs/UnrealBuildTool/Windows/UEBuildWindows.cs
@@ -357,8 +357,6 @@ namespace UnrealBuildTool
 			}
 			InBuildTarget.GlobalCompileEnvironment.Config.Definitions.Add("PLATFORM_WINDOWS=1");
 
-			var ProjectSourceDirectory = Path.Combine( UnrealBuildTool.GetUProjectPath(), "Source" );
-
 			String MorpheusShaderPath = Path.Combine(BuildConfiguration.RelativeEnginePath, "Shaders/PS4/PostProcessHMDMorpheus.usf");
 			if (File.Exists(MorpheusShaderPath))
 			{
