Mac support for staging intermediate framework assets folder

diff --git a/Engine/Source/Programs/UnrealBuildTool/IOS/UEDeployIOS.cs b/Engine/Source/Programs/UnrealBuildTool/IOS/UEDeployIOS.cs
--- a/Engine/Source/Programs/UnrealBuildTool/IOS/UEDeployIOS.cs
+++ b/Engine/Source/Programs/UnrealBuildTool/IOS/UEDeployIOS.cs
@@ -191,9 +191,17 @@ namespace UnrealBuildTool.IOS
 			}
 			
 			// copy additional engine framework assets in
-			if ( Directory.Exists( InEngineDir + "/Build/IOS/Resources/FrameworkAssets" ) )
+			string FrameworkAssetsPath = InEngineDir + "/Intermediate/IOS/FrameworkAssets";
+
+			// Let project override assets if they exist
+			if ( Directory.Exists( InProjectDirectory + "/Intermediate/IOS/FrameworkAssets" ) )
+			{
+				FrameworkAssetsPath = InProjectDirectory + "/Intermediate/IOS/FrameworkAssets";
+			}
+
+			if ( Directory.Exists( FrameworkAssetsPath ) )
 			{
-				CopyFiles( InEngineDir + "/Build/IOS/Resources/FrameworkAssets", AppDirectory, "*.*", true );
+				CopyFolder( FrameworkAssetsPath, AppDirectory, true );
 			}
 
 			//CopyFiles(BuildDirectory, PayloadDirectory, null, "iTunesArtwork", null);
