Use intermediate directory when unzipping frameworks

* Create intermediate directories to use for unzipping so we don't pollute directories not expecthing these temp files
* Skip frameworks not part of the target being considered
* Cleanup

diff --git a/Engine/Source/Programs/UnrealBuildTool/IOS/IOSToolChain.cs b/Engine/Source/Programs/UnrealBuildTool/IOS/IOSToolChain.cs
--- a/Engine/Source/Programs/UnrealBuildTool/IOS/IOSToolChain.cs
+++ b/Engine/Source/Programs/UnrealBuildTool/IOS/IOSToolChain.cs
@@ -308,12 +308,12 @@ namespace UnrealBuildTool
 
 		string GetLocalFrameworkZipPath( UEBuildFramework Framework )
 		{
-			if ( Framework.OwningModule != null )
+			if ( Framework.OwningModule == null )
 			{
-				// If we have a source module, assume that the path name is relative to that
-				return Path.GetFullPath( Framework.OwningModule.ModuleDirectory + "/" + Framework.FrameworkZipPath );
+				throw new BuildException( "GetLocalFrameworkZipPath: No owning module for framework {0}", Framework.FrameworkName );
 			}
-			return Path.GetFullPath( Framework.FrameworkZipPath );
+
+			return Path.GetFullPath( Framework.OwningModule.ModuleDirectory + "/" + Framework.FrameworkZipPath );
 		}
 
 		string GetRemoteFrameworkZipPath( UEBuildFramework Framework )
@@ -326,6 +326,48 @@ namespace UnrealBuildTool
 			return GetLocalFrameworkZipPath( Framework );
 		}
 
+		string GetRemoteIntermediateFrameworkZipPath( UEBuildFramework Framework )
+		{
+			if ( Framework.OwningModule == null )
+			{
+				throw new BuildException( "GetRemoteIntermediateFrameworkZipPath: No owning module for framework {0}", Framework.FrameworkName );
+			}
+
+			string IntermediatePath = Framework.OwningModule.Target.ProjectDirectory + "/Intermediate/UnzippedFrameworks/" + Framework.OwningModule.Name;
+			IntermediatePath =  Path.GetFullPath( ( IntermediatePath + Framework.FrameworkZipPath ).Replace( ".zip", "" ) );
+
+			if ( ExternalExecution.GetRuntimePlatform() != UnrealTargetPlatform.Mac )
+			{
+				return ConvertPath( IntermediatePath );
+			}
+
+			return IntermediatePath;
+		}
+
+		void CleanIntermediateDirectory( string Path )
+		{
+			if ( ExternalExecution.GetRuntimePlatform() != UnrealTargetPlatform.Mac )
+			{
+				// Delete the intermediate directory on the mac
+				RPCUtilHelper.Command( "/", String.Format( "rm -rf {0}", Path ), "", null );
+
+				// Create a fresh intermediate after we delete it
+				RPCUtilHelper.Command( "/", String.Format( "mkdir -p {0}", Path ), "", null );
+			}
+			else
+			{
+				// Delete the local dest directory if it exists
+				if ( Directory.Exists( Path ) )
+				{
+					Directory.Delete( Path, true );
+				}
+
+				// Create the intermediate local directory
+				string ResultsText;
+				RunExecutableAndWait( "mkdir", String.Format( "-p {0}", Path ), out ResultsText );
+			}
+		}
+
 		string GetLinkArguments_Global( LinkEnvironment LinkEnvironment )
 		{
 			string Result = "";
@@ -351,13 +393,10 @@ namespace UnrealBuildTool
 			}
 			foreach (UEBuildFramework Framework in LinkEnvironment.Config.AdditionalFrameworks)
 			{
-				if ( Framework.FrameworkZipPath != null )
+				if ( Framework.OwningModule != null && Framework.FrameworkZipPath != null && Framework.FrameworkZipPath != "" )
 				{
 					// If this framework has a zip specified, we'll need to setup the path as well
-					string FrameworkZipPath = GetRemoteFrameworkZipPath( Framework );
-
-					// Assume the path is the full name without the zip extension
-					Result += " -F " + FrameworkZipPath.Replace( ".zip", "" ); ;
+					Result += " -F " + GetRemoteIntermediateFrameworkZipPath( Framework );
 				}
 
 				Result += " -framework " + Framework.FrameworkName;
@@ -870,35 +909,26 @@ namespace UnrealBuildTool
 			// Unzip any third party frameworks that are stored as zips
 			foreach ( UEBuildFramework Framework in RememberedAdditionalFrameworks )
 			{
-				string LocalZipPath = GetLocalFrameworkZipPath( Framework );
+				string ZipSrcPath = GetRemoteFrameworkZipPath( Framework );
+				string ZipDstPath = GetRemoteIntermediateFrameworkZipPath( Framework );
 
-				FileItem FrameworkZipItem = FileItem.GetExistingItemByPath( LocalZipPath );
+				Log.TraceInformation( "Unzipping: {0} -> {1}", ZipSrcPath, ZipDstPath );
 
-				if ( FrameworkZipItem == null )
-				{
-					Log.TraceInformation( "FrameworkZipItem not found for {0}", LocalZipPath );
-					continue;
-				}
+				CleanIntermediateDirectory( ZipDstPath );
+
+				// Assume that there is another directory inside the zip with the same name as the zip
+				ZipDstPath = ZipDstPath.Substring( 0, ZipDstPath.LastIndexOf( '/' ) );
 
 				if ( ExternalExecution.GetRuntimePlatform() == UnrealTargetPlatform.Mac )
 				{
-					Log.TraceInformation( "Unzipping {0}", FrameworkZipItem.AbsolutePath );
-
 					// If we're on the mac, just unzip using the shell
 					string ResultsText;
-					string LocalUnzipZipPath = FrameworkZipItem.AbsolutePath.Substring( 0, FrameworkZipItem.AbsolutePath.LastIndexOf( '/' ) );
-					RunExecutableAndWait( "unzip", String.Format( "-o {0} -d {1}", FrameworkZipItem.AbsolutePath, LocalUnzipZipPath ), out ResultsText );
+					RunExecutableAndWait( "unzip", String.Format( "-o {0} -d {1}", ZipSrcPath, ZipDstPath ), out ResultsText );
 					continue;
 				}
 
-				// We copied using RPC utility, we need to unzip using RPC utility as well
-				FileItem RemoteShadowFile = LocalToRemoteFileItem( FrameworkZipItem, false );
-
-				string ZipPath = RemoteShadowFile.AbsolutePath.Substring( 0, RemoteShadowFile.AbsolutePath.LastIndexOf( '/' ) );
-
-				Log.TraceInformation( "Unzipping: {0}", RemoteShadowFile.AbsolutePath );
-
-				Hashtable Result = RPCUtilHelper.Command( "/", String.Format( "unzip -o {0} -d {1}", RemoteShadowFile.AbsolutePath, ZipPath ), "", null );
+				// Use RPC utility if the zip is on remote mac
+				Hashtable Result = RPCUtilHelper.Command( "/", String.Format( "unzip -o {0} -d {1}", ZipSrcPath, ZipDstPath ), "", null );
 
 				foreach ( DictionaryEntry Entry in Result )
 				{
@@ -1009,15 +1039,8 @@ namespace UnrealBuildTool
 					// Copy bundled assets from additional frameworks to the intermediate assets directory (so they can get picked up during staging)
 					String LocalFrameworkAssets = Path.GetFullPath( Target.ProjectDirectory + "/Intermediate/IOS/FrameworkAssets" );
 
-					// Delete the local dest directory if it exists
-					if ( Directory.Exists( LocalFrameworkAssets ) )
-					{
-						Directory.Delete( LocalFrameworkAssets, true );
-					}
-
-					// Create the intermediate local directory
-					string ResultsText;
-					RunExecutableAndWait( "mkdir", String.Format( "-p {0}", LocalFrameworkAssets ), out ResultsText );
+					// Clean the local dest directory if it exists
+					CleanIntermediateDirectory( LocalFrameworkAssets );
 
 					foreach ( UEBuildFramework Framework in RememberedAdditionalFrameworks )
 					{
@@ -1026,6 +1049,11 @@ namespace UnrealBuildTool
 							continue;		// Only care if we need to copy bundle assets
 						}
 
+						if ( Framework.OwningModule.Target != Target )
+						{
+							continue;		// This framework item doesn't belong to this target, skip it
+						}
+
 						string LocalZipPath = GetLocalFrameworkZipPath( Framework );
 
 						LocalZipPath = LocalZipPath.Replace( ".zip", "" );
@@ -1037,6 +1065,7 @@ namespace UnrealBuildTool
 
 						Log.TraceInformation( "Copying bundled asset... LocalSource: {0}, LocalDest: {1}", LocalSource, LocalDest );
 
+						string ResultsText;
 						RunExecutableAndWait( "cp", String.Format( "-R -L {0} {1}", LocalSource, LocalDest ), out ResultsText );
 					}
 				}
@@ -1190,11 +1219,7 @@ namespace UnrealBuildTool
 					String LocalFrameworkAssets = Path.GetFullPath( Target.ProjectDirectory + "/Intermediate/IOS/FrameworkAssets" );
 					String RemoteFrameworkAssets = ConvertPath( LocalFrameworkAssets );
 
-					// Delete the intermediate directory on the mac
-					RPCUtilHelper.Command( "/", String.Format( "rm -rf {0}", RemoteFrameworkAssets ), "", null );
-
-					// Create a fresh intermediate after we delete it
-					RPCUtilHelper.Command( "/", String.Format( "mkdir -p {0}", RemoteFrameworkAssets ), "", null );
+					CleanIntermediateDirectory( RemoteFrameworkAssets );
 
 					// Delete the local dest directory if it exists
 					if ( Directory.Exists( LocalFrameworkAssets ) )
@@ -1209,6 +1234,11 @@ namespace UnrealBuildTool
 							continue;		// Only care if we need to copy bundle assets
 						}
 
+						if ( Framework.OwningModule.Target != Target )
+						{
+							continue;		// This framework item doesn't belong to this target, skip it
+						}
+
 						string RemoteZipPath = GetRemoteFrameworkZipPath( Framework );
 
 						RemoteZipPath = RemoteZipPath.Replace( ".zip", "" );
