Fixed issues with loading game plugins that depend on other plugins on Mac

diff --git a/Engine/Source/Programs/UnrealBuildTool/Mac/MacToolChain.cs b/Engine/Source/Programs/UnrealBuildTool/Mac/MacToolChain.cs
--- a/Engine/Source/Programs/UnrealBuildTool/Mac/MacToolChain.cs
+++ b/Engine/Source/Programs/UnrealBuildTool/Mac/MacToolChain.cs
@@ -518,31 +518,18 @@ namespace UnrealBuildTool
 			// This is fixed in later step, FixDylibDependencies. For this and to know what libraries to copy whilst creating an app bundle,
 			// we gather the list of engine dylibs.
 			List<string> EngineAndGameLibraries = new List<string>();
-			
+
 			string DylibsPath = "@executable_path";
 
 			string AbsolutePath = OutputFile.AbsolutePath.Replace("\\", "/");
 			if (!AbsolutePath.Contains("/Engine/Binaries/Mac/"))
 			{
-				if (AbsolutePath.Contains("/Engine/Plugins/"))
-				{
-					string RelativePath = Utils.MakePathRelativeTo(Path.GetDirectoryName(OutputFile.AbsolutePath), ".");
-					if (!LinkEnvironment.Config.bIsBuildingConsoleApplication)
-					{
-						DylibsPath += Path.Combine("/../../../..", RelativePath);
-					}
-					else
-					{
-						DylibsPath += Path.Combine("/..", RelativePath);
-					}
-				}
-				else
-				{
-					DylibsPath = "@loader_path";
-				}
+				DylibsPath = AbsolutePath.Contains("/Plugins/") ? "@rpath" : "@loader_path";
+			}
+			if (!bIsBuildingLibrary)
+			{
+				LinkCommand += " -rpath @loader_path -rpath @executable_path";
 			}
-
-			DylibsPath = DylibsPath.Replace("\\", "/");
 
 			List<string> ThirdPartyLibraries = new List<string>();
 
@@ -564,9 +551,10 @@ namespace UnrealBuildTool
 				}
 			}
 
-
 			if (!bIsBuildingLibrary || LinkEnvironment.Config.bIncludeDependentLibrariesInLibrary)
 			{
+				List<string> RPaths = new List<string>();
+
 				// Add the additional libraries to the argument list.
 				foreach (string AdditionalLibrary in LinkEnvironment.Config.AdditionalLibraries)
 				{
@@ -624,6 +612,17 @@ namespace UnrealBuildTool
 					{
 						LinkCommand += string.Format(" -l{0}", AdditionalLibrary);
 					}
+
+					bool bIsGamePlugin = AdditionalLibrary.Contains("/Plugins/") && !AdditionalLibrary.Contains("/Engine/Plugins/"); // @todo: some better way of determining if it's a game plugin
+					if (bIsGamePlugin && Path.GetDirectoryName(AdditionalLibrary) != Path.GetDirectoryName(AbsolutePath))
+					{
+						string RelativePath = Utils.MakePathRelativeTo(Path.GetDirectoryName(AdditionalLibrary), Path.GetDirectoryName(AbsolutePath));
+						if (!RPaths.Contains(RelativePath))
+						{
+							RPaths.Add(RelativePath);
+							LinkCommand += " -rpath \"@loader_path/" + RelativePath + "\"";
+						}
+					}
 				}
 
 				foreach (string AdditionalLibrary in LinkEnvironment.Config.DelayLoadDLLs)
