Simplified code that's searching for custom resources (Info.plist, icons, xib) in MacToolChain

diff --git a/Engine/Source/Programs/UnrealBuildTool/Mac/MacToolChain.cs b/Engine/Source/Programs/UnrealBuildTool/Mac/MacToolChain.cs
--- a/Engine/Source/Programs/UnrealBuildTool/Mac/MacToolChain.cs
+++ b/Engine/Source/Programs/UnrealBuildTool/Mac/MacToolChain.cs
@@ -836,38 +836,42 @@ namespace UnrealBuildTool
 					string BundleVersion = ExeName.StartsWith("UnrealEngineLauncher") ? LoadLauncherDisplayVersion() : LoadEngineDisplayVersion();
 					string EngineSourcePath = ConvertPath(Directory.GetCurrentDirectory()).Replace("$", "\\$");
 
-					// Copy resources
-					string CustomIcon = EngineSourcePath + "/Programs/" + GameName + "/Resources/Mac/" + GameName + ".icns";
-					if (!File.Exists(CustomIcon))
+					string UProjectFilePath = UProjectInfo.GetProjectFilePath(GameName);
+					string CustomResourcesPath = "";
+					if (string.IsNullOrEmpty(UProjectFilePath))
 					{
-						CustomIcon = EngineSourcePath + "/Programs/NoRedist/" + GameName + "/Resources/Mac/" + GameName + ".icns";
+						string[] TargetFiles = Directory.GetFiles(EngineSourcePath, GameName + ".Target.cs", SearchOption.AllDirectories);
+						if (TargetFiles.Length == 1)
+						{
+							CustomResourcesPath = Path.GetDirectoryName(TargetFiles[0]) + "/Resources/Mac";
+						}
+						else
+						{
+							Log.TraceWarning("Found {0} Target.cs files for {1}", TargetFiles.Length, GameName);
+						}
 					}
-					if (File.Exists(CustomIcon))
+					else
 					{
-						AppendMacLine(CreateAppBundleScript, "cp -f \"{0}\" \"{2}.app/Contents/Resources/{1}.icns\"", CustomIcon, IconName, ExeName);
+						CustomResourcesPath = Path.GetDirectoryName(UProjectFilePath) + "/Source/" + GameName + "/Resources/Mac";
 					}
-					else
+
+					// Copy resources
+					string CustomIcon = CustomResourcesPath + "/" + GameName + ".icns";
+					if (!File.Exists(CustomIcon))
 					{
-						AppendMacLine(CreateAppBundleScript, "cp -f \"{0}/Runtime/Launch/Resources/Mac/{1}.icns\" \"{2}.app/Contents/Resources/{1}.icns\"", EngineSourcePath, IconName, ExeName);
+						CustomIcon = EngineSourcePath + "/Runtime/Launch/Resources/Mac/" + IconName + ".icns";
 					}
+					AppendMacLine(CreateAppBundleScript, "cp -f \"{0}\" \"{2}.app/Contents/Resources/{1}.icns\"", CustomIcon, IconName, ExeName);
 
 					if (ExeName.StartsWith("UE4Editor"))
 					{
 						AppendMacLine(CreateAppBundleScript, "cp -f \"{0}/Runtime/Launch/Resources/Mac/UProject.icns\" \"{1}.app/Contents/Resources/UProject.icns\"", EngineSourcePath, ExeName);
 					}
 
-					string InfoPlistFile = EngineSourcePath + "/Programs/" + GameName + "/Resources/Mac/Info.plist";
+					string InfoPlistFile = CustomResourcesPath + "/Info.plist";
 					if (!File.Exists(InfoPlistFile))
 					{
-						InfoPlistFile = EngineSourcePath + "/Programs/NoRedist/" + GameName + "/Resources/Mac/Info.plist";
-						if (!File.Exists(InfoPlistFile))
-						{
-							InfoPlistFile = EngineSourcePath + "/Programs/Mac/" + GameName + "/Resources/Mac/Info.plist";
-						}
-						if (!File.Exists(InfoPlistFile))
-						{
-							InfoPlistFile = EngineSourcePath + "/Runtime/Launch/Resources/Mac/" + (GameName.EndsWith("Editor") ? "Info-Editor.plist" : "Info.plist");
-						}
+						InfoPlistFile = EngineSourcePath + "/Runtime/Launch/Resources/Mac/" + (GameName.EndsWith("Editor") ? "Info-Editor.plist" : "Info.plist");
 					}
 					AppendMacLine(CreateAppBundleScript, "cp -f \"{0}\" \"{1}.app/Contents/Info.plist\"", InfoPlistFile, ExeName);
 
@@ -891,18 +895,10 @@ namespace UnrealBuildTool
 					AppendMacLine(CreateAppBundleScript, "cp -f \"{0}/ThirdParty/Mac/RadioEffectUnit/RadioEffectUnit.component/Contents/Info.plist\" \"{1}.app/Contents/Resources/RadioEffectUnit.component/Contents/Info.plist\"", EngineSourcePath, ExeName);
 
 					// Compile XIB resource
-					string XIBFile = EngineSourcePath + "/Programs/" + GameName + "/Resources/Mac/MainMenu.xib";
+					string XIBFile = CustomResourcesPath + "/MainMenu.xib";
 					if (!File.Exists(XIBFile))
 					{
-						XIBFile = EngineSourcePath + "/Programs/NoRedist/" + GameName + "/Resources/Mac/MainMenu.xib";
-						if (!File.Exists(XIBFile))
-						{
-							XIBFile = EngineSourcePath + "/Programs/Mac/" + GameName + "/Resources/Mac/MainMenu.xib";
-						}
-						if (!File.Exists(XIBFile))
-						{
-							XIBFile = EngineSourcePath + "/Runtime/Launch/Resources/Mac/English.lproj/MainMenu.xib";
-						}
+						XIBFile = EngineSourcePath + "/Runtime/Launch/Resources/Mac/English.lproj/MainMenu.xib";
 					}
 					AppendMacLine(CreateAppBundleScript, "\"{0}usr/bin/ibtool\" --errors --warnings --notices --output-format human-readable-text --compile \"{2}.app/Contents/Resources/English.lproj/MainMenu.nib\" \"{1}\" --sdk \"{0}/Platforms/MacOSX.platform/Developer/SDKs/MacOSX{3}.sdk\"", DeveloperDir, XIBFile, ExeName, MacOSSDKVersion);
 
