Quick fix for launching Mac UnrealFrontend broken by recent changes to plugin loading

diff --git a/Engine/Source/Programs/UnrealBuildTool/Mac/MacToolChain.cs b/Engine/Source/Programs/UnrealBuildTool/Mac/MacToolChain.cs
--- a/Engine/Source/Programs/UnrealBuildTool/Mac/MacToolChain.cs
+++ b/Engine/Source/Programs/UnrealBuildTool/Mac/MacToolChain.cs
@@ -613,11 +613,10 @@ namespace UnrealBuildTool
 						LinkCommand += string.Format(" -l{0}", AdditionalLibrary);
 					}
 
-					bool bIsGamePlugin = AdditionalLibrary.Contains("/Plugins/") && !AdditionalLibrary.Contains("/Engine/Plugins/"); // @todo: some better way of determining if it's a game plugin
-					if (bIsGamePlugin && Path.GetDirectoryName(AdditionalLibrary) != Path.GetDirectoryName(AbsolutePath))
+					if (AdditionalLibrary.Contains("/Plugins/") && Path.GetDirectoryName(AdditionalLibrary) != Path.GetDirectoryName(AbsolutePath))
 					{
 						string RelativePath = Utils.MakePathRelativeTo(Path.GetDirectoryName(AdditionalLibrary), Path.GetDirectoryName(AbsolutePath));
-						if (!RPaths.Contains(RelativePath))
+						if (!RelativePath.Contains(Path.GetDirectoryName(AdditionalLibrary)) && !RPaths.Contains(RelativePath))
 						{
 							RPaths.Add(RelativePath);
 							LinkCommand += " -rpath \"@loader_path/" + RelativePath + "\"";
