Better project INIs support in UBT. Added project INIs to UAT project properties.

diff --git a/Engine/Source/Programs/AutomationTool/ProjectUtils.cs b/Engine/Source/Programs/AutomationTool/ProjectUtils.cs
--- a/Engine/Source/Programs/AutomationTool/ProjectUtils.cs
+++ b/Engine/Source/Programs/AutomationTool/ProjectUtils.cs
@@ -60,6 +60,11 @@ namespace AutomationTool
 		/// </summary>
 		public Dictionary<TargetRules.TargetType, SingleTargetProperties> Targets = new Dictionary<TargetRules.TargetType, SingleTargetProperties>();
 
+		/// <summary>
+		/// List of all Engine ini files for this project
+		/// </summary>
+		public Dictionary<UnrealTargetPlatform, ConfigCacheIni> EngineConfigs = new Dictionary<UnrealTargetPlatform,ConfigCacheIni>();
+
 		/// <summary>
 		/// List of all programs detected for this project.
 		/// </summary>
@@ -175,6 +180,18 @@ namespace AutomationTool
 				Properties.bIsCodeBasedProject = uprojectStr.Contains("\"Modules\"");
 			}
 
+			// Get all ini files
+			var EngineDirectory = CommandUtils.CombinePaths(CommandUtils.CmdEnv.LocalRoot, "Engine");
+			CommandUtils.Log("Loading ini files for {0}", RawProjectPath);
+			foreach (UnrealTargetPlatform TargetPlatformType in Enum.GetValues(typeof(UnrealTargetPlatform)))
+			{
+				if (TargetPlatformType != UnrealTargetPlatform.Unknown)
+				{
+					var Config = new ConfigCacheIni(TargetPlatformType, "Engine", Path.GetDirectoryName(RawProjectPath), EngineDirectory);
+					Properties.EngineConfigs.Add(TargetPlatformType, Config);
+				}
+			}
+
 			return Properties;
 		}
 
diff --git a/Engine/Source/Programs/UnrealBuildTool/Configuration/EngineConfiguration.cs b/Engine/Source/Programs/UnrealBuildTool/Configuration/EngineConfiguration.cs
--- a/Engine/Source/Programs/UnrealBuildTool/Configuration/EngineConfiguration.cs
+++ b/Engine/Source/Programs/UnrealBuildTool/Configuration/EngineConfiguration.cs
@@ -6,135 +6,474 @@ using System.IO;
 
 namespace UnrealBuildTool
 {
-	public static class EngineConfiguration
+	/// <summary>
+	/// Exception when parsing ini files
+	/// </summary>
+	public class IniParsingException : Exception
+	{
+		public IniParsingException(string Message)
+			: base(Message)
+		{}
+		public IniParsingException(string Format, params object[] Args)
+			: base(String.Format(Format, Args))
+		{}
+	}
+
+	/// <summary>
+	/// Equivalent of FConfigCacheIni. Parses ini files.
+	/// </summary>
+	public class ConfigCacheIni
 	{
 		/// <summary>
-		/// Reads an array from the config files, matching a given section and key.
+		/// List of values (or a single value)
 		/// </summary>
-		public static void ReadArray(string ProjectDirectory, UnrealTargetPlatform Platform, string BaseIniName, string FindSection, string FindKey, List<string> Values)
+		public class IniValues : List<string>
 		{
-			foreach (string IniFileName in EnumerateIniFileNames(ProjectDirectory, Platform, BaseIniName))
+			public override string ToString()
+			{
+				return String.Join(",", ToArray());
+			}
+		}
+
+		/// <summary>
+		/// Ini section (map of keys and values)
+		/// </summary>
+		public class IniSection : Dictionary<string, IniValues>
+		{
+			public IniSection()
+				: base(StringComparer.InvariantCultureIgnoreCase)
+			{}
+			public override string ToString()
+			{
+				return "IniSection";
+			}
+		}
+
+		/// <summary>
+		/// All sections parsed from ini file
+		/// </summary>
+		Dictionary<string, IniSection> Sections;
+
+		/// <summary>
+		/// Constructor. Parses ini hierarchy for the specified platform and project.
+		/// </summary>
+		/// <param name="ProjectDirectory">Project path</param>
+		/// <param name="Platform">Target platform</param>
+		/// <param name="BaseIniName">Ini name (Engine, Editor, etc)</param>
+		public ConfigCacheIni(UnrealTargetPlatform Platform, string BaseIniName, string ProjectDirectory, string EngineDirectory = null)
+		{
+			Sections = new Dictionary<string, IniSection>(StringComparer.InvariantCultureIgnoreCase);
+			if (String.IsNullOrEmpty(EngineDirectory))
+			{
+				EngineDirectory = BuildConfiguration.RelativeEnginePath;
+			}
+			foreach (var IniFileName in EnumerateIniFileNames(ProjectDirectory, EngineDirectory, Platform, BaseIniName))
 			{
 				if (File.Exists(IniFileName))
 				{
-					bool bInCorrectSection = false;
-					foreach(string Line in File.ReadAllLines(IniFileName))
-					{
-						// Skip the whitespace at the start of the line
-						int LineMin = 0;
-						while(LineMin < Line.Length && Char.IsWhiteSpace(Line[LineMin])) LineMin++;
+					ParseIniFile(IniFileName);
+				}
+			}
+		}
 
-						// Skip the whitespace at the end of the line
-						int LineMax = Line.Length;
-						while (LineMax > LineMin && Char.IsWhiteSpace(Line[LineMax - 1])) LineMax--;
+		/// <summary>
+		/// Finds a section in INI
+		/// </summary>
+		/// <param name="SectionName"></param>
+		/// <returns>Found section or null</returns>
+		public IniSection FindSection(string SectionName)
+		{
+			IniSection Section;
+			Sections.TryGetValue(SectionName, out Section);
+			return Section;
+		}
 
-						// Check it's not a comment
-						if (LineMax > LineMin && Line[LineMin] != ';')
-						{
-							if (Line[LineMin] == '[' && Line[LineMax - 1] == ']')
-							{
-								// This is a section heading. Check if it's the section we're looking for.
-								bInCorrectSection = (FindSection.Length == LineMax - LineMin - 2 && String.Compare(Line, LineMin + 1, FindSection, 0, LineMax - LineMin - 2) == 0);
-							}
-							else if (bInCorrectSection)
-							{
-								// Split the line into key/value.
-								TryParseArrayKeyValue(Line, LineMin, LineMax, FindKey, Values);
-							}
-						}
-					}
+		/// <summary>
+		/// Finds values associated with the specified key (does not copy the list)
+		/// </summary>
+		private bool GetList(string SectionName, string Key, out IniValues Value)
+		{
+			bool Result = false;
+			var Section = FindSection(SectionName);
+			Value = null;
+			if (Section != null)
+			{
+				if (Section.TryGetValue(Key, out Value))
+				{
+					Result = true;
 				}
 			}
+			return Result;
 		}
 
 		/// <summary>
-		/// Returns a list of INI filenames for the given project
+		/// Gets all values associated with the specified key
 		/// </summary>
-		private static IEnumerable<string> EnumerateIniFileNames(string ProjectDirectory, UnrealTargetPlatform Platform, string BaseIniName)
+		/// <param name="SectionName">Section where the key is located</param>
+		/// <param name="Key">Key name</param>
+		/// <param name="Value">Copy of the list containing all values associated with the specified key</param>
+		/// <returns>True if the key exists</returns>
+		public bool GetArray(string SectionName, string Key, out List<string> Value)
 		{
-			string PlatformName = GetIniPlatformName(Platform);
+			Value = null;
+			IniValues ValueList;
+			bool Result = GetList(SectionName, Key, out ValueList);
+			if (Result)
+			{
+				Value = new List<string>(ValueList);
+			}
+			return Result;
+		}
 
-			// Engine/Config/Base.ini (included in every ini type, required)
-			yield return Path.Combine(BuildConfiguration.RelativeEnginePath, "Config", "Base.ini");
+		/// <summary>
+		/// Gets a single string value associated with the specified key.
+		/// </summary>
+		/// <param name="SectionName">Section name</param>
+		/// <param name="Key">Key name</param>
+		/// <param name="Value">Value associated with the specified key. If the key has more than one value, only the first one is returned</param>
+		/// <returns>True if the key exists</returns>
+		public bool GetString(string SectionName, string Key, out string Value)
+		{
+			Value = String.Empty;
+			IniValues ValueList;
+			bool Result = GetList(SectionName, Key, out ValueList);
+			if (Result && ValueList != null && ValueList.Count > 0)
+			{
+				Value = ValueList[0];
+				Result = true;
+			}
+			else
+			{
+				Result = false;
+			}
+			return Result;
+		}
 
-			// Engine/Config/Base* ini
-			yield return Path.Combine(BuildConfiguration.RelativeEnginePath, "Config", "Base" +  BaseIniName + ".ini");
+		/// <summary>
+		/// Gets a single bool value associated with the specified key.
+		/// </summary>
+		/// <param name="SectionName">Section name</param>
+		/// <param name="Key">Key name</param>
+		/// <param name="Value">Value associated with the specified key. If the key has more than one value, only the first one is returned</param>
+		/// <returns>True if the key exists</returns>
+		public bool GetBool(string SectionName, string Key, out bool Value)
+		{			
+			Value = false;
+			string TextValue;
+			bool Result = GetString(SectionName, Key, out TextValue);
+			if (Result)
+			{
+				// C# Boolean type expects "False" or "True" but since we're not case sensitive, we need to suppor that manually
+				if (String.Compare(TextValue, "true", false) == 0 || String.Compare(TextValue, "1") == 0)
+				{
+					Value = true;
+				}
+				else if (String.Compare(TextValue, "false", false) == 0 || String.Compare(TextValue, "0") == 0)
+				{
+					Value = false;
+				}
+				else
+				{
+					// Failed to parse
+					Result = false;
+				}
+			}
+			return Result;
+		}
 
-			// Game/Config/Default* ini
-			yield return Path.Combine(ProjectDirectory, "Config", "Default" + BaseIniName + ".ini");
+		/// <summary>
+		/// Gets a single Int32 value associated with the specified key.
+		/// </summary>
+		/// <param name="SectionName">Section name</param>
+		/// <param name="Key">Key name</param>
+		/// <param name="Value">Value associated with the specified key. If the key has more than one value, only the first one is returned</param>
+		/// <returns>True if the key exists</returns>
+		public bool GetInt32(string SectionName, string Key, out int Value)
+		{
+			Value = 0;
+			string TextValue;
+			bool Result = GetString(SectionName, Key, out TextValue);
+			if (Result)
+			{
+				Result = Int32.TryParse(TextValue, out Value);
+			}
+			return Result;
+		}
 
-			// Engine/Config/Platform/Platform* ini
-			yield return Path.Combine(BuildConfiguration.RelativeEnginePath, "Config", PlatformName, PlatformName + BaseIniName + ".ini");
-	
-			// Game/Config/Platform/Platform* ini
-			yield return Path.Combine(ProjectDirectory, "Config", PlatformName, PlatformName + BaseIniName + ".ini");
+		/// <summary>
+		/// Gets a single float value associated with the specified key.
+		/// </summary>
+		/// <param name="SectionName">Section name</param>
+		/// <param name="Key">Key name</param>
+		/// <param name="Value">Value associated with the specified key. If the key has more than one value, only the first one is returned</param>
+		/// <returns>True if the key exists</returns>
+		public bool GetSingle(string SectionName, string Key, out float Value)
+		{
+			Value = 0.0f;
+			string TextValue;
+			bool Result = GetString(SectionName, Key, out TextValue);
+			if (Result)
+			{
+				Result = Single.TryParse(TextValue, out Value);
+			}
+			return Result;
 		}
 
 		/// <summary>
-		/// Returns the platform name to use as part of platform-specific config files
+		/// Gets a single double value associated with the specified key.
 		/// </summary>
-		private static string GetIniPlatformName(UnrealTargetPlatform TargetPlatform)
+		/// <param name="SectionName">Section name</param>
+		/// <param name="Key">Key name</param>
+		/// <param name="Value">Value associated with the specified key. If the key has more than one value, only the first one is returned</param>
+		/// <returns>True if the key exists</returns>
+		public bool GetDouble(string SectionName, string Key, out double Value)
 		{
-			if(TargetPlatform == UnrealTargetPlatform.Win32 || TargetPlatform == UnrealTargetPlatform.Win64)
+			Value = 0.0;
+			string TextValue;
+			bool Result = GetString(SectionName, Key, out TextValue);
+			if (Result)
 			{
-				return "Windows";
+				Result = Double.TryParse(TextValue, out Value);
+			}
+			return Result;
+		}
+
+		/// <summary>
+		/// List of actions that can be performed on a single line from ini file
+		/// </summary>
+		enum ParseAction
+		{
+			None,
+			New,
+			Add,
+			Remove
+		}
+
+		/// <summary>
+		/// Checks what action should be performed on a single line from ini file
+		/// </summary>
+		private ParseAction GetActionForLine(ref string Line)
+		{
+			if (String.IsNullOrEmpty(Line) || Line.StartsWith(";") || Line.StartsWith("//"))
+			{
+				return ParseAction.None;
+			}
+			else if (Line.StartsWith("-"))
+			{
+				Line = Line.Substring(1).TrimStart();
+				return ParseAction.Remove;
+			}
+			else if (Line.StartsWith("+"))
+			{
+				Line = Line.Substring(1).TrimStart();
+				return ParseAction.Add;
 			}
 			else
 			{
-				return Enum.GetName(typeof(UnrealTargetPlatform), TargetPlatform);
+				return ParseAction.New;
 			}
 		}
 
 		/// <summary>
-		/// Try to match a config file key, and adjust an array of values to reflect it.
+		/// Loads and parses ini file.
 		/// </summary>
-		private static void TryParseArrayKeyValue(string Line, int LineMin, int LineMax, string FindKey, List<string> Values)
+		public void ParseIniFile(string Filename)
 		{
-			int EqualsIdx = Line.IndexOf('=', LineMin, LineMax - LineMin);
-			if (EqualsIdx != -1)
-			{
-				// Skip the + or - at the start of the key name
-				int KeyMin = LineMin;
-				if (Line[KeyMin] == '+' || Line[KeyMin] == '-')
-				{
-					KeyMin++;
-				}
+			var IniLines = File.ReadAllLines(Filename);
+			IniSection CurrentSection = null;
 
-				// Skip the whitespace between the key name and equals sign
-				int KeyMax = EqualsIdx;
-				while (KeyMax > KeyMin && Char.IsWhiteSpace(Line[KeyMax - 1]))
+			// Line Index for exceptions
+			var LineIndex = 1;
+			var bMultiLine = false;
+			var SingleValue = "";
+			var Key = "";
+			var LastAction = ParseAction.None;
+
+			// Parse each line
+			foreach (var Line in IniLines)
+			{
+				var TrimmedLine = Line.Trim();
+				// Multiline value support
+				bool bWasMultiLine = bMultiLine;
+				bMultiLine = TrimmedLine.EndsWith("\\");
+				if (bMultiLine)
 				{
-					KeyMax--;
+					TrimmedLine = TrimmedLine.Substring(0, TrimmedLine.Length - 1).TrimEnd();
 				}
-
-				// Check if it matches the key we're looking for
-				int KeyLength = KeyMax - KeyMin;
-				if (FindKey.Length == KeyLength && String.Compare(Line, KeyMin, FindKey, 0, KeyLength, true) == 0)
+				if (!bWasMultiLine)
 				{
-					// Get the start and end of the value
-					int ValueMin = EqualsIdx + 1;
-					while (ValueMin < LineMax && Char.IsWhiteSpace(Line[ValueMin]))
+					if (TrimmedLine.StartsWith("["))
 					{
-						ValueMin++;
+						CurrentSection = FindOrAddSection(TrimmedLine, Filename, LineIndex);
+						LastAction = ParseAction.None;
 					}
-					int ValueMax = LineMax;
-					if (ValueMax > ValueMin + 1 && Line[ValueMin] == '\"' && Line[ValueMax - 1] == '\"')
+					else
 					{
-						ValueMin++;
-						ValueMax--;
+						if (LastAction != ParseAction.None)
+						{
+							throw new IniParsingException("Parsing new key/value pair when the previous one has not yet been processed ({0}, {1}) in {2}, line {3}: {4}", Key, SingleValue, Filename, LineIndex, TrimmedLine);
+						}
+						if (CurrentSection == null)
+						{
+							throw new IniParsingException("Trying to parse key/value pair that doesn't belong to any section in {0}, line {1}: {2}", Filename, LineIndex, TrimmedLine);
+						}
+						// Check if the line is empty or a comment, also remove any +/- markers
+						LastAction = GetActionForLine(ref TrimmedLine);
+						if (LastAction != ParseAction.None)
+						{
+							ParseKeyValuePair(TrimmedLine, Filename, LineIndex, out Key, out SingleValue);
+						}
 					}
+				}
+				if (bWasMultiLine)
+				{
+					SingleValue += TrimmedLine;
+				}
+				if (!bMultiLine && LastAction != ParseAction.None)
+				{
+					ProcessKeyValuePair(CurrentSection, Key, SingleValue, LastAction);
+					LastAction = ParseAction.None;
+					SingleValue = "";
+					Key = "";
+				}
+				LineIndex++;
+			}
+		}
+
+		/// <summary>
+		/// Splits a line into key and value
+		/// </summary>
+		private static void ParseKeyValuePair(string TrimmedLine, string Filename, int LineIndex, out string Key, out string Value)
+		{			
+			var AssignIndex = TrimmedLine.IndexOf('=');
+			if (AssignIndex < 0)
+			{
+				throw new IniParsingException("Failed to find value when parsing {0}, line {1}: {2}", Filename, LineIndex, TrimmedLine);
+			}
+			Key = TrimmedLine.Substring(0, AssignIndex).Trim();
+			if (String.IsNullOrEmpty(Key))
+			{
+				throw new IniParsingException("Empty key when parsing {0}, line {1}: {2}", Filename, LineIndex, TrimmedLine);
+			}
+			Value = TrimmedLine.Substring(AssignIndex + 1).Trim();
+			if (Value.StartsWith("\""))
+			{
+				// Remove quotes
+				var QuoteEnd = Value.LastIndexOf('\"');
+				if (QuoteEnd == 0)
+				{
+					throw new IniParsingException("Mismatched quotes when parsing {0}, line {1}: {2}", Filename, LineIndex, TrimmedLine);
+				}
+				Value = Value.Substring(1, Value.Length - 2);
+			}
+		}
 
-					// Add or remove it from the list
-					string Value = Line.Substring(ValueMin, ValueMax - ValueMin);
-					if (Line.StartsWith("-"))
+		/// <summary>
+		/// Processes parsed key/value pair
+		/// </summary>
+		private static void ProcessKeyValuePair(IniSection CurrentSection, string Key, string SingleValue, ParseAction Action)
+		{
+			switch (Action)
+			{
+				case ParseAction.New:
 					{
-						Values.RemoveAll(x => String.Compare(x, Value, true) == 0);
+						// New/replace
+						IniValues Value;
+						if (CurrentSection.TryGetValue(Key, out Value) == false)
+						{
+							Value = new IniValues();
+							CurrentSection.Add(Key, Value);
+						}
+						Value.Clear();
+						Value.Add(SingleValue);						
 					}
-					else
+					break;
+				case ParseAction.Add:
 					{
-						Values.Add(Value);
+						IniValues Value;
+						if (CurrentSection.TryGetValue(Key, out Value) == false)
+						{
+							Value = new IniValues();
+							CurrentSection.Add(Key, Value);
+						}
+						Value.Add(SingleValue);
 					}
-				}
+					break;
+				case ParseAction.Remove:
+					{
+						IniValues Value;
+						if (CurrentSection.TryGetValue(Key, out Value))
+						{
+							var ExistingIndex = Value.FindIndex(X => (String.Compare(Key, X, true) == 0));
+							if (ExistingIndex >= 0)
+							{
+								Value.RemoveAt(ExistingIndex);
+							}
+						}
+					}
+					break;
+			}
+		}
+
+		/// <summary>
+		/// Finds an existing section or adds a new one
+		/// </summary>
+		private IniSection FindOrAddSection(string TrimmedLine, string Filename, int LineIndex)
+		{
+			var SectionEndIndex = TrimmedLine.IndexOf(']');
+			if (SectionEndIndex != (TrimmedLine.Length - 1))
+			{
+				throw new IniParsingException("Mismatched brackets when parsing section name in {0}, line {1}: {2}", Filename, LineIndex, TrimmedLine);
+			}
+			var SectionName = TrimmedLine.Substring(1, TrimmedLine.Length - 2);
+			if (String.IsNullOrEmpty(SectionName))
+			{
+				throw new IniParsingException("Empty section name when parsing {0}, line {1}: {2}", Filename, LineIndex, TrimmedLine);
+			}
+			IniSection CurrentSection;
+			if (Sections.TryGetValue(SectionName, out CurrentSection) == false)
+			{
+				CurrentSection = new IniSection();
+				Sections.Add(SectionName, CurrentSection);
+			}			
+			return CurrentSection;
+		}
+
+		/// <summary>
+		/// Returns a list of INI filenames for the given project
+		/// </summary>
+		private static IEnumerable<string> EnumerateIniFileNames(string ProjectDirectory, string EngineDirectory, UnrealTargetPlatform Platform, string BaseIniName)
+		{
+			string PlatformName = GetIniPlatformName(Platform);
+
+			// Engine/Config/Base.ini (included in every ini type, required)
+			yield return Path.Combine(EngineDirectory, "Config", "Base.ini");
+
+			// Engine/Config/Base* ini
+			yield return Path.Combine(EngineDirectory, "Config", "Base" + BaseIniName + ".ini");
+
+			// Game/Config/Default* ini
+			yield return Path.Combine(ProjectDirectory, "Config", "Default" + BaseIniName + ".ini");
+
+			// Engine/Config/Platform/Platform* ini
+			yield return Path.Combine(EngineDirectory, "Config", PlatformName, PlatformName + BaseIniName + ".ini");
+
+			// Game/Config/Platform/Platform* ini
+			yield return Path.Combine(ProjectDirectory, "Config", PlatformName, PlatformName + BaseIniName + ".ini");
+		}
+
+		/// <summary>
+		/// Returns the platform name to use as part of platform-specific config files
+		/// </summary>
+		private static string GetIniPlatformName(UnrealTargetPlatform TargetPlatform)
+		{
+			if (TargetPlatform == UnrealTargetPlatform.Win32 || TargetPlatform == UnrealTargetPlatform.Win64)
+			{
+				return "Windows";
+			}
+			else
+			{
+				return Enum.GetName(typeof(UnrealTargetPlatform), TargetPlatform);
 			}
 		}
 	}
diff --git a/Engine/Source/Programs/UnrealBuildTool/Configuration/UEBuildTarget.cs b/Engine/Source/Programs/UnrealBuildTool/Configuration/UEBuildTarget.cs
--- a/Engine/Source/Programs/UnrealBuildTool/Configuration/UEBuildTarget.cs
+++ b/Engine/Source/Programs/UnrealBuildTool/Configuration/UEBuildTarget.cs
@@ -1957,15 +1957,28 @@ namespace UnrealBuildTool
 				// Build the enabled plugin list
 				if (ShouldCompileMonolithic() || Rules.Type == TargetRules.TargetType.Program)
 				{
-					List<string> FilterPluginNames = new List<string>(Rules.AdditionalPlugins);
+					var FilterPluginNames = new List<string>(Rules.AdditionalPlugins);
+					var ConfigDirectory = String.Empty;
+					var ConfigKey = String.Empty;
+
 					if (UEBuildConfiguration.bCompileAgainstEngine)
 					{
-						EngineConfiguration.ReadArray(ProjectDirectory, Platform, "Engine", "Plugins", "EnabledPlugins", FilterPluginNames);
+						ConfigDirectory = ProjectDirectory;
+						ConfigKey = "EnabledPlugins";
 					}
 					else if (Rules.Type == TargetRules.TargetType.Program)
 					{
-						var ProgramDirectory = Path.Combine(ProjectDirectory, "Programs", GetTargetName());
-						EngineConfiguration.ReadArray(ProgramDirectory, Platform, "Engine", "Plugins", "ProgramEnabledPlugins", FilterPluginNames);
+						ConfigDirectory = Path.Combine(ProjectDirectory, "Programs", GetTargetName());
+						ConfigKey = "ProgramEnabledPlugins";
+					}
+					if (!String.IsNullOrEmpty(ConfigDirectory))
+					{
+						var Config = new ConfigCacheIni(Platform, "Engine", ConfigDirectory);
+						List<string> ConfigPlugins;
+						if (Config.GetArray("Plugins", ConfigKey, out ConfigPlugins))
+						{
+							FilterPluginNames.AddRange(ConfigPlugins);
+						}
 					}
 					EnabledPlugins.AddRange(ValidPlugins.Where(x => FilterPluginNames.Contains(x.Name)));
 				}
