Fix exception in UBT when config file starts with a blank line.

diff --git a/Engine/Source/Programs/UnrealBuildTool/Configuration/EngineConfiguration.cs b/Engine/Source/Programs/UnrealBuildTool/Configuration/EngineConfiguration.cs
--- a/Engine/Source/Programs/UnrealBuildTool/Configuration/EngineConfiguration.cs
+++ b/Engine/Source/Programs/UnrealBuildTool/Configuration/EngineConfiguration.cs
@@ -307,7 +307,7 @@ namespace UnrealBuildTool
 						CurrentSection = FindOrAddSection(TrimmedLine, Filename, LineIndex);
 						LastAction = ParseAction.None;
 					}
-					else
+					else if(TrimmedLine.Length > 0)
 					{
 						if (LastAction != ParseAction.None)
 						{
