UBT: Proper fix for exceptions when parsing ini files that start with an empty line

diff --git a/Engine/Source/Programs/UnrealBuildTool/Configuration/EngineConfiguration.cs b/Engine/Source/Programs/UnrealBuildTool/Configuration/EngineConfiguration.cs
--- a/Engine/Source/Programs/UnrealBuildTool/Configuration/EngineConfiguration.cs
+++ b/Engine/Source/Programs/UnrealBuildTool/Configuration/EngineConfiguration.cs
@@ -307,20 +307,20 @@ namespace UnrealBuildTool
 						CurrentSection = FindOrAddSection(TrimmedLine, Filename, LineIndex);
 						LastAction = ParseAction.None;
 					}
-					else if(TrimmedLine.Length > 0)
+					else
 					{
 						if (LastAction != ParseAction.None)
 						{
 							throw new IniParsingException("Parsing new key/value pair when the previous one has not yet been processed ({0}, {1}) in {2}, line {3}: {4}", Key, SingleValue, Filename, LineIndex, TrimmedLine);
 						}
-						if (CurrentSection == null)
-						{
-							throw new IniParsingException("Trying to parse key/value pair that doesn't belong to any section in {0}, line {1}: {2}", Filename, LineIndex, TrimmedLine);
-						}
 						// Check if the line is empty or a comment, also remove any +/- markers
 						LastAction = GetActionForLine(ref TrimmedLine);
 						if (LastAction != ParseAction.None)
 						{
+							if (CurrentSection == null)
+							{
+								throw new IniParsingException("Trying to parse key/value pair that doesn't belong to any section in {0}, line {1}: {2}", Filename, LineIndex, TrimmedLine);
+							}
 							ParseKeyValuePair(TrimmedLine, Filename, LineIndex, out Key, out SingleValue);
 						}
 					}
