#summary Respect Xcode's compiler options & analyze feature.
#change Parse the compiler selected in Xcode & find the Xcode plugin which contains the executable to run when it isn't the default.
#change When the analyze option is enabled (either standalone or as a build setting) we need to pass through the analyze clang flag.
#change Pass in the Xcode max. parallel build tasks default to UBT as an environment variable rather than just assuming we want to use half the CPUs.
#notes This means we can use an external compiler (even a GPL one like distcc) without any direct dependency. Editor builds still use the default compiler since it isn't clear how to know which compiler is intended to be used in that case - the Xcode environment variable won't have been set.

diff --git a/Engine/Build/BatchFiles/Mac/Build.sh b/Engine/Build/BatchFiles/Mac/Build.sh
--- a/Engine/Build/BatchFiles/Mac/Build.sh
+++ b/Engine/Build/BatchFiles/Mac/Build.sh
@@ -24,6 +24,14 @@ case $ACTION in
                 Platform=""
                 AdditionalFlags=""
 		
+                case $CLANG_STATIC_ANALYZER_MODE in
+				"deep")
+					AdditionalFlags+="-skipActionHistory"
+					;;
+				"shallow")
+					AdditionalFlags+="-skipActionHistory"
+					;;
+				esac
 
 		case $2 in 
 	                "iphoneos") 
@@ -43,6 +51,9 @@ case $ACTION in
 				AdditionalFlags+=" -deploy " 
 			;; 
 		esac
+		
+		BuildTasks=$(defaults read com.apple.dt.Xcode IDEBuildOperationMaxNumberOfConcurrentCompileTasks)
+		export NumUBTBuildTasks=$BuildTasks
 
 		echo Running command : Engine/Binaries/DotNET/UnrealBuildTool.exe $1 $Platform $3 $AdditionalFlags "$4" 
 		mono Engine/Binaries/DotNET/UnrealBuildTool.exe $1 $Platform $3 $AdditionalFlags "$4"
diff --git a/Engine/Source/Programs/UnrealBuildTool/Mac/MacToolChain.cs b/Engine/Source/Programs/UnrealBuildTool/Mac/MacToolChain.cs
--- a/Engine/Source/Programs/UnrealBuildTool/Mac/MacToolChain.cs
+++ b/Engine/Source/Programs/UnrealBuildTool/Mac/MacToolChain.cs
@@ -391,8 +391,69 @@ namespace UnrealBuildTool
 				}
 
 				CompileAction.WorkingDirectory = GetMacDevSrcRoot();
-				CompileAction.CommandPath = "xcrun";
-				CompileAction.CommandArguments = MacCompiler + Arguments + FileArguments + CompileEnvironment.Config.AdditionalArguments;
+				
+				string CompilerBundleIdentifier = Environment.GetEnvironmentVariable("GCC_VERSION");
+				string StaticAnalysisMode = Environment.GetEnvironmentVariable("CLANG_STATIC_ANALYZER_MODE");
+				if(CompilerBundleIdentifier == null || CompilerBundleIdentifier == "" 
+					|| CompilerBundleIdentifier == "com.apple.compilers.llvm.clang.1_0" 
+					|| (StaticAnalysisMode != null && StaticAnalysisMode != ""))
+				{
+					CompileAction.CommandPath = MacCompiler;
+					if(StaticAnalysisMode != null && StaticAnalysisMode != "")
+					{
+						FileArguments = " --analyze " + FileArguments;
+					}
+				}
+				else
+				{
+					CompileAction.CommandPath = null;
+					string UserDir = Environment.GetEnvironmentVariable("HOME");
+					string XcodePlugins = UserDir + "/Library/Application Support/Developer/Shared/Xcode/Plug-ins/";
+					if(Directory.Exists(XcodePlugins))
+					{
+						IEnumerable<string> Plugins = Directory.EnumerateDirectories(XcodePlugins, "*.xcplugin");
+						foreach(string Plugin in Plugins)
+						{
+							string Resources = Plugin + "/Contents/Resources";
+							string[] CompilerSpecs = Directory.GetFiles(Resources, "*.xcspec");
+							if(CompilerSpecs.Length > 0)
+							{
+								IEnumerable<string> Lines = File.ReadLines(CompilerSpecs[0]);
+								bool bFindExecPath = false;
+								bool bFoundExecPath = false;
+								foreach(string Line in Lines)
+								{
+									if(!bFindExecPath && Line.Contains(CompilerBundleIdentifier))
+									{
+										bFindExecPath = true;
+									}
+									else if(bFindExecPath && Line.Contains("ExecPath"))
+									{
+										char[] CharsToTrim = {';', '"'};
+										string ExecPath = Line.Trim().Substring(12).TrimEnd(CharsToTrim);
+										CompileAction.CommandPath = ExecPath;
+										bFoundExecPath = true;
+										break;
+									}
+								}
+								if(bFoundExecPath)
+								{
+									break;
+								}
+							}
+						}
+					}
+					else
+					{
+						throw new BuildException("Couldn't find Xcode compiler plugins path: {0}", XcodePlugins);
+					}
+					
+					if(CompileAction.CommandPath == null)
+					{
+						throw new BuildException("Couldn't find Xcode compiler: {0}", CompilerBundleIdentifier);
+					}
+				}
+				CompileAction.CommandArguments = Arguments + FileArguments + CompileEnvironment.Config.AdditionalArguments;
 				CompileAction.CommandDescription = "Compile";
 				CompileAction.StatusDescription = Path.GetFileName(SourceFile.AbsolutePath);
 				CompileAction.StatusDetailedDescription = SourceFile.Description;
diff --git a/Engine/Source/Programs/UnrealBuildTool/System/LocalExecutor.cs b/Engine/Source/Programs/UnrealBuildTool/System/LocalExecutor.cs
--- a/Engine/Source/Programs/UnrealBuildTool/System/LocalExecutor.cs
+++ b/Engine/Source/Programs/UnrealBuildTool/System/LocalExecutor.cs
@@ -326,6 +326,13 @@ namespace UnrealBuildTool
 				// @todo iosmerge: this should be looking at available memory as well since some of our
 				// Macs have a poor memory/CPU ratio
 				MaxActionsToExecuteInParallel /= 2;
+
+				string NumUBTBuildTasks = Environment.GetEnvironmentVariable("NumUBTBuildTasks");
+				Int32 MaxUBTBuildTasks = MaxActionsToExecuteInParallel;
+				if(Int32.TryParse(NumUBTBuildTasks, out MaxUBTBuildTasks))
+				{
+					MaxActionsToExecuteInParallel = MaxUBTBuildTasks;
+				}
 			}
 
             Log.TraceInformation("Performing {0} actions (max {1} parallel jobs)", Actions.Count, MaxActionsToExecuteInParallel);
