Linux: UBT: limit number of parallel jobs by RAM available.

Limit number of parallel actions by introducing memory budget per each action - arguably a better heuristic than just dividing number of cores by two.

diff --git a/Engine/Source/Programs/UnrealBuildTool/System/LocalExecutor.cs b/Engine/Source/Programs/UnrealBuildTool/System/LocalExecutor.cs
--- a/Engine/Source/Programs/UnrealBuildTool/System/LocalExecutor.cs
+++ b/Engine/Source/Programs/UnrealBuildTool/System/LocalExecutor.cs
@@ -334,6 +334,15 @@ namespace UnrealBuildTool
 					MaxActionsToExecuteInParallel = MaxUBTBuildTasks;
 				}
 			}
+			else if (ExternalExecution.GetRuntimePlatform () == UnrealTargetPlatform.Linux && Utils.IsRunningOnMono) 
+			{
+				// heuristic: give each action at least 1.5GB of RAM (some clang instances will need more, actually)
+				long MinMemoryPerActionMB = 3 * 1024 / 2;
+				long PhysicalRAMAvailableMB = (new PerformanceCounter ("Mono Memory", "Total Physical Memory").RawValue) / (1024 * 1024);
+				int MaxActionsAffordedByMemory = (int)(Math.Max(1, (PhysicalRAMAvailableMB) / MinMemoryPerActionMB));
+
+				MaxActionsToExecuteInParallel = Math.Min(MaxActionsToExecuteInParallel, MaxActionsAffordedByMemory);
+			}
 
             Log.TraceInformation("Performing {0} actions (max {1} parallel jobs)", Actions.Count, MaxActionsToExecuteInParallel);
 
