Invalidate Current SDK install if any failures occur

diff --git a/Engine/Source/Programs/UnrealBuildTool/Configuration/UEBuildPlatform.cs b/Engine/Source/Programs/UnrealBuildTool/Configuration/UEBuildPlatform.cs
--- a/Engine/Source/Programs/UnrealBuildTool/Configuration/UEBuildPlatform.cs
+++ b/Engine/Source/Programs/UnrealBuildTool/Configuration/UEBuildPlatform.cs
@@ -560,6 +560,12 @@ namespace UnrealBuildTool
             return false;
         }
 
+        public void InvalidateCurrentlyInstalled(String PlatformSDKRoot)
+        {
+            SetCurrentlyInstalledSDKString(PlatformSDKRoot, "");
+            SetLastRunScriptVersion(PlatformSDKRoot, "");
+        }
+
         /**
          *	Whether the required external SDKs are installed for this platform
          */
@@ -590,11 +596,11 @@ namespace UnrealBuildTool
                             if (!RunSDKHooks(PlatformSDKRoot, CurrentSDKString, SDKHookType.Uninstall))
                             {
                                 Console.WriteLine("Failed to uninstall currently installed SDK {0}", CurrentSDKString);
+                                InvalidateCurrentlyInstalled(PlatformSDKRoot);
                                 return SDKStatus.Invalid;
                             }
                             // delete Manifest file to avoid multiple uninstalls
-                            SetCurrentlyInstalledSDKString(PlatformSDKRoot, "");
-                            SetLastRunScriptVersion(PlatformSDKRoot, "");
+                            InvalidateCurrentlyInstalled(PlatformSDKRoot);
 
                             if (!RunSDKHooks(PlatformSDKRoot, GetRequiredSDKString(), SDKHookType.Install, false))
                             {								
@@ -611,6 +617,7 @@ namespace UnrealBuildTool
                         if (!SetupEnvironmentFromSDK(PlatformSDKRoot, GetRequiredSDKString()))
                         {
                             Console.WriteLine("Failed to load environment from required SDK {0}", GetRequiredSDKString());
+                            InvalidateCurrentlyInstalled(PlatformSDKRoot);
                             return SDKStatus.Invalid;
                         }                        
                         return SDKStatus.Valid;
