#ttp 333737: UHT should not output anything differently for Rocket - resubmitted after CL# 2068336 was backed out due to UAT rebuild problems.

#change Root module path added to include paths.
#change Plugins now write intermediates to the project folder.
#change All generated #includes are now relative.
#change Updated UHTDebugging.manifest to new format.
#remove Rocket-specific UHT execution removed.
#remove Some bHackHeaderGenerator usage removed.
#remove iOS header hack removed.

diff --git a/Engine/Source/Programs/UnrealBuildTool/Configuration/UEBuildModule.cs b/Engine/Source/Programs/UnrealBuildTool/Configuration/UEBuildModule.cs
--- a/Engine/Source/Programs/UnrealBuildTool/Configuration/UEBuildModule.cs
+++ b/Engine/Source/Programs/UnrealBuildTool/Configuration/UEBuildModule.cs
@@ -476,6 +476,9 @@ namespace UnrealBuildTool
 					var IncludePathModule = Target.GetModuleByName(IncludePathModuleName);
 					IncludePathModule.SetupPublicCompileEnvironment( SourceBinary, bInnerIncludePathsOnly, ref IncludePaths, ref SystemIncludePaths, ref Definitions, ref VisitedModules );
 				}
+
+				// Add the module's directory to the include path, so we can root #includes to it
+				IncludePaths.Add(ModuleDirectory);
 			}
 		}
 		
@@ -1708,7 +1711,7 @@ namespace UnrealBuildTool
 
 								// If it's an engine module, output intermediates to the engine intermediates directory. 
 								string IntermediateDirectory = Binary.Config.IntermediateDirectory;
-								if (PluginInfo == null && IntermediateDirectory != Target.EngineIntermediateDirectory && Path.GetFullPath(DependencyModule.ModuleDirectory).StartsWith(Path.GetFullPath(BuildConfiguration.RelativeEnginePath)))
+								if (IntermediateDirectory != Target.EngineIntermediateDirectory && Path.GetFullPath(DependencyModule.ModuleDirectory).StartsWith(Path.GetFullPath(BuildConfiguration.RelativeEnginePath)))
 								{
 									IntermediateDirectory = Target.EngineIntermediateDirectory;
 								}
@@ -1788,12 +1791,12 @@ namespace UnrealBuildTool
 					// Intermediates under Engine intermediate folder (program name will be appended later)
 					RootDirectory = Path.GetFullPath(BuildConfiguration.RelativeEnginePath);
 				}
-				BaseDirectory = Path.Combine(RootDirectory, BuildConfiguration.PlatformIntermediateFolder, Target.GetTargetName());
+				BaseDirectory = Path.Combine(RootDirectory, BuildConfiguration.PlatformIntermediateFolder, Target.GetTargetName(), "Inc");
 			}
 			else if (Plugins.IsPluginModule(ModuleName))
 			{
 				// Plugin module
-				BaseDirectory = Path.Combine(Plugins.GetPluginInfoForModule(ModuleName).Directory, BuildConfiguration.PlatformIntermediateFolder);
+				BaseDirectory = Plugins.GetPluginInfoForModule(ModuleName).IntermediateIncPath;
 			}
 			else
 			{
@@ -1805,11 +1808,11 @@ namespace UnrealBuildTool
 					BaseDirectory = ProjectFileGenerator.EngineRelativePath;
 				}
 
-				BaseDirectory = Path.GetFullPath(Path.Combine(BaseDirectory, BuildConfiguration.PlatformIntermediateFolder));
+				BaseDirectory = Path.GetFullPath(Path.Combine(BaseDirectory, BuildConfiguration.PlatformIntermediateFolder, "Inc"));
 			}
 
 			// Construct the intermediate path.
-			var GeneratedCodeDirectory = Path.Combine(BaseDirectory, "Inc", ModuleName);
+			var GeneratedCodeDirectory = Path.Combine(BaseDirectory, ModuleName);
 			return GeneratedCodeDirectory + Path.DirectorySeparatorChar;
 		}
 
diff --git a/Engine/Source/Programs/UnrealBuildTool/Configuration/UEBuildTarget.cs b/Engine/Source/Programs/UnrealBuildTool/Configuration/UEBuildTarget.cs
--- a/Engine/Source/Programs/UnrealBuildTool/Configuration/UEBuildTarget.cs
+++ b/Engine/Source/Programs/UnrealBuildTool/Configuration/UEBuildTarget.cs
@@ -1075,13 +1075,13 @@ namespace UnrealBuildTool
 					Manifest.AddBinaryNames(Path.Combine(Binary.Config.IntermediateDirectory, Path.GetFileNameWithoutExtension(Binary.Config.OutputFilePath) + ".lib"), "");
 				}
 			}
-            {
-                string DebugInfoExtension = BuildPlatform.GetDebugInfoExtension(UEBuildBinaryType.StaticLibrary);
-                foreach (var RedistLib in SpecialRocketLibFilesThatAreBuildProducts)
-                {
-                    Manifest.AddBinaryNames(RedistLib, DebugInfoExtension);
-                }
-            }
+			{
+				string DebugInfoExtension = BuildPlatform.GetDebugInfoExtension(UEBuildBinaryType.StaticLibrary);
+				foreach (var RedistLib in SpecialRocketLibFilesThatAreBuildProducts)
+				{
+					Manifest.AddBinaryNames(RedistLib, DebugInfoExtension);
+				}
+			}
 
 
 			if (UEBuildConfiguration.bCleanProject)
@@ -1500,18 +1500,7 @@ namespace UnrealBuildTool
 				string HeaderFilename                    = LinkerFixupsName + "Name.h";
 				string LinkerFixupHeaderFilenameWithPath = Path.Combine(GlobalCompileEnvironment.Config.OutputDirectory, HeaderFilename);
 
-				UnrealTargetPlatform TargetPlatform = UEBuildTarget.CPPTargetPlatformToUnrealTargetPlatform( GlobalCompileEnvironment.Config.TargetPlatform );
-
-				// Quick fix. 
-				// For some reason, when compiled over RPCUtility, the path can't be found, and we don't need the long path on Mac anyhow
-				if ( ExternalExecution.GetRuntimePlatform() != UnrealTargetPlatform.Mac && TargetPlatform == UnrealTargetPlatform.IOS )
-				{
-					LinkerFixupsFileContents.Add("#include \"" + HeaderFilename + "\"");
-				}
-				else
-				{
-					LinkerFixupsFileContents.Add("#include \"" + LinkerFixupHeaderFilenameWithPath + "\"");
-				}
+				LinkerFixupsFileContents.Add("#include \"" + HeaderFilename + "\"");
 
 				// Add a function that is not referenced by anything that invokes all the empty functions in the different static libraries
 				LinkerFixupsFileContents.Add("void " + LinkerFixupsName + "()");
@@ -1720,15 +1709,14 @@ namespace UnrealBuildTool
 			// Is this a Rocket module?
 			bool bIsRocketModule = RulesCompiler.IsRocketProjectModule(Module.Name);
 
-			// Get the binary type to build
-			UEBuildBinaryType BinaryType = ShouldCompileMonolithic()? UEBuildBinaryType.StaticLibrary : UEBuildBinaryType.DynamicLinkLibrary;
+			bool bCompileMonolithic = ShouldCompileMonolithic();
 
-			// Make the plugin intermediate path
-			string IntermediateDirectory = Path.Combine(Plugin.Directory, BuildConfiguration.PlatformIntermediateFolder, Plugins.GetPluginSubfolderName(BinaryType, AppName), Configuration.ToString());
+			// Get the binary type to build
+			UEBuildBinaryType BinaryType = bCompileMonolithic ? UEBuildBinaryType.StaticLibrary : UEBuildBinaryType.DynamicLinkLibrary;
 
 			// Get the output path. Don't prefix the app name for Rocket
 			string OutputFilePath;
-			if ((UnrealBuildTool.BuildingRocket() || UnrealBuildTool.RunningRocket()) && BinaryType == UEBuildBinaryType.StaticLibrary)
+			if ((UnrealBuildTool.BuildingRocket() || UnrealBuildTool.RunningRocket()) && bCompileMonolithic)
 			{
 				OutputFilePath = MakeBinaryPath(Module.Name, Module.Name, BinaryType, Rules.Type, bIsRocketModule, Plugin, AppName);
                 if (UnrealBuildTool.BuildingRocket())
@@ -1744,18 +1732,19 @@ namespace UnrealBuildTool
 			// Try to determine if we have the rules file
 			var ModuleFilename = RulesCompiler.GetModuleFilename(Module.Name);
 			var bHasModuleRules = String.IsNullOrEmpty(ModuleFilename) == false;
+
 			// Figure out whether we should build it from source
 			var ModuleSourceFolder = bHasModuleRules ? Path.GetDirectoryName(RulesCompiler.GetModuleFilename(Module.Name)) : ModuleFilename;
 			bool bShouldBeBuiltFromSource = bHasModuleRules && Directory.GetFiles(ModuleSourceFolder, "*.cpp", SearchOption.AllDirectories).Length > 0;
 
 			// Create the binary
-			UEBuildBinaryConfiguration Config = new UEBuildBinaryConfiguration( InType: BinaryType, 
-																				InOutputFilePath: OutputFilePath,
-																				InIntermediateDirectory: IntermediateDirectory,
-																				bInAllowExports: true,
-																				bInAllowCompilation: bShouldBeBuiltFromSource,
-																				bInHasModuleRules : bHasModuleRules,
-																				InModuleNames: new List<string> { Module.Name } );
+			UEBuildBinaryConfiguration Config = new UEBuildBinaryConfiguration( InType:                  BinaryType,
+																				InOutputFilePath:        OutputFilePath,
+																				InIntermediateDirectory: Plugin.IntermediateBuildPath,
+																				bInAllowExports:         true,
+																				bInAllowCompilation:     bShouldBeBuiltFromSource,
+																				bInHasModuleRules:       bHasModuleRules,
+																				InModuleNames:           new List<string> { Module.Name } );
 			AppBinaries.Add(new UEBuildBinaryCPP(this, Config));
             return SpecialRocketLibFilesThatAreBuildProducts;
 		}
@@ -1950,6 +1939,22 @@ namespace UnrealBuildTool
 				{
 					if(Plugin.LoadedFrom != PluginInfo.LoadedFromType.GameProject || Plugin.Directory.StartsWith(ProjectDirectory, StringComparison.InvariantCultureIgnoreCase))
 					{
+						if (Plugin.LoadedFrom == PluginInfo.LoadedFromType.Engine)
+						{
+							// Plugin folder is in the engine directory
+							Plugin.IntermediateIncPath   = Path.GetFullPath(Path.Combine(BuildConfiguration.RelativeEnginePath, BuildConfiguration.PlatformIntermediateFolder));
+							Plugin.IntermediateBuildPath = Path.GetFullPath(Path.Combine(BuildConfiguration.RelativeEnginePath, BuildConfiguration.PlatformIntermediateFolder, AppName, Configuration.ToString()));
+						}
+						else
+						{
+							// Plugin folder is in the project directory
+							Plugin.IntermediateIncPath   = Path.GetFullPath(Path.Combine(ProjectDirectory, BuildConfiguration.PlatformIntermediateFolder));
+							Plugin.IntermediateBuildPath = Path.GetFullPath(Path.Combine(ProjectDirectory, BuildConfiguration.PlatformIntermediateFolder, GetTargetName(), Configuration.ToString()));
+						}
+
+						Plugin.IntermediateIncPath   = Path.Combine(Plugin.IntermediateIncPath,   "Inc", "Plugins");
+						Plugin.IntermediateBuildPath = Path.Combine(Plugin.IntermediateBuildPath, "Plugins", ShouldCompileMonolithic() ? "Static" : "Dynamic");
+
 						ValidPlugins.Add(Plugin);
 					}
 				}
@@ -2114,7 +2119,7 @@ namespace UnrealBuildTool
             {
                 GlobalCompileEnvironment.Config.Definitions.Add("UE_ROCKET=0");
             }
-			
+
 			// Rocket intermediates go to the project's intermediate folder.  Rocket never writes to the engine intermediate folder. (Those files are immutable)
 			// Also, when compiling in monolithic, all intermediates go to the project's folder.  This is because a project can change definitions that affects all engine translation
 			// units too, so they can't be shared between different targets.  They are effectively project-specific engine intermediates.
@@ -2124,11 +2129,11 @@ namespace UnrealBuildTool
 				var IntermediateConfiguration = Configuration;
 				if( UnrealBuildTool.RunningRocket() )
 				{
-					// Only Development and Shipping are supported for engine modules
+				// Only Development and Shipping are supported for engine modules
 					if( IntermediateConfiguration != UnrealTargetConfiguration.Development && IntermediateConfiguration != UnrealTargetConfiguration.Shipping )
-					{
-						IntermediateConfiguration = UnrealTargetConfiguration.Development;
-					}
+				{
+					IntermediateConfiguration = UnrealTargetConfiguration.Development;
+				}
 				}
 
 				GlobalCompileEnvironment.Config.OutputDirectory = Path.Combine(BuildConfiguration.PlatformIntermediatePath, OutputAppName, IntermediateConfiguration.ToString());
diff --git a/Engine/Source/Programs/UnrealBuildTool/System/ExternalExecution.cs b/Engine/Source/Programs/UnrealBuildTool/System/ExternalExecution.cs
--- a/Engine/Source/Programs/UnrealBuildTool/System/ExternalExecution.cs
+++ b/Engine/Source/Programs/UnrealBuildTool/System/ExternalExecution.cs
@@ -58,6 +58,7 @@ namespace UnrealBuildTool
 		{
 			public string       Name;
 			public string       BaseDirectory;
+			public string       IncludeBase;     // The include path which all UHT-generated includes should be relative to
 			public string       OutputDirectory;
 			public List<string> ClassesHeaders;
 			public List<string> PublicHeaders;
@@ -72,17 +73,17 @@ namespace UnrealBuildTool
 			}
 		}
 
-		public UHTManifest(bool InUseRelativePaths, UEBuildTarget Target, string InRootLocalPath, string InRootBuildPath, IEnumerable<UHTModuleInfo> ModuleInfo)
+		public UHTManifest(UEBuildTarget Target, string InRootLocalPath, string InRootBuildPath, IEnumerable<UHTModuleInfo> ModuleInfo)
 		{
-			UseRelativePaths = InUseRelativePaths;
-			IsGameTarget     = TargetRules.IsGameType(Target.Rules.Type);
-			RootLocalPath    = InRootLocalPath;
-			RootBuildPath    = InRootBuildPath;
-			TargetName       = Target.GetTargetName();
+			IsGameTarget  = TargetRules.IsGameType(Target.Rules.Type);
+			RootLocalPath = InRootLocalPath;
+			RootBuildPath = InRootBuildPath;
+			TargetName    = Target.GetTargetName();
 
 			Modules = ModuleInfo.Select(Info => new Module{
 				Name                     = Info.ModuleName,
 				BaseDirectory            = Info.ModuleDirectory,
+				IncludeBase              = Info.ModuleDirectory,
 				OutputDirectory          = UEBuildModuleCPP.GetGeneratedCodeDirectoryForModule(Target, Info.ModuleDirectory, Info.ModuleName),
 				ClassesHeaders           = Info.PublicUObjectClassesHeaders.Select((Header) => Header.AbsolutePath).ToList(),
 				PublicHeaders            = Info.PublicUObjectHeaders       .Select((Header) => Header.AbsolutePath).ToList(),
@@ -95,11 +96,10 @@ namespace UnrealBuildTool
 			}).ToList();
 		}
 
-		public bool         UseRelativePaths; // Generate relative paths or absolute paths
 		public bool         IsGameTarget;     // True if the current target is a game target
 		public string       RootLocalPath;    // The engine path on the local machine
 		public string       RootBuildPath;    // The engine path on the build machine, if different (e.g. Mac/iOS builds)
-		public string		TargetName;       // Name of the target currently being compiled		
+		public string       TargetName;       // Name of the target currently being compiled
 		public List<Module> Modules;
 	}
 
@@ -385,7 +385,7 @@ namespace UnrealBuildTool
 			var CppPlatform    = BuildPlatform.GetCPPTargetPlatform(Target.Platform);
 			var ToolChain      = UEToolChain.GetPlatformToolChain(CppPlatform);
 			var RootLocalPath  = Path.GetFullPath(ProjectFileGenerator.RootRelativePath);
-			var Manifest       = new UHTManifest(UnrealBuildTool.BuildingRocket() || UnrealBuildTool.RunningRocket(), Target, RootLocalPath, ToolChain.ConvertPath(RootLocalPath + '\\'), UObjectModules);
+			var Manifest       = new UHTManifest(Target, RootLocalPath, ToolChain.ConvertPath(RootLocalPath + '\\'), UObjectModules);
 
 			// ensure the headers are up to date
 			if (!bIsBuildingUHT && 
diff --git a/Engine/Source/Programs/UnrealBuildTool/System/Plugins.cs b/Engine/Source/Programs/UnrealBuildTool/System/Plugins.cs
--- a/Engine/Source/Programs/UnrealBuildTool/System/Plugins.cs
+++ b/Engine/Source/Programs/UnrealBuildTool/System/Plugins.cs
@@ -48,6 +48,12 @@ namespace UnrealBuildTool
 		// Path to the plugin's root directory
 		public string Directory;
 
+		// Path to the plugin's intermediate build folder
+		public string IntermediateBuildPath;
+
+		// Path to the plugin's Inc folder
+		public string IntermediateIncPath;
+
 		// List of modules in this plugin
 		public readonly List<PluginModuleInfo> Modules = new List<PluginModuleInfo>();
 
@@ -425,17 +431,6 @@ namespace UnrealBuildTool
 			}
 		}
 
-		public static string GetPluginSubfolderName(UEBuildBinaryType InBinaryType, string InTargetName)
-		{
-			string SubFolderName = Path.Combine("Dynamic", InTargetName);
-			if (InBinaryType != UEBuildBinaryType.DynamicLinkLibrary)
-			{
-				SubFolderName = Path.Combine("Static", InTargetName);
-			}
-
-			return SubFolderName;
-		}
-
 		/// Access a mapping of modules to their respective owning plugin.  Dictionary is case-insensitive.
 		private static Dictionary< string, PluginInfo > ModuleNameToPluginMap
 		{
diff --git a/Engine/Source/Programs/UnrealHeaderTool/Private/CodeGenerator.cpp b/Engine/Source/Programs/UnrealHeaderTool/Private/CodeGenerator.cpp
--- a/Engine/Source/Programs/UnrealHeaderTool/Private/CodeGenerator.cpp
+++ b/Engine/Source/Programs/UnrealHeaderTool/Private/CodeGenerator.cpp
@@ -163,6 +163,11 @@ static FString Tabify(const TCHAR *Input)
 	return Result;
 }
 
+void ConvertToBuildIncludePath(UPackage* Package, FString& LocalPath)
+{
+	FPaths::MakePathRelativeTo(LocalPath, *GPackageToManifestModuleMap.FindChecked(Package)->IncludeBase);
+}
+
 /**
  *	Helper function for finding the location of a package
  *	This is required as source now lives in several possible directories
@@ -2013,17 +2018,7 @@ void FNativeClassHeaderGenerator::ExportClassHeaderWrapper( UClass* Class, bool
 
 		check(SourceFilename.EndsWith(TEXT(".h")));
 
-		if(bUseRelativePaths)
-		{
-			// When building Rocket, we need relative paths in the generated headers so that they can be installed to any directory.
-			FPaths::MakePathRelativeTo(NewFileName, *ClassHeaderPath);
-			NewFileName.ReplaceInline(TEXT("\\"), TEXT("/"));
-		}
-		// ...we cannot rebase paths relative to the install directory. It may be on a different drive.
-		else if (FPaths::MakePathRelativeTo(NewFileName, *GManifest.RootLocalPath))
-		{
-			NewFileName = GManifest.RootBuildPath / NewFileName;
-		}
+		ConvertToBuildIncludePath(Package, NewFileName);
 
 		auto IncludeStr = FString::Printf(
 			TEXT("#ifndef %s_%s_generated_h")													LINE_TERMINATOR
@@ -3819,14 +3814,13 @@ bool IsExportOrTemporaryClass(FClass* Class)
 }
 
 // Constructor.
-FNativeClassHeaderGenerator::FNativeClassHeaderGenerator( UPackage* InPackage, FClasses& AllClasses, bool InAllowSaveExportedHeaders, bool bInUseRelativePaths )
+FNativeClassHeaderGenerator::FNativeClassHeaderGenerator( UPackage* InPackage, FClasses& AllClasses, bool InAllowSaveExportedHeaders )
 	: CurrentClass                          (NULL)
 	, API                                   (FPackageName::GetShortName(InPackage).ToUpper())
 	, Package                               (InPackage)
 	, bIsExportingForOffsetDeterminationOnly(false)
 	, bAllowSaveExportedHeaders             (InAllowSaveExportedHeaders)
 	, bFailIfGeneratedCodeChanges           (false)
-	, bUseRelativePaths                     (bInUseRelativePaths)
 {
 	const FString PackageName = FPackageName::GetShortName(Package);
 
@@ -4296,18 +4290,7 @@ void FNativeClassHeaderGenerator::ExportGeneratedCPP()
 	if (ModuleInfo->PCH.Len())
 	{
 		FString PCH = ModuleInfo->PCH;
-		if(bUseRelativePaths)
-		{
-			// When building Rocket, we need relative paths in the generated headers so that they can be installed to any directory.
-			FPaths::MakePathRelativeTo(PCH, *ModuleInfo->GeneratedCPPFilenameBase);
-			PCH.ReplaceInline(TEXT("\\"), TEXT("/"));
-		}
-		// ...we cannot rebase paths relative to the install directory. It may be on a different drive.
-		else if (FPaths::MakePathRelativeTo(PCH, *GManifest.RootLocalPath))
-		{
-			PCH = GManifest.RootBuildPath / PCH;
-		}
-
+		ConvertToBuildIncludePath(Package, PCH);
 		ModulePCHInclude = FString::Printf(TEXT("#include \"%s\"") LINE_TERMINATOR, *PCH);
 	}
 
@@ -4656,7 +4639,7 @@ ECompilationResult::Type UnrealHeaderTool_Main(const FString& ModuleInfoFilename
 			{
 				if (UPackage* Package = Cast<UPackage>(StaticFindObjectFast(UPackage::StaticClass(), NULL, FName(*Module.LongPackageName), false, false)))
 				{
-					Result = FHeaderParser::ParseAllHeadersInside(GWarn, Package, Module.SaveExportedHeaders, GManifest.UseRelativePaths, ScriptPlugins);
+					Result = FHeaderParser::ParseAllHeadersInside(GWarn, Package, Module.SaveExportedHeaders, ScriptPlugins);
 					if (Result != ECompilationResult::Succeeded)
 					{
 						++NumFailures;
diff --git a/Engine/Source/Programs/UnrealHeaderTool/Private/HeaderParser.cpp b/Engine/Source/Programs/UnrealHeaderTool/Private/HeaderParser.cpp
--- a/Engine/Source/Programs/UnrealHeaderTool/Private/HeaderParser.cpp
+++ b/Engine/Source/Programs/UnrealHeaderTool/Private/HeaderParser.cpp
@@ -6448,7 +6448,7 @@ bool FHeaderParser::DependentClassNameFromHeader(const TCHAR* HeaderFilename, FS
 }
 
 // Begins the process of exporting C++ class declarations for native classes in the specified package
-void FHeaderParser::ExportNativeHeaders( UPackage* CurrentPackage, FClasses& AllClasses, bool bAllowSaveExportedHeaders, bool bUseRelativePaths )
+void FHeaderParser::ExportNativeHeaders( UPackage* CurrentPackage, FClasses& AllClasses, bool bAllowSaveExportedHeaders )
 {
 	// Build a list of header filenames
 	TArray<FString>	ClassHeaderFilenames;
@@ -6493,7 +6493,7 @@ void FHeaderParser::ExportNativeHeaders( UPackage* CurrentPackage, FClasses& All
 		}
 
 		// Export native class definitions to package header files.
-		FNativeClassHeaderGenerator(CurrentPackage, AllClasses, bAllowSaveExportedHeaders, bUseRelativePaths);
+		FNativeClassHeaderGenerator(CurrentPackage, AllClasses, bAllowSaveExportedHeaders);
 	}
 }
 
@@ -6609,7 +6609,7 @@ void ExportClassTreeToScriptPlugins(const FClassTree* Node, TArray<IScriptGenera
 }
 
 // Parse all headers for classes that are inside CurrentPackage.
-ECompilationResult::Type FHeaderParser::ParseAllHeadersInside(FFeedbackContext* Warn, UPackage* CurrentPackage, bool bAllowSaveExportedHeaders, bool bUseRelativePaths, TArray<IScriptGeneratorPluginInterface*>& ScriptPlugins)
+ECompilationResult::Type FHeaderParser::ParseAllHeadersInside(FFeedbackContext* Warn, UPackage* CurrentPackage, bool bAllowSaveExportedHeaders, TArray<IScriptGeneratorPluginInterface*>& ScriptPlugins)
 {
 	// Disable loading of objects outside of this package (or more exactly, objects which aren't UFields, CDO, or templates)
 	TGuardValue<bool> AutoRestoreVerifyObjectRefsFlag(GVerifyObjectReferencesOnly, true);
@@ -6646,7 +6646,7 @@ ECompilationResult::Type FHeaderParser::ParseAllHeadersInside(FFeedbackContext*
 			// from the feedback context.
 			Warn->SetContext(NULL);
 
-			ExportNativeHeaders(CurrentPackage, AllClasses, bAllowSaveExportedHeaders, bUseRelativePaths);
+			ExportNativeHeaders(CurrentPackage, AllClasses, bAllowSaveExportedHeaders);
 
 			// Done with header generation
 			if (HeaderParser.LinesParsed > 0)
diff --git a/Engine/Source/Programs/UnrealHeaderTool/Private/HeaderParser.h b/Engine/Source/Programs/UnrealHeaderTool/Private/HeaderParser.h
--- a/Engine/Source/Programs/UnrealHeaderTool/Private/HeaderParser.h
+++ b/Engine/Source/Programs/UnrealHeaderTool/Private/HeaderParser.h
@@ -80,7 +80,7 @@ class FHeaderParser : public FBaseParser, public FContextSupplier
 {
 public:
 	// Parse all headers for classes that are inside LimitOuter.
-	static ECompilationResult::Type ParseAllHeadersInside(FFeedbackContext* Warn, UPackage* LimitOuter, bool bAllowSaveExportedHeaders, bool bUseRelativePaths, TArray<class IScriptGeneratorPluginInterface*>& ScriptPlugins);
+	static ECompilationResult::Type ParseAllHeadersInside(FFeedbackContext* Warn, UPackage* LimitOuter, bool bAllowSaveExportedHeaders, TArray<class IScriptGeneratorPluginInterface*>& ScriptPlugins);
 
 	// Performs a preliminary parse of the text in the specified buffer, pulling out:
 	//   Class name and parent class name
@@ -299,7 +299,7 @@ protected:
 	 * @param	CurrentPackage	the package being compiled
 	 * @param	AllClasses		the class tree for CurrentPackage
 	 */
-	static void ExportNativeHeaders( UPackage* CurrentPackage, FClasses& AllClasses, bool bAllowSaveExportedHeaders, bool bUseRelativePaths );
+	static void ExportNativeHeaders( UPackage* CurrentPackage, FClasses& AllClasses, bool bAllowSaveExportedHeaders );
 
 	// FContextSupplier interface.
 	virtual FString GetContext() OVERRIDE;
diff --git a/Engine/Source/Programs/UnrealHeaderTool/Private/Manifest.cpp b/Engine/Source/Programs/UnrealHeaderTool/Private/Manifest.cpp
--- a/Engine/Source/Programs/UnrealHeaderTool/Private/Manifest.cpp
+++ b/Engine/Source/Programs/UnrealHeaderTool/Private/Manifest.cpp
@@ -61,20 +61,18 @@ FManifest FManifest::LoadFromFile(const FString& Filename)
 
 	TArray<TSharedPtr<FJsonValue>> ModulesArray;
 
-	GetJsonFieldValue(Result.UseRelativePaths, RootObject, TEXT("UseRelativePaths"), TEXT("{manifest root}"));
-	GetJsonFieldValue(Result.IsGameTarget,     RootObject, TEXT("IsGameTarget"),     TEXT("{manifest root}"));
-	GetJsonFieldValue(Result.RootLocalPath,    RootObject, TEXT("RootLocalPath"),    TEXT("{manifest root}"));
-	GetJsonFieldValue(Result.RootBuildPath,    RootObject, TEXT("RootBuildPath"),    TEXT("{manifest root}"));
-	GetJsonFieldValue(Result.TargetName,       RootObject, TEXT("TargetName"),       TEXT("{manifest root}"));
-	GetJsonFieldValue(ModulesArray,            RootObject, TEXT("Modules"),          TEXT("{manifest root}"));
-
-	UE_LOG(LogCompile, Log, TEXT("Loaded manifest: %s"),          *Filename);
-	UE_LOG(LogCompile, Log, TEXT("Manifest.UseRelativePaths=%s"), Result.UseRelativePaths ? TEXT("True") : TEXT("False"));
-	UE_LOG(LogCompile, Log, TEXT("Manifest.IsGameTarget=%s"),     Result.IsGameTarget ? TEXT("True") : TEXT("False"));
-	UE_LOG(LogCompile, Log, TEXT("Manifest.RootLocalPath=%s"),    *Result.RootLocalPath);
-	UE_LOG(LogCompile, Log, TEXT("Manifest.RootBuildPath=%s"),    *Result.RootBuildPath);
-	UE_LOG(LogCompile, Log, TEXT("Manifest.TargetName=%s"),       *Result.TargetName);
-	UE_LOG(LogCompile, Log, TEXT("Manifest.Modules=%d"),          ModulesArray.Num());
+	GetJsonFieldValue(Result.IsGameTarget,  RootObject, TEXT("IsGameTarget"),  TEXT("{manifest root}"));
+	GetJsonFieldValue(Result.RootLocalPath, RootObject, TEXT("RootLocalPath"), TEXT("{manifest root}"));
+	GetJsonFieldValue(Result.RootBuildPath, RootObject, TEXT("RootBuildPath"), TEXT("{manifest root}"));
+	GetJsonFieldValue(Result.TargetName,    RootObject, TEXT("TargetName"),    TEXT("{manifest root}"));
+	GetJsonFieldValue(ModulesArray,         RootObject, TEXT("Modules"),       TEXT("{manifest root}"));
+
+	UE_LOG(LogCompile, Log, TEXT("Loaded manifest: %s"),       *Filename);
+	UE_LOG(LogCompile, Log, TEXT("Manifest.IsGameTarget=%s"),  Result.IsGameTarget ? TEXT("True") : TEXT("False"));
+	UE_LOG(LogCompile, Log, TEXT("Manifest.RootLocalPath=%s"), *Result.RootLocalPath);
+	UE_LOG(LogCompile, Log, TEXT("Manifest.RootBuildPath=%s"), *Result.RootBuildPath);
+	UE_LOG(LogCompile, Log, TEXT("Manifest.TargetName=%s"),    *Result.TargetName);
+	UE_LOG(LogCompile, Log, TEXT("Manifest.Modules=%d"),       ModulesArray.Num());
 
 	Result.RootLocalPath = FPaths::ConvertRelativePathToFull(FilenamePath, Result.RootLocalPath);
 	Result.RootBuildPath = FPaths::ConvertRelativePathToFull(FilenamePath, Result.RootBuildPath);
@@ -99,6 +97,7 @@ FManifest FManifest::LoadFromFile(const FString& Filename)
 
 		GetJsonFieldValue(KnownModule.Name,                      ModuleObj, TEXT("Name"),                     *Outer);
 		GetJsonFieldValue(KnownModule.BaseDirectory,             ModuleObj, TEXT("BaseDirectory"),            *Outer);
+		GetJsonFieldValue(KnownModule.IncludeBase,               ModuleObj, TEXT("IncludeBase"),              *Outer);
 		GetJsonFieldValue(KnownModule.GeneratedIncludeDirectory, ModuleObj, TEXT("OutputDirectory"),          *Outer);
 		GetJsonFieldValue(KnownModule.SaveExportedHeaders,       ModuleObj, TEXT("SaveExportedHeaders"),      *Outer);
 		GetJsonFieldValue(ClassesHeaders,                        ModuleObj, TEXT("ClassesHeaders"),           *Outer);
@@ -111,11 +110,13 @@ FManifest FManifest::LoadFromFile(const FString& Filename)
 
 		// Convert relative paths
 		KnownModule.BaseDirectory             = FPaths::ConvertRelativePathToFull(FilenamePath, KnownModule.BaseDirectory);
+		KnownModule.IncludeBase               = FPaths::ConvertRelativePathToFull(FilenamePath, KnownModule.IncludeBase);
 		KnownModule.GeneratedIncludeDirectory = FPaths::ConvertRelativePathToFull(FilenamePath, KnownModule.GeneratedIncludeDirectory);
 		KnownModule.GeneratedCPPFilenameBase  = FPaths::ConvertRelativePathToFull(FilenamePath, KnownModule.GeneratedCPPFilenameBase);
 
 		// Ensure directories end with a slash, because this aids their use with FPaths::MakePathRelativeTo.
 		if (!KnownModule.BaseDirectory            .EndsWith(TEXT("/"))) { KnownModule.BaseDirectory            .AppendChar(TEXT('/')); }
+		if (!KnownModule.IncludeBase              .EndsWith(TEXT("/"))) { KnownModule.IncludeBase              .AppendChar(TEXT('/')); }
 		if (!KnownModule.GeneratedIncludeDirectory.EndsWith(TEXT("/"))) { KnownModule.GeneratedIncludeDirectory.AppendChar(TEXT('/')); }
 
 		KnownModule.PublicUObjectClassesHeaders.AddZeroed(ClassesHeaders.Num());
diff --git a/Engine/Source/Programs/UnrealHeaderTool/Private/Manifest.h b/Engine/Source/Programs/UnrealHeaderTool/Private/Manifest.h
--- a/Engine/Source/Programs/UnrealHeaderTool/Private/Manifest.h
+++ b/Engine/Source/Programs/UnrealHeaderTool/Private/Manifest.h
@@ -18,6 +18,9 @@ struct FManifestModule
 	/** Base directory of this module on disk */
 	FString BaseDirectory;
 
+	/** The directory to which #includes from this module should be relative */
+	FString IncludeBase;
+
 	/** Directory where generated include files should go */
 	FString GeneratedIncludeDirectory;
 
@@ -42,8 +45,7 @@ struct FManifestModule
 
 struct FManifest
 {
-	bool    UseRelativePaths;
-	bool IsGameTarget;
+	bool    IsGameTarget;
 	FString RootLocalPath;
 	FString RootBuildPath;
 	FString TargetName;
diff --git a/Engine/Source/Programs/UnrealHeaderTool/Private/NativeClassExporter.h b/Engine/Source/Programs/UnrealHeaderTool/Private/NativeClassExporter.h
--- a/Engine/Source/Programs/UnrealHeaderTool/Private/NativeClassExporter.h
+++ b/Engine/Source/Programs/UnrealHeaderTool/Private/NativeClassExporter.h
@@ -98,9 +98,6 @@ private:
 	/** If true, any change in the generated headers will result in UHT failure. */
 	bool bFailIfGeneratedCodeChanges;
 
-	/** If true, writes absolute paths to the generated files. */
-	bool bUseRelativePaths;
-
 	// This exists because it makes debugging much easier on VC2010, since the visualizers can't properly understand templates with templated args
 	struct HeaderDependents : TArray<const FString*>
 	{
@@ -466,5 +463,5 @@ private:
 public:
 
 	// Constructor
-	FNativeClassHeaderGenerator( UPackage* InPackage, FClasses& AllClasses, bool InAllowSaveExportedHeaders, bool bInUseRelativePaths );
+	FNativeClassHeaderGenerator( UPackage* InPackage, FClasses& AllClasses, bool InAllowSaveExportedHeaders );
 };
diff --git a/Engine/Source/Programs/UnrealHeaderTool/Resources/UHTDebugging.manifest b/Engine/Source/Programs/UnrealHeaderTool/Resources/UHTDebugging.manifest
--- a/Engine/Source/Programs/UnrealHeaderTool/Resources/UHTDebugging.manifest
+++ b/Engine/Source/Programs/UnrealHeaderTool/Resources/UHTDebugging.manifest
@@ -1,5 +1,4 @@
 {
-	"UseRelativePaths": false,
 	"IsGameTarget": false,
 	"RootLocalPath": "..\\..\\..\\..",
 	"RootBuildPath": "..\\..\\..\\..",
@@ -7,6 +6,7 @@
 	"Modules":[{
 			"Name": "CoreUObject",
 			"BaseDirectory": "..\\..\\..\\Runtime\\CoreUObject",
+			"IncludeBase": "..\\..\\..\\Runtime\\CoreUObject",
 			"OutputDirectory": "..\\..\\..\\..\\Intermediate/Build/Inc\\CoreUObject\\",
 			"ClassesHeaders": [
 				"..\\..\\..\\Runtime\\CoreUObject\\Classes\\Object.h"
@@ -21,6 +21,7 @@
 		}, {
 			"Name": "UHTDebugging",
 			"BaseDirectory": ".\\UHTDebugging",
+			"IncludeBase": ".\\UHTDebugging",
 			"OutputDirectory": "..\\..\\..\\..\\Intermediate/Build/Inc\\UHTDebugging\\",
 			"ClassesHeaders": [
 				".\\UHTDebugging\\TestObject.h",
