Changed Xcode project generator so that files created by Xcode itself are stored in UE4's Intermediate dirs

diff --git a/Engine/Source/Programs/UnrealBuildTool/System/XcodeProjectFileGenerator.cs b/Engine/Source/Programs/UnrealBuildTool/System/XcodeProjectFileGenerator.cs
--- a/Engine/Source/Programs/UnrealBuildTool/System/XcodeProjectFileGenerator.cs
+++ b/Engine/Source/Programs/UnrealBuildTool/System/XcodeProjectFileGenerator.cs
@@ -424,6 +424,7 @@ namespace UnrealBuildTool
 
 		private void AppendProjectConfig(ref StringBuilder Contents, string ConfigName, string ConfigGuid, string PreprocessorDefinitions, string HeaderSearchPaths)
 		{
+			string EngineSubdir = (bGeneratingGameProjectFiles || bGeneratingRocketProjectFiles) ? "" : "Engine/";
 			Contents.Append(
 				"\t\t" + ConfigGuid + " /* " + ConfigName + " */ = {" + ProjectFileGenerator.NewLine +
 				"\t\t\tisa = XCBuildConfiguration;" + ProjectFileGenerator.NewLine +
@@ -437,6 +438,7 @@ namespace UnrealBuildTool
 				"\t\t\t\tSUPPORTED_PLATFORMS = \"macosx\";" + ProjectFileGenerator.NewLine +
 				"\t\t\t\tONLY_ACTIVE_ARCH = YES;" + ProjectFileGenerator.NewLine +
 				"\t\t\t\tSDKROOT = macosx;" + ProjectFileGenerator.NewLine +
+				"\t\t\t\tSYMROOT = " + EngineSubdir + "Intermediate/Build;" + ProjectFileGenerator.NewLine +
 				"\t\t\t};" + ProjectFileGenerator.NewLine +
 				"\t\t\tname = " + ConfigName + ";" + ProjectFileGenerator.NewLine +
 				"\t\t};" + ProjectFileGenerator.NewLine);
@@ -563,14 +565,14 @@ namespace UnrealBuildTool
 				"\t\t};" + ProjectFileGenerator.NewLine);
 		}
 
-		private void AppendIOSXCTestConfig(ref StringBuilder Contents, string ConfigName, string ConfigGuid, string TargetName, string EngineTarget, string EngineRelative)
+		private void AppendIOSXCTestConfig(ref StringBuilder Contents, string ConfigName, string ConfigGuid, string TargetName, string EngineSubdir, string EngineRelative)
 		{
 			Contents.Append(
 				"\t\t" + ConfigGuid + " /* " + ConfigName + " */ = {" + ProjectFileGenerator.NewLine +
 				"\t\t\tisa = XCBuildConfiguration;" + ProjectFileGenerator.NewLine +
 				"\t\t\tbuildSettings = {" + ProjectFileGenerator.NewLine +
 				"\t\t\t\tALWAYS_SEARCH_USER_PATHS = NO;" + ProjectFileGenerator.NewLine +
-				"\t\t\t\tBUNDLE_LOADER = \"" + EngineTarget + "Binaries/IOS/Payload/" + TargetName + ".app/" + TargetName + "\";" + ProjectFileGenerator.NewLine +
+				"\t\t\t\tBUNDLE_LOADER = \"" + EngineSubdir + "Binaries/IOS/Payload/" + TargetName + ".app/" + TargetName + "\";" + ProjectFileGenerator.NewLine +
 				"\t\t\t\tCLANG_CXX_LANGUAGE_STANDARD = \"gnu++0x\";" + ProjectFileGenerator.NewLine +
 				"\t\t\t\tCLANG_CXX_LIBRARY = \"libc++\";" + ProjectFileGenerator.NewLine +
 				"\t\t\t\tCLANG_ENABLE_MODULES = YES;" + ProjectFileGenerator.NewLine +
@@ -593,9 +595,9 @@ namespace UnrealBuildTool
 				"\t\t\t\t);" + ProjectFileGenerator.NewLine +
 				"\t\t\t\tGCC_SYMBOLS_PRIVATE_EXTERN = NO;" + ProjectFileGenerator.NewLine +
 				"\t\t\t\tINFOPLIST_FILE = \"" + EngineRelative + "Engine/Build/IOS/UE4CmdLineRun/UE4CmdLineRun-Info.plist\";" + ProjectFileGenerator.NewLine +
-				"\t\t\t\tSYMROOT = \"" + EngineTarget + "Binaries/IOS\";" + ProjectFileGenerator.NewLine +
-				"\t\t\t\tOBJROOT = \"" + EngineTarget + "Intermediate/IOS/build\";" + ProjectFileGenerator.NewLine +
-				"\t\t\t\tCONFIGURATION_BUILD_DIR = \"" + EngineTarget + "Binaries/IOS/Payload\";" + ProjectFileGenerator.NewLine +
+				"\t\t\t\tSYMROOT = \"" + EngineSubdir + "Binaries/IOS\";" + ProjectFileGenerator.NewLine +
+				"\t\t\t\tOBJROOT = \"" + EngineSubdir + "Intermediate/IOS/build\";" + ProjectFileGenerator.NewLine +
+				"\t\t\t\tCONFIGURATION_BUILD_DIR = \"" + EngineSubdir + "Binaries/IOS/Payload\";" + ProjectFileGenerator.NewLine +
 				"\t\t\t\tSUPPORTED_PLATFORMS = \"iphoneos iphonesimulator\";" + ProjectFileGenerator.NewLine +
 				"\t\t\t\tIPHONEOS_DEPLOYMENT_TARGET = 7.0;" + ProjectFileGenerator.NewLine +
 				"\t\t\t\tONLY_ACTIVE_ARCH = YES;" + ProjectFileGenerator.NewLine +
@@ -635,16 +637,8 @@ namespace UnrealBuildTool
 						}
 						else
 						{
-							string EngineTarget = "";
-							if (bIsUE4Game || bIsUE4Client)
-							{
-								if (!bGeneratingGameProjectFiles && !bGeneratingRocketProjectFiles)
-								{
-									EngineTarget = "Engine/";
-								}
-							}
-
-							AppendIOSXCTestConfig(ref Contents, ConfigName, ConfigGuid, Target.TargetName, EngineTarget, EngineRelative);
+							string EngineSubdir = (bGeneratingGameProjectFiles || bGeneratingRocketProjectFiles) ? "" : "Engine/";
+							AppendIOSXCTestConfig(ref Contents, ConfigName, ConfigGuid, Target.TargetName, EngineSubdir, EngineRelative);
 						}
 					}
 				}
