PS4 now supports VS 2013.  Don't fall back to 2012 for PS4 anymore.

diff --git a/Engine/Source/Programs/UnrealBuildTool/System/VCProjectFileGenerator.cs b/Engine/Source/Programs/UnrealBuildTool/System/VCProjectFileGenerator.cs
--- a/Engine/Source/Programs/UnrealBuildTool/System/VCProjectFileGenerator.cs
+++ b/Engine/Source/Programs/UnrealBuildTool/System/VCProjectFileGenerator.cs
@@ -193,15 +193,8 @@ namespace UnrealBuildTool
 							// ...but only if the user didn't override this via the command-line.
 							if( !UnrealBuildTool.CommandLineContains( "-2013" ) )
 							{ 
-								// Visual Studio 2013 is not supported by PS4 debugger add-in yet
-								if( SupportedPlatform == UnrealTargetPlatform.PS4 )		
-								{
-									Log.TraceInformation( "Forcing Visual Studio 2012 projects for PS4 compatibility (use '-2013' to override.)");
-									ProjectFileFormat = VCProjectFileFormat.VisualStudio2012;
-								}
-
 								// Visual Studio 2013 is not supported by Xbox One debugger add-in yet
-								else if( SupportedPlatform == UnrealTargetPlatform.XboxOne )		
+								if( SupportedPlatform == UnrealTargetPlatform.XboxOne )
 								{
 									Log.TraceInformation( "Forcing Visual Studio 2012 projects for Xbox One compatibility (use '-2013' to override.)");
 									ProjectFileFormat = VCProjectFileFormat.VisualStudio2012;
