UBT: Apply memory-based job heuristics to all Mono platforms

diff --git a/Engine/Source/Programs/UnrealBuildTool/System/LocalExecutor.cs b/Engine/Source/Programs/UnrealBuildTool/System/LocalExecutor.cs
--- a/Engine/Source/Programs/UnrealBuildTool/System/LocalExecutor.cs
+++ b/Engine/Source/Programs/UnrealBuildTool/System/LocalExecutor.cs
@@ -321,28 +321,25 @@ namespace UnrealBuildTool
 				MaxActionsToExecuteInParallel = NumCores;
 			}
 
-            if (ExternalExecution.GetRuntimePlatform() == UnrealTargetPlatform.Mac)
-			{
-				// @todo iosmerge: this should be looking at available memory as well since some of our
-				// Macs have a poor memory/CPU ratio
-				MaxActionsToExecuteInParallel /= 2;
-
-				string NumUBTBuildTasks = Environment.GetEnvironmentVariable("NumUBTBuildTasks");
-				Int32 MaxUBTBuildTasks = MaxActionsToExecuteInParallel;
-				if(Int32.TryParse(NumUBTBuildTasks, out MaxUBTBuildTasks))
-				{
-					MaxActionsToExecuteInParallel = MaxUBTBuildTasks;
-				}
-			}
-			else if (ExternalExecution.GetRuntimePlatform () == UnrealTargetPlatform.Linux && Utils.IsRunningOnMono) 
-			{
+            if (Utils.IsRunningOnMono)
+            {
 				// heuristic: give each action at least 1.5GB of RAM (some clang instances will need more, actually)
 				long MinMemoryPerActionMB = 3 * 1024 / 2;
 				long PhysicalRAMAvailableMB = (new PerformanceCounter ("Mono Memory", "Total Physical Memory").RawValue) / (1024 * 1024);
 				int MaxActionsAffordedByMemory = (int)(Math.Max(1, (PhysicalRAMAvailableMB) / MinMemoryPerActionMB));
 
 				MaxActionsToExecuteInParallel = Math.Min(MaxActionsToExecuteInParallel, MaxActionsAffordedByMemory);
-			}
+
+                if (ExternalExecution.GetRuntimePlatform() == UnrealTargetPlatform.Mac)
+			    {
+				    string NumUBTBuildTasks = Environment.GetEnvironmentVariable("NumUBTBuildTasks");
+				    Int32 MaxUBTBuildTasks = MaxActionsToExecuteInParallel;
+				    if(Int32.TryParse(NumUBTBuildTasks, out MaxUBTBuildTasks))
+				    {
+					    MaxActionsToExecuteInParallel = MaxUBTBuildTasks;
+				    }
+                }
+            }
 
             Log.TraceInformation("Performing {0} actions (max {1} parallel jobs)", Actions.Count, MaxActionsToExecuteInParallel);
 
