build script, remove IOS on PC path from rocket

diff --git a/Engine/Source/Programs/UnrealBuildTool/IOS/UEBuildIOS.cs b/Engine/Source/Programs/UnrealBuildTool/IOS/UEBuildIOS.cs
--- a/Engine/Source/Programs/UnrealBuildTool/IOS/UEBuildIOS.cs
+++ b/Engine/Source/Programs/UnrealBuildTool/IOS/UEBuildIOS.cs
@@ -30,6 +30,10 @@ namespace UnrealBuildTool
 		 */
 		public override void RegisterBuildPlatform()
 		{
+            if ((UnrealBuildTool.RunningRocket() || UnrealBuildTool.BuildingRocket()) && ExternalExecution.GetRuntimePlatform() != UnrealTargetPlatform.Mac)
+            {
+                return;
+            }
 			// Register this build platform for IOS
 			Log.TraceVerbose("        Registering for {0}", UnrealTargetPlatform.IOS.ToString());
 			UEBuildPlatform.RegisterBuildPlatform(UnrealTargetPlatform.IOS, this);
diff --git a/Engine/Source/UE4Editor.Target.cs b/Engine/Source/UE4Editor.Target.cs
--- a/Engine/Source/UE4Editor.Target.cs
+++ b/Engine/Source/UE4Editor.Target.cs
@@ -101,7 +101,7 @@ public class UE4EditorTarget : TargetRules
         }
         else
         {
-            MobilePlats = new List<UnrealTargetPlatform> { HostPlatform, UnrealTargetPlatform.Android, UnrealTargetPlatform.IOS };
+            MobilePlats = new List<UnrealTargetPlatform> { HostPlatform, UnrealTargetPlatform.Android };
         }
 
         NonCodeProjectNames.Add("BlackJack", MobilePlats);
