Updated Mac UE4EditorServices to use CleanGameProject(), bumped its version number and added missing uproject icon file to manifest

diff --git a/Engine/Source/Programs/Mac/UE4EditorServices/Private/Mac/UE4EditorServicesAppDelegate.cpp b/Engine/Source/Programs/Mac/UE4EditorServices/Private/Mac/UE4EditorServicesAppDelegate.cpp
--- a/Engine/Source/Programs/Mac/UE4EditorServices/Private/Mac/UE4EditorServicesAppDelegate.cpp
+++ b/Engine/Source/Programs/Mac/UE4EditorServices/Private/Mac/UE4EditorServicesAppDelegate.cpp
@@ -339,7 +339,19 @@
 				NSString* EnginePath = Iter.Value().GetNSString();
 				if ([EnginePath isEqualToString:SelectedEnginePath])
 				{
-					FDesktopPlatformModule::Get()->SetEngineIdentifierForProject([FileURL path], Iter.Key());
+					if (!FDesktopPlatformModule::Get()->SetEngineIdentifierForProject([FileURL path], Iter.Key()))
+					{
+						FPlatformMisc::MessageBoxExt(EAppMsgType::Ok, TEXT("Couldn't set association for project. Check the file is writeable."), TEXT("Error"));
+						return;
+					}
+
+					// Clean all the build products
+					FFeedbackContext *Warn = FDesktopPlatformModule::Get()->GetNativeFeedbackContext();
+					if (!FDesktopPlatformModule::Get()->CleanGameProject(FPaths::GetPath([FileURL path]), Warn))
+					{
+						FPlatformMisc::MessageBoxExt(EAppMsgType::Ok, TEXT("Couldn't clean project build products."), TEXT("Error"));
+						return;
+					}
 					break;
 				}
 			}
diff --git a/Engine/Source/Programs/Mac/UE4EditorServices/Resources/Mac/Info.plist b/Engine/Source/Programs/Mac/UE4EditorServices/Resources/Mac/Info.plist
--- a/Engine/Source/Programs/Mac/UE4EditorServices/Resources/Mac/Info.plist
+++ b/Engine/Source/Programs/Mac/UE4EditorServices/Resources/Mac/Info.plist
@@ -38,11 +38,11 @@
 	<key>CFBundlePackageType</key>
 	<string>APPL</string>
 	<key>CFBundleShortVersionString</key>
-	<string>1.3</string>
+	<string>1.4</string>
 	<key>CFBundleSignature</key>
 	<string>????</string>
 	<key>CFBundleVersion</key>
-	<string>1.3</string>
+	<string>1.4</string>
 	<key>LSApplicationCategoryType</key>
 	<string>public.app-category.developer-tools</string>
 	<key>LSBackgroundOnly</key>
diff --git a/Engine/Source/Programs/UnrealBuildTool/Mac/MacToolChain.cs b/Engine/Source/Programs/UnrealBuildTool/Mac/MacToolChain.cs
--- a/Engine/Source/Programs/UnrealBuildTool/Mac/MacToolChain.cs
+++ b/Engine/Source/Programs/UnrealBuildTool/Mac/MacToolChain.cs
@@ -958,7 +958,7 @@ namespace UnrealBuildTool
 					// copy over some needed files
 					// @todo mac: Make a QueueDirectoryForBatchUpload
 					QueueFileForBatchUpload(FileItem.GetItemByFullPath(Path.GetFullPath("../../Engine/Source/Runtime/Launch/Resources/Mac/UE4.icns")));
-					QueueFileForBatchUpload(FileItem.GetItemByFullPath(Path.GetFullPath("../../Engine/Source/Runtime/Launch/Resources/Mac/Rocket.icns")));
+					QueueFileForBatchUpload(FileItem.GetItemByFullPath(Path.GetFullPath("../../Engine/Source/Runtime/Launch/Resources/Mac/UProject.icns")));
 					QueueFileForBatchUpload(FileItem.GetItemByFullPath(Path.GetFullPath("../../Engine/Source/Runtime/Launch/Resources/Mac/Info.plist")));
 					QueueFileForBatchUpload(FileItem.GetItemByFullPath(Path.GetFullPath("../../Engine/Source/Runtime/Launch/Resources/Mac/English.lproj/InfoPlist.strings")));
 					QueueFileForBatchUpload(FileItem.GetItemByFullPath(Path.GetFullPath("../../Engine/Source/Runtime/Launch/Resources/Mac/English.lproj/MainMenu.xib")));
@@ -1188,6 +1188,11 @@ namespace UnrealBuildTool
 				Manifest.AddFileName(BundleContentsDirectory + "Resources/RadioEffectUnit.component/Contents/Resources/English.lproj/Localizable.strings");
 				Manifest.AddFileName(BundleContentsDirectory + "Resources/RadioEffectUnit.component/Contents/Resources/RadioEffectUnit.rsrc");
 				Manifest.AddFileName(BundleContentsDirectory + "Resources/RadioEffectUnit.component/Contents/Info.plist");
+
+				if (Binary.Config.TargetName.StartsWith("UE4Editor"))
+				{
+					Manifest.AddFileName(BundleContentsDirectory + "Resources/UProject.icns");
+				}
 			}
 		}
 
