Get UBT to write out progress information while compiling. Also fix missing output if telemetry fails to initialize.

diff --git a/Engine/Source/Programs/UnrealBuildTool/System/ExternalExecution.cs b/Engine/Source/Programs/UnrealBuildTool/System/ExternalExecution.cs
--- a/Engine/Source/Programs/UnrealBuildTool/System/ExternalExecution.cs
+++ b/Engine/Source/Programs/UnrealBuildTool/System/ExternalExecution.cs
@@ -387,100 +387,108 @@ namespace UnrealBuildTool
 			var RootLocalPath  = Path.GetFullPath(ProjectFileGenerator.RootRelativePath);
 			var Manifest       = new UHTManifest(UnrealBuildTool.BuildingRocket() || UnrealBuildTool.RunningRocket(), Target, RootLocalPath, ToolChain.ConvertPath(RootLocalPath + '\\'), UObjectModules);
 
-			// ensure the headers are up to date
-			if (!bIsBuildingUHT && 
-				(UEBuildConfiguration.bForceHeaderGeneration == true || AreGeneratedCodeFilesOutOfDate(Target, UObjectModules)))
+			using (ProgressWriter Progress = new ProgressWriter("Generating headers...", false))
 			{
- 				// Always build UnrealHeaderTool if header regeneration is required, unless we're running within a Rocket ecosystem
-				if (UnrealBuildTool.RunningRocket() == false && UEBuildConfiguration.bDoNotBuildUHT == false)
+				// ensure the headers are up to date
+				if (!bIsBuildingUHT && 
+					(UEBuildConfiguration.bForceHeaderGeneration == true || AreGeneratedCodeFilesOutOfDate(Target, UObjectModules)))
 				{
-					// If it is out of date or not there it will be built.
-					// If it is there and up to date, it will add 0.8 seconds to the build time.
-					Log.TraceInformation("Building UnrealHeaderTool...");
+					// Always build UnrealHeaderTool if header regeneration is required, unless we're running within a Rocket ecosystem
+					if (UnrealBuildTool.RunningRocket() == false && UEBuildConfiguration.bDoNotBuildUHT == false)
+					{
+						// If it is out of date or not there it will be built.
+						// If it is there and up to date, it will add 0.8 seconds to the build time.
+						Log.TraceInformation("Building UnrealHeaderTool...");
+
+						var UBTArguments = new StringBuilder();
+
+						UBTArguments.Append( "UnrealHeaderTool" );
+
+						// Which desktop platform do we need to compile UHT for?
+						UBTArguments.Append( " " + GetRuntimePlatform().ToString() );
+						// NOTE: We force Development configuration for UHT so that it runs quickly, even when compiling debug
+						UBTArguments.Append( " " + UnrealTargetConfiguration.Development.ToString() );
+
+						// NOTE: We disable mutex when launching UBT from within UBT to compile UHT
+						UBTArguments.Append( " -NoMutex" );
 
-					var UBTArguments = new StringBuilder();
+						if (UnrealBuildTool.CommandLineContains("-noxge"))
+						{
+							UBTArguments.Append(" -noxge");
+						}
+
+						if ( RunExternalExecutable( UnrealBuildTool.GetUBTPath(), UBTArguments.ToString() ) != 0 )
+							return false;
+					}
 
-					UBTArguments.Append( "UnrealHeaderTool" );
+					Progress.Write(1, 3);
 
-					// Which desktop platform do we need to compile UHT for?
-                    UBTArguments.Append( " " + GetRuntimePlatform().ToString() );
-					// NOTE: We force Development configuration for UHT so that it runs quickly, even when compiling debug
-					UBTArguments.Append( " " + UnrealTargetConfiguration.Development.ToString() );
-				
-					// NOTE: We disable mutex when launching UBT from within UBT to compile UHT
-					UBTArguments.Append( " -NoMutex" );
+					var ActualTargetName = String.IsNullOrEmpty( Target.GetTargetName() ) ? "UE4" : Target.GetTargetName();
+					Log.TraceInformation( "Parsing headers for {0}", ActualTargetName );
 
-					if (UnrealBuildTool.CommandLineContains("-noxge"))
+					string HeaderToolPath = GetHeaderToolPath();
+					if (!File.Exists(HeaderToolPath))
 					{
-						UBTArguments.Append(" -noxge");
+						throw new BuildException( "Unable to generate headers because UnrealHeaderTool binary was not found ({0}).", Path.GetFullPath( HeaderToolPath ) );
 					}
 
-					if ( RunExternalExecutable( UnrealBuildTool.GetUBTPath(), UBTArguments.ToString() ) != 0 )
-						return false;
-				}
+					// Disable extensions when serializing to remove the $type fields
+					Directory.CreateDirectory(Path.GetDirectoryName(ModuleInfoFileName));
+					System.IO.File.WriteAllText(ModuleInfoFileName, fastJSON.JSON.Instance.ToJSON(Manifest, new fastJSON.JSONParameters{ UseExtensions = false }));
 
-				var ActualTargetName = String.IsNullOrEmpty( Target.GetTargetName() ) ? "UE4" : Target.GetTargetName();
-				Log.TraceInformation( "Parsing headers for {0}", ActualTargetName );
+					string CmdLine = (UnrealBuildTool.HasUProjectFile()) ? "\"" + UnrealBuildTool.GetUProjectFile() + "\"" : Target.GetTargetName();
+					CmdLine += " \"" + ModuleInfoFileName + "\" -LogCmds=\"loginit warning, logexit warning, logdatabase error\"";
+					if (UnrealBuildTool.RunningRocket())
+					{
+						CmdLine += " -rocket -installed";
+					}
 
-				string HeaderToolPath = GetHeaderToolPath();
-				if (!File.Exists(HeaderToolPath))
-				{
-					throw new BuildException( "Unable to generate headers because UnrealHeaderTool binary was not found ({0}).", Path.GetFullPath( HeaderToolPath ) );
-				}
+					if (UEBuildConfiguration.bFailIfGeneratedCodeChanges)
+					{
+						CmdLine += " -FailIfGeneratedCodeChanges";
+					}
 
-				// Disable extensions when serializing to remove the $type fields
-				Directory.CreateDirectory(Path.GetDirectoryName(ModuleInfoFileName));
-				System.IO.File.WriteAllText(ModuleInfoFileName, fastJSON.JSON.Instance.ToJSON(Manifest, new fastJSON.JSONParameters{ UseExtensions = false }));
+					Stopwatch s = new Stopwatch();
+					s.Start();
+					UHTResult = (ECompilationResult) RunExternalExecutable(ExternalExecution.GetHeaderToolPath(), CmdLine);
+					s.Stop();
 
-				string CmdLine = (UnrealBuildTool.HasUProjectFile()) ? "\"" + UnrealBuildTool.GetUProjectFile() + "\"" : Target.GetTargetName();
-				CmdLine += " \"" + ModuleInfoFileName + "\" -LogCmds=\"loginit warning, logexit warning, logdatabase error\"";
-				if (UnrealBuildTool.RunningRocket())
-				{
-					CmdLine += " -rocket -installed";
-				}
+					if (UHTResult != ECompilationResult.Succeeded)
+					{
+						Log.TraceInformation("Error: Failed to generate code for {0} - error code: {1}", ActualTargetName, (int) UHTResult);
+						return false;
+					}
+
+					Log.TraceInformation( "Code generation finished for {0} and took {1}", ActualTargetName, (double)s.ElapsedMilliseconds/1000.0 );
 
-				if (UEBuildConfiguration.bFailIfGeneratedCodeChanges)
+					// Now that UHT has successfully finished generating code, we need to update all cached FileItems in case their last write time has changed.
+					// Otherwise UBT might not detect changes UHT made.
+					DateTime StartTime = DateTime.UtcNow;
+					FileItem.ResetInfos();
+					double ResetDuration = (DateTime.UtcNow - StartTime).TotalSeconds;
+					Log.TraceVerbose("FileItem.ResetInfos() duration: {0}s", ResetDuration);
+				}
+				else
 				{
-					CmdLine += " -FailIfGeneratedCodeChanges";
+					Log.TraceVerbose( "Generated code is up to date." );
 				}
 
-				Stopwatch s = new Stopwatch();
-				s.Start();
-				UHTResult = (ECompilationResult) RunExternalExecutable(ExternalExecution.GetHeaderToolPath(), CmdLine);
-				s.Stop();
+				Progress.Write(2, 3);
 
-				if (UHTResult != ECompilationResult.Succeeded)
+				// There will never be generated code if we're building UHT, so this should never be called.
+				if (!bIsBuildingUHT)
 				{
-					Log.TraceInformation("Error: Failed to generate code for {0} - error code: {1}", ActualTargetName, (int) UHTResult);
-					return false;
+					// Allow generated code to be sync'd to remote machines if needed. This needs to be done even if UHT did not run because
+					// generated headers include other generated headers using absolute paths which in case of building remotely are already
+					// the remote machine absolute paths. Because of that parsing headers will not result in finding all includes properly.
+					ToolChain.PostCodeGeneration(Target, Manifest);
 				}
 
-				Log.TraceInformation( "Code generation finished for {0} and took {1}", ActualTargetName, (double)s.ElapsedMilliseconds/1000.0 );
+				// touch the directories
+				UpdateDirectoryTimestamps(Target, UObjectModules);
 
-				// Now that UHT has successfully finished generating code, we need to update all cached FileItems in case their last write time has changed.
-				// Otherwise UBT might not detect changes UHT made.
-				DateTime StartTime = DateTime.UtcNow;
-				FileItem.ResetInfos();
-				double ResetDuration = (DateTime.UtcNow - StartTime).TotalSeconds;
-				Log.TraceVerbose("FileItem.ResetInfos() duration: {0}s", ResetDuration);
+				Progress.Write(3, 3);
 			}
-			else
-			{
-				Log.TraceVerbose( "Generated code is up to date." );
-			}
-
-			// There will never be generated code if we're building UHT, so this should never be called.
-			if (!bIsBuildingUHT)
-			{
-				// Allow generated code to be sync'd to remote machines if needed. This needs to be done even if UHT did not run because
-				// generated headers include other generated headers using absolute paths which in case of building remotely are already
-				// the remote machine absolute paths. Because of that parsing headers will not result in finding all includes properly.
-				ToolChain.PostCodeGeneration(Target, Manifest);
-			}
-
-			// touch the directories
-			UpdateDirectoryTimestamps(Target, UObjectModules);
-
 			return true;
 		}
 	}
diff --git a/Engine/Source/Programs/UnrealBuildTool/System/LocalExecutor.cs b/Engine/Source/Programs/UnrealBuildTool/System/LocalExecutor.cs
--- a/Engine/Source/Programs/UnrealBuildTool/System/LocalExecutor.cs
+++ b/Engine/Source/Programs/UnrealBuildTool/System/LocalExecutor.cs
@@ -345,98 +345,110 @@ namespace UnrealBuildTool
 
 			Dictionary<Action, ActionThread> ActionThreadDictionary = new Dictionary<Action, ActionThread>();
             int JobNumber = 1;
-			while(true)
+			using(ProgressWriter ProgressWriter = new ProgressWriter("Compiling source code...", false))
 			{
-				// Count the number of pending and still executing actions.
-				int NumUnexecutedActions = 0;
-				int NumExecutingActions = 0;
-				foreach (Action Action in Actions)
+				int ProgressValue = 0;
+				while(true)
 				{
-					ActionThread ActionThread = null;
-					bool bFoundActionProcess = ActionThreadDictionary.TryGetValue(Action, out ActionThread);
-					if (bFoundActionProcess == false)
+					// Count the number of pending and still executing actions.
+					int NumUnexecutedActions = 0;
+					int NumExecutingActions = 0;
+					foreach (Action Action in Actions)
 					{
-						NumUnexecutedActions++;
-					}
-					else if (ActionThread != null)
-					{
-						if (ActionThread.bComplete == false)
+						ActionThread ActionThread = null;
+						bool bFoundActionProcess = ActionThreadDictionary.TryGetValue(Action, out ActionThread);
+						if (bFoundActionProcess == false)
 						{
 							NumUnexecutedActions++;
-							NumExecutingActions++;
+						}
+						else if (ActionThread != null)
+						{
+							if (ActionThread.bComplete == false)
+							{
+								NumUnexecutedActions++;
+								NumExecutingActions++;
+							}
 						}
 					}
-				}
 
-				// If there aren't any pending actions left, we're done executing.
-				if (NumUnexecutedActions == 0)
-				{
-					break;
-				}
+					// Update the current progress
+					int NewProgressValue = Actions.Count + 1 - NumUnexecutedActions;
+					if (ProgressValue != NewProgressValue)
+					{
+						ProgressWriter.Write(ProgressValue, Actions.Count + 1);
+						ProgressValue = NewProgressValue;
+					}
 
-				// If there are fewer actions executing than the maximum, look for pending actions that don't have any outdated
-				// prerequisites.
-				foreach (Action Action in Actions)
-				{
-					ActionThread ActionProcess = null;
-					bool bFoundActionProcess = ActionThreadDictionary.TryGetValue(Action, out ActionProcess);
-					if (bFoundActionProcess == false)
+					// If there aren't any pending actions left, we're done executing.
+					if (NumUnexecutedActions == 0)
+					{
+						break;
+					}
+
+					// If there are fewer actions executing than the maximum, look for pending actions that don't have any outdated
+					// prerequisites.
+					foreach (Action Action in Actions)
 					{
-						if (NumExecutingActions < Math.Max(1,MaxActionsToExecuteInParallel))
+						ActionThread ActionProcess = null;
+						bool bFoundActionProcess = ActionThreadDictionary.TryGetValue(Action, out ActionProcess);
+						if (bFoundActionProcess == false)
 						{
-							// Determine whether there are any prerequisites of the action that are outdated.
-							bool bHasOutdatedPrerequisites = false;
-							bool bHasFailedPrerequisites = false;
-							foreach (FileItem PrerequisiteItem in Action.PrerequisiteItems)
+							if (NumExecutingActions < Math.Max(1,MaxActionsToExecuteInParallel))
 							{
-								if (PrerequisiteItem.ProducingAction != null && Actions.Contains(PrerequisiteItem.ProducingAction))
+								// Determine whether there are any prerequisites of the action that are outdated.
+								bool bHasOutdatedPrerequisites = false;
+								bool bHasFailedPrerequisites = false;
+								foreach (FileItem PrerequisiteItem in Action.PrerequisiteItems)
 								{
-									ActionThread PrerequisiteProcess = null;
-									bool bFoundPrerequisiteProcess = ActionThreadDictionary.TryGetValue( PrerequisiteItem.ProducingAction, out PrerequisiteProcess );
-									if (bFoundPrerequisiteProcess == true)
+									if (PrerequisiteItem.ProducingAction != null && Actions.Contains(PrerequisiteItem.ProducingAction))
 									{
-										if (PrerequisiteProcess == null)
+										ActionThread PrerequisiteProcess = null;
+										bool bFoundPrerequisiteProcess = ActionThreadDictionary.TryGetValue( PrerequisiteItem.ProducingAction, out PrerequisiteProcess );
+										if (bFoundPrerequisiteProcess == true)
 										{
-											bHasFailedPrerequisites = true;
+											if (PrerequisiteProcess == null)
+											{
+												bHasFailedPrerequisites = true;
+											}
+											else if (PrerequisiteProcess.bComplete == false)
+											{
+												bHasOutdatedPrerequisites = true;
+											}
+											else if (PrerequisiteProcess.ExitCode != 0)
+											{
+												bHasFailedPrerequisites = true;
+											}
 										}
-										else if (PrerequisiteProcess.bComplete == false)
+										else
 										{
 											bHasOutdatedPrerequisites = true;
 										}
-										else if (PrerequisiteProcess.ExitCode != 0)
-										{
-											bHasFailedPrerequisites = true;
-										}
-									}
-									else
-									{
-										bHasOutdatedPrerequisites = true;
 									}
 								}
-							}
 
-							// If there are any failed prerequisites of this action, don't execute it.
-							if (bHasFailedPrerequisites)
-							{
-								// Add a null entry in the dictionary for this action.
-								ActionThreadDictionary.Add( Action, null );
-							}
-							// If there aren't any outdated prerequisites of this action, execute it.
-                            else if (!bHasOutdatedPrerequisites)
-                            {
-                                ActionThread ActionThread = new ActionThread(Action, JobNumber, Actions.Count);
-                                JobNumber++;
-								ActionThread.Run();
+								// If there are any failed prerequisites of this action, don't execute it.
+								if (bHasFailedPrerequisites)
+								{
+									// Add a null entry in the dictionary for this action.
+									ActionThreadDictionary.Add( Action, null );
+								}
+								// If there aren't any outdated prerequisites of this action, execute it.
+								else if (!bHasOutdatedPrerequisites)
+								{
+									ActionThread ActionThread = new ActionThread(Action, JobNumber, Actions.Count);
+									JobNumber++;
+									ActionThread.Run();
 
-								ActionThreadDictionary.Add(Action, ActionThread);
+									ActionThreadDictionary.Add(Action, ActionThread);
 
-								NumExecutingActions++;
-                            }
+									NumExecutingActions++;
+								}
+							}
 						}
 					}
-				}
 
-				System.Threading.Thread.Sleep(TimeSpan.FromSeconds(LoopSleepTime));
+					System.Threading.Thread.Sleep(TimeSpan.FromSeconds(LoopSleepTime));
+				}
 			}
 
 			Log.WriteLineIf(BuildConfiguration.bLogDetailedActionStats, TraceEventType.Information, "-------- Begin Detailed Action Stats ----------------------------------------------------------");
diff --git a/Engine/Source/Programs/UnrealBuildTool/System/ProjectFileGenerator.cs b/Engine/Source/Programs/UnrealBuildTool/System/ProjectFileGenerator.cs
--- a/Engine/Source/Programs/UnrealBuildTool/System/ProjectFileGenerator.cs
+++ b/Engine/Source/Programs/UnrealBuildTool/System/ProjectFileGenerator.cs
@@ -1223,7 +1223,7 @@ namespace UnrealBuildTool
 			var bSuccess = true;
 			if( ShouldGenerateIntelliSenseData() && TargetFiles.Count > 0 )
 			{
-				using(ProgressWriter Progress = new ProgressWriter("Binding IntelliSense data..."))
+				using(ProgressWriter Progress = new ProgressWriter("Binding IntelliSense data...", true))
 				{
 					for(int TargetIndex = 0; TargetIndex < TargetFiles.Count; ++TargetIndex)
 					{
@@ -1853,7 +1853,7 @@ namespace UnrealBuildTool
 		/// <returns>True if successful</returns>
 		protected virtual bool WriteProjectFiles()
 		{
-			using(ProgressWriter Progress = new ProgressWriter("Writing project files..."))
+			using(ProgressWriter Progress = new ProgressWriter("Writing project files...", true))
 			{
 				var TotalProjectFileCount = GeneratedProjectFiles.Count + 1;	// +1 for the master project file, which we'll save next
 
diff --git a/Engine/Source/Programs/UnrealBuildTool/System/Utils.cs b/Engine/Source/Programs/UnrealBuildTool/System/Utils.cs
--- a/Engine/Source/Programs/UnrealBuildTool/System/Utils.cs
+++ b/Engine/Source/Programs/UnrealBuildTool/System/Utils.cs
@@ -893,13 +893,15 @@ namespace UnrealBuildTool
 	{
 		public static bool bWriteMarkup = false;
 
+		bool bWriteToConsole;
 		string Message;
 		int NumCharsToBackspaceOver;
 
-		public ProgressWriter(string InMessage)
+		public ProgressWriter(string InMessage, bool bInWriteToConsole)
 		{
 			Message = InMessage;
-			if(!bWriteMarkup)
+			bWriteToConsole = bInWriteToConsole;
+			if (!bWriteMarkup && bWriteToConsole)
 			{
 				Console.Write(Message + " ");
 			}
@@ -908,7 +910,7 @@ namespace UnrealBuildTool
 
 		public void Dispose()
 		{
-			if(!bWriteMarkup)
+			if (!bWriteMarkup && bWriteToConsole)
 			{
 				Console.WriteLine();
 			}
@@ -923,7 +925,7 @@ namespace UnrealBuildTool
 			{
 				Log.WriteLine(TraceEventType.Information, "@progress '{0}' {1}", Message, ProgressString);
 			}
-			else
+			else if (bWriteToConsole)
 			{
 				// Backspace over previous progress value
 				while (NumCharsToBackspaceOver-- > 0)
diff --git a/Engine/Source/Programs/UnrealBuildTool/System/XGE.cs b/Engine/Source/Programs/UnrealBuildTool/System/XGE.cs
--- a/Engine/Source/Programs/UnrealBuildTool/System/XGE.cs
+++ b/Engine/Source/Programs/UnrealBuildTool/System/XGE.cs
@@ -13,6 +13,8 @@ namespace UnrealBuildTool
 {
 	public class XGE
 	{
+		private const string ProgressMarkupPrefix = "@action";
+
 		public static void SaveXGEFile(string XGETaskFilePath)
 		{
 			int FileNum = 0;
@@ -41,7 +43,7 @@ namespace UnrealBuildTool
 			{
 				// Write the actions to execute to a XGE task file.
 				string XGETaskFilePath = Path.Combine(BuildConfiguration.BaseIntermediatePath, "XGETasks.xml");
-				WriteTaskFile(Actions, XGETaskFilePath);
+				WriteTaskFile(Actions, XGETaskFilePath, ProgressWriter.bWriteMarkup);
 
 				if (BuildConfiguration.bXGEExport)
 				{
@@ -53,7 +55,7 @@ namespace UnrealBuildTool
 					if (Telemetry.IsAvailable())
 					{
 						// Add a custom output handler to determine the build duration for each task and map it back to the action that generated it.
-						XGEResult = XGE.ExecuteTaskFile(XGETaskFilePath, (sender, args) =>
+						XGEResult = XGE.ExecuteTaskFileWithProgressMarkup(XGETaskFilePath, Actions.Count, (sender, args) =>
 							{
 								// sometimes the args comes in as null
 								var match = XGEDurationRegex.Match(args.Data ?? "");
@@ -115,17 +117,7 @@ namespace UnrealBuildTool
 								// forward the output on to the normal handler or the console like normal
 								if (Actions[0].OutputEventHandler != null)
 								{
-									// NOTE: This is a pretty nasty hack with C# reflection, however it saves us from having to replace all the
-									// handlers in various Toolchains with a wrapper handler that takes a string - it is certainly doable, but 
-									// touches code outside of this class that I don't want to touch right now
-
-									// DataReceivedEventArgs is not normally constructable, so work around it with creating a scratch Args object
-									DataReceivedEventArgs EventArgs = (DataReceivedEventArgs)System.Runtime.Serialization.FormatterServices.GetUninitializedObject(typeof(DataReceivedEventArgs));
-
-									// now we need to set the Data field using reflection, since it is read only
-									FieldInfo[] ArgFields = typeof(DataReceivedEventArgs).GetFields(BindingFlags.NonPublic | BindingFlags.Instance | BindingFlags.DeclaredOnly);
-
-									ArgFields[0].SetValue(EventArgs, updatedData);
+									DataReceivedEventArgs EventArgs = ConstructDataReceivedEventArgs(updatedData);
 									Actions[0].OutputEventHandler(sender, EventArgs);
 								}
 								else
@@ -136,15 +128,40 @@ namespace UnrealBuildTool
 					}
 					else
 					{
-						XGEResult = XGE.ExecuteTaskFile(XGETaskFilePath, Actions[0].OutputEventHandler);
+						XGEResult = XGE.ExecuteTaskFileWithProgressMarkup(XGETaskFilePath, Actions.Count, (Sender, Args) => 
+							{
+								if(Actions[0].OutputEventHandler != null)
+								{
+									Actions[0].OutputEventHandler(Sender, Args);
+								}
+								else
+								{
+									Console.WriteLine(Args.Data);
+								}
+							});
 					}
 				}
 			}
 			return XGEResult;
 		}
 
+		private static DataReceivedEventArgs ConstructDataReceivedEventArgs(string NewData)
+		{
+			// NOTE: This is a pretty nasty hack with C# reflection, however it saves us from having to replace all the
+			// handlers in various Toolchains with a wrapper handler that takes a string - it is certainly doable, but 
+			// touches code outside of this class that I don't want to touch right now
+
+			// DataReceivedEventArgs is not normally constructable, so work around it with creating a scratch Args object
+			DataReceivedEventArgs EventArgs = (DataReceivedEventArgs)System.Runtime.Serialization.FormatterServices.GetUninitializedObject(typeof(DataReceivedEventArgs));
+
+			// now we need to set the Data field using reflection, since it is read only
+			FieldInfo[] ArgFields = typeof(DataReceivedEventArgs).GetFields(BindingFlags.NonPublic | BindingFlags.Instance | BindingFlags.DeclaredOnly);
+			ArgFields[0].SetValue(EventArgs, NewData);
+			return EventArgs;
+		}
+
 		/** Writes a XGE task file containing the specified actions to the specified file path. */
-		public static void WriteTaskFile(List<Action> InActions, string TaskFilePath)
+		public static void WriteTaskFile(List<Action> InActions, string TaskFilePath, bool bProgressMarkup)
 		{
             Dictionary<string, string> ExportEnv = new Dictionary<string, string>();
 
@@ -277,10 +294,21 @@ namespace UnrealBuildTool
 				ToolsElement.AppendChild(ToolElement);
 				ToolElement.SetAttribute("Name", string.Format("Tool{0}", ActionIndex));
 				ToolElement.SetAttribute("AllowRemote", Action.bCanExecuteRemotely.ToString());
+
+				string OutputPrefix = "";
+				if (bProgressMarkup)
+				{
+					OutputPrefix += ProgressMarkupPrefix;
+				}
 				if( Action.bShouldOutputStatusDescription )
 				{
-					ToolElement.SetAttribute( "OutputPrefix", Action.StatusDescription );
+					OutputPrefix += Action.StatusDescription;
+				}
+				if(OutputPrefix.Length > 0)
+				{
+					ToolElement.SetAttribute("OutputPrefix", OutputPrefix);
 				}
+
 				// batch files (.bat, .cmd) need to be run via or cmd /c or shellexecute, 
 				// the latter which we can't use because we want to redirect input/output
 
@@ -444,5 +472,33 @@ namespace UnrealBuildTool
 				return ExecutionResult.Unavailable;
 			}
 		}
+
+		/**
+		 * Executes the tasks in the specified file, parsing progress markup as part of the output.
+		 */
+		public static ExecutionResult ExecuteTaskFileWithProgressMarkup(string TaskFilePath, int NumActions, DataReceivedEventHandler OutputEventHandler)
+		{
+			using (ProgressWriter Writer = new ProgressWriter("Compiling source files...", false))
+			{
+				int NumCompletedActions = 0;
+
+				// Create a wrapper delegate that will parse the output actions
+				DataReceivedEventHandler EventHandlerWrapper = (Sender, Args) =>
+				{
+					if (Args.Data != null && Args.Data.StartsWith(ProgressMarkupPrefix))
+					{
+						Writer.Write(++NumCompletedActions, NumActions);
+						Args = ConstructDataReceivedEventArgs(Args.Data.Substring(ProgressMarkupPrefix.Length));
+					}
+					if(OutputEventHandler != null)
+					{
+						OutputEventHandler(Sender, Args);
+					}
+				};
+
+				// Run through the standard XGE executor
+				return ExecuteTaskFile(TaskFilePath, EventHandlerWrapper);
+			}
+		}
 	}
 }
