UBT: Fix for wrong target name (and intermediate output path) when re-building game DLLs in the editor

diff --git a/Engine/Source/Programs/UnrealBuildTool/System/RulesCompiler.cs b/Engine/Source/Programs/UnrealBuildTool/System/RulesCompiler.cs
--- a/Engine/Source/Programs/UnrealBuildTool/System/RulesCompiler.cs
+++ b/Engine/Source/Programs/UnrealBuildTool/System/RulesCompiler.cs
@@ -1569,6 +1569,12 @@ namespace UnrealBuildTool
 		{
 			string TargetFileName;
 			TargetRules RulesObject = CreateTargetRules(TargetName, Target, bInEditorRecompile, out TargetFileName);
+			if (bInEditorRecompile)
+			{
+				// Now that we found the actual Editor target, make sure we're no longer using the old TargetName (which is the Game target)
+				var TargetSuffixIndex = RulesObject.TargetName.LastIndexOf("Target");
+				TargetName = (TargetSuffixIndex > 0) ? RulesObject.TargetName.Substring(0, TargetSuffixIndex) : RulesObject.TargetName;
+			}
 			if ((ProjectFileGenerator.bGenerateProjectFiles == false) && (RulesObject.SupportsPlatform(Target.Platform) == false))
 			{
 				if (UEBuildConfiguration.bCleanProject)
