Fix for UTF-8-encoded files written by UBT.

#fix FileItem.CreateIntermediateTextFile now creates the text file in UTF-8 if the contents string contains non-ASCII characters.
#change ResponseFile.Create now calls FileItem.CreateIntermediateTextFile.

This fix was made in relation to TTP# 334906, even though that problem was not fixed in the end.

diff --git a/Engine/Source/Programs/UnrealBuildTool/System/FileItem.cs b/Engine/Source/Programs/UnrealBuildTool/System/FileItem.cs
--- a/Engine/Source/Programs/UnrealBuildTool/System/FileItem.cs
+++ b/Engine/Source/Programs/UnrealBuildTool/System/FileItem.cs
@@ -177,6 +177,17 @@ namespace UnrealBuildTool
 			}
 		}
 
+		/// <summary>
+		/// Determines the appropriate encoding for a string: either ASCII or UTF-8.
+		/// </summary>
+		/// <param name="Str">The string to test.</param>
+		/// <returns>Either System.Text.Encoding.ASCII or System.Text.Encoding.UTF8, depending on whether or not the string contains non-ASCII characters.</returns>
+		private static Encoding GetEncodingForString(string Str)
+		{
+			// If the string length is equivalent to the encoded length, then no non-ASCII characters were present in the string.
+			return (Encoding.UTF8.GetByteCount(Str) == Str.Length) ? Encoding.ASCII : Encoding.UTF8;
+		}
+
 		/**
 		 * Creates a text file with the given contents.  If the contents of the text file aren't changed, it won't write the new contents to
 		 * the file to avoid causing an action to be considered outdated.
@@ -189,7 +200,7 @@ namespace UnrealBuildTool
 			// Only write the file if its contents have changed.
 			if (!File.Exists(AbsolutePath) || Utils.ReadAllText(AbsolutePath) != Contents)
 			{
-				File.WriteAllText(AbsolutePath, Contents);
+				File.WriteAllText(AbsolutePath, Contents, GetEncodingForString(Contents));
 			}
 
 			return GetItemByPath(AbsolutePath);
diff --git a/Engine/Source/Programs/UnrealBuildTool/System/ResponseFile.cs b/Engine/Source/Programs/UnrealBuildTool/System/ResponseFile.cs
--- a/Engine/Source/Programs/UnrealBuildTool/System/ResponseFile.cs
+++ b/Engine/Source/Programs/UnrealBuildTool/System/ResponseFile.cs
@@ -27,21 +27,7 @@ namespace UnrealBuildTool
 				TempFileInfo.Refresh();
 			}
 
-			// Create the folder if it doesn't exist
-			if( !TempFolderInfo.Exists )
-			{
-				// Create the 
-				TempFolderInfo.Create();
-				TempFolderInfo.Refresh();
-			}
-
-			using( FileStream Writer = TempFileInfo.OpenWrite() )
-			{
-				using( StreamWriter TextWriter = new StreamWriter( Writer ) )
-				{
-					Lines.ForEach( x => TextWriter.WriteLine( x ) );
-				}
-			}
+			FileItem.CreateIntermediateTextFile(TempFileName, string.Join(Environment.NewLine, Lines));
 
 			return TempFileName;
 		}
