Fixed remote compilation for Mac

diff --git a/Engine/Source/Programs/UnrealBuildTool/Mac/MacToolChain.cs b/Engine/Source/Programs/UnrealBuildTool/Mac/MacToolChain.cs
--- a/Engine/Source/Programs/UnrealBuildTool/Mac/MacToolChain.cs
+++ b/Engine/Source/Programs/UnrealBuildTool/Mac/MacToolChain.cs
@@ -888,7 +888,7 @@ namespace UnrealBuildTool
 					string CustomResourcesPath = "";
 					if (string.IsNullOrEmpty(UProjectFilePath))
 					{
-						string[] TargetFiles = Directory.GetFiles(EngineSourcePath, GameName + ".Target.cs", SearchOption.AllDirectories);
+						string[] TargetFiles = Directory.GetFiles(Directory.GetCurrentDirectory(), GameName + ".Target.cs", SearchOption.AllDirectories);
 						if (TargetFiles.Length == 1)
 						{
 							CustomResourcesPath = Path.GetDirectoryName(TargetFiles[0]) + "/Resources/Mac";
