Fix for UBT not detecting UnrealHeaderTool being out of date if any of its DDLs/Plugins have changed

diff --git a/Engine/Source/Programs/UnrealBuildTool/System/ExternalExecution.cs b/Engine/Source/Programs/UnrealBuildTool/System/ExternalExecution.cs
--- a/Engine/Source/Programs/UnrealBuildTool/System/ExternalExecution.cs
+++ b/Engine/Source/Programs/UnrealBuildTool/System/ExternalExecution.cs
@@ -113,6 +113,9 @@ namespace UnrealBuildTool
 		{
 		}
 
+		/// <summary>
+		/// Gets UnrealHeaderTool.exe path. Does not care if UnrealheaderTool was build as a monolithic exe or not.
+		/// </summary>
 		static string GetHeaderToolPath()
 		{
 			UnrealTargetPlatform Platform = GetRuntimePlatform();
@@ -122,6 +125,83 @@ namespace UnrealBuildTool
 			return HeaderToolPath;
 		}
 
+		/// <summary>
+		/// Finds all UnrealHeaderTool plugins in the plugin directory
+		/// </summary>
+		static void RecursivelyCollectHeaderToolPlugins(string RootPath, string Pattern, string Platform, List<string> PluginBinaries)
+		{
+			var SubDirectories = Directory.GetDirectories(RootPath);
+			foreach (var Dir in SubDirectories)
+			{
+				if (Dir.IndexOf("Intermediate", StringComparison.InvariantCultureIgnoreCase) < 0)
+				{
+					RecursivelyCollectHeaderToolPlugins(Dir, Pattern, Platform, PluginBinaries);
+					if (Dir.EndsWith("Binaries", StringComparison.InvariantCultureIgnoreCase))
+					{
+						// No need to search the other folders
+						break;
+					}
+				}
+			}
+			var Binaries = Directory.GetFiles(RootPath, Pattern);
+			foreach (var Binary in Binaries)
+			{
+				if (Binary.Contains(Platform))
+				{
+					PluginBinaries.Add(Binary);
+				}
+			}
+		}
+
+		/// <summary>
+		/// Gets all UnrealHeaderTool binaries (including DLLs if it was not build monolithically)
+		/// </summary>
+		static string[] GetHeaderToolBinaries()
+		{
+			var Binaries = new List<string>();
+			var HeaderToolExe = GetHeaderToolPath();
+			if (File.Exists(HeaderToolExe))
+			{
+				Binaries.Add(HeaderToolExe);
+
+				var HeaderToolLocation = Path.GetDirectoryName(HeaderToolExe);
+				var Platform = GetRuntimePlatform();
+				var DLLExtension = UEBuildPlatform.GetBuildPlatform(Platform).GetBinaryExtension(UEBuildBinaryType.DynamicLinkLibrary);
+				var DLLSearchPattern = "UnrealHeaderTool-*" + DLLExtension;
+				var HeaderToolDLLs = Directory.GetFiles(HeaderToolLocation, DLLSearchPattern, SearchOption.TopDirectoryOnly);
+								
+				Binaries.AddRange(HeaderToolDLLs);
+
+				var PluginDirectory = Path.Combine("..", "Plugins");
+				RecursivelyCollectHeaderToolPlugins(PluginDirectory, DLLSearchPattern, Platform.ToString(), Binaries);
+			}
+			return Binaries.ToArray();
+		}
+
+		/// <summary>
+		/// Gets the latest write time of any of the UnrealHeaderTool binaries (including DLLs and Plugins) or DateTime.MaxValue if UnrealHeaderTool does not exist
+		/// </summary>
+		static DateTime GetHeaderToolTimestamp()
+		{			 
+			var HeaderToolBinaries = GetHeaderToolBinaries();
+			var LatestWriteTime = DateTime.MinValue;
+			// Find the latest write time for all UnrealHeaderTool binaries
+			foreach (var Binary in HeaderToolBinaries)
+			{
+				var BinaryInfo = new FileInfo(Binary);
+				if (BinaryInfo.Exists)
+				{
+					if (BinaryInfo.LastWriteTime > LatestWriteTime)
+					{
+						LatestWriteTime = BinaryInfo.LastWriteTime;
+					}
+				}
+			}
+			// If UHT doesn't exist, force regenerate.
+			return LatestWriteTime > DateTime.MinValue ? LatestWriteTime : DateTime.MaxValue;
+		}
+
+
 		/** Returns the name of platform UBT is running on */
 		public static UnrealTargetPlatform GetRuntimePlatform()
 		{
@@ -187,12 +267,7 @@ namespace UnrealBuildTool
 			bool bIsOutOfDate = false;
 
 			// Get UnrealHeaderTool timestamp. If it's newer than generated headers, they need to be rebuilt too.
-			string HeaderToolPath = GetHeaderToolPath();
-			var HeaderToolTimestamp = DateTime.MaxValue;				// If UHT doesn't exist, force regenerate.
-			if( File.Exists( HeaderToolPath ) )
-			{
-				HeaderToolTimestamp = new FileInfo(HeaderToolPath).LastWriteTime;
-			}
+			var HeaderToolTimestamp = GetHeaderToolTimestamp();
 
 			// Get CoreUObject.generated.cpp timestamp.  If the source files are older than the CoreUObject generated code, we'll
 			// need to regenerate code for the module
