[GitHub] 144 : Fix indentation of trace output

diff --git a/Engine/Source/Programs/UnrealBuildTool/System/UnrealBuildTool.cs b/Engine/Source/Programs/UnrealBuildTool/System/UnrealBuildTool.cs
--- a/Engine/Source/Programs/UnrealBuildTool/System/UnrealBuildTool.cs
+++ b/Engine/Source/Programs/UnrealBuildTool/System/UnrealBuildTool.cs
@@ -1,4 +1,4 @@
-// Copyright 1998-2014 Epic Games, Inc. All Rights Reserved.
+﻿// Copyright 1998-2014 Epic Games, Inc. All Rights Reserved.
 
 using System;
 using System.Collections.Generic;
@@ -416,7 +416,7 @@ namespace UnrealBuildTool
 						}
 						else if (CheckType.IsSubclassOf(typeof(UEBuildDeploy)))
 						{
-							Log.TraceVerbose("    Registering build deploy: {0}", CheckType.ToString());
+							Log.TraceVerbose("    Registering build deploy  : {0}", CheckType.ToString());
 							UEBuildDeploy TempInst = (UEBuildDeploy)(UBTAssembly.CreateInstance(CheckType.FullName, true));
 							TempInst.RegisterBuildDeploy();
 						}
