UBT: Fixed a minor warning in Mono.

diff --git a/Engine/Source/Programs/UnrealBuildTool/System/ResponseFile.cs b/Engine/Source/Programs/UnrealBuildTool/System/ResponseFile.cs
--- a/Engine/Source/Programs/UnrealBuildTool/System/ResponseFile.cs
+++ b/Engine/Source/Programs/UnrealBuildTool/System/ResponseFile.cs
@@ -17,8 +17,6 @@ namespace UnrealBuildTool
 		public static string Create(string TempFileName, List<string> Lines)
 		{
 			FileInfo TempFileInfo = new FileInfo( TempFileName );
-			DirectoryInfo TempFolderInfo = new DirectoryInfo( TempFileInfo.DirectoryName );
-
 			// Delete the existing file if it exists
 			if( TempFileInfo.Exists )
 			{
