UE4
  TTP 335792. CRITICAL: Regression: ROCKET: Code Templates fail to generate project files).

   Engine should be dependent on platform specific modules, use SetupNewlyLoadedModules instead.

diff --git a/Engine/Source/Programs/UnrealBuildTool/HTML5/UEBuildHTML5.cs b/Engine/Source/Programs/UnrealBuildTool/HTML5/UEBuildHTML5.cs
--- a/Engine/Source/Programs/UnrealBuildTool/HTML5/UEBuildHTML5.cs
+++ b/Engine/Source/Programs/UnrealBuildTool/HTML5/UEBuildHTML5.cs
@@ -256,49 +256,25 @@ namespace UnrealBuildTool
 				if (InModule.ToString() == "Core")
 				{
 					InModule.AddPublicIncludePath("Runtime/Core/Public/HTML5");
-                    InModule.AddPublicDependencyModule("zlib");
+					InModule.AddPublicDependencyModule("zlib");
 				}
 				else if (InModule.ToString() == "Engine")
 				{
 					InModule.AddPrivateDependencyModule("zlib");
 					InModule.AddPrivateDependencyModule("UElibPNG");
-                    InModule.AddPublicDependencyModule("UEOgg");
-                    InModule.AddPublicDependencyModule("Vorbis");
+					InModule.AddPublicDependencyModule("UEOgg");
+					InModule.AddPublicDependencyModule("Vorbis");
 				}
 			}
-			else if (Target.Platform == UnrealTargetPlatform.Win64)
+			else if (Target.Platform == UnrealTargetPlatform.Win64 || Target.Platform == UnrealTargetPlatform.Mac )
 			{
-                bool bBuildShaderFormats = UEBuildConfiguration.bForceBuildShaderFormats;
-				if (!UEBuildConfiguration.bBuildRequiresCookedData)
+				if ( (!UEBuildConfiguration.bBuildRequiresCookedData
+					&& InModule.ToString() == "Engine"
+					&& UEBuildConfiguration.bBuildDeveloperTools)
+					|| UEBuildConfiguration.bForceBuildTargetPlatforms)
 				{
-					if (InModule.ToString() == "Engine")
-					{
-						if (UEBuildConfiguration.bBuildDeveloperTools)
-						{
-                            InModule.AddPlatformSpecificDynamicallyLoadedModule("HTML5TargetPlatform");
-						}
-					}
-					else if (InModule.ToString() == "TargetPlatform")
-					{
-                        bBuildShaderFormats = true;
-// 							if (UEBuildConfiguration.bBuildDeveloperTools)
-// 							{
-                        // 								InModule.AddPlatformSpecificDynamicallyLoadedModule("AT9AudioFormat");
-                        // 								InModule.AddPlatformSpecificDynamicallyLoadedModule("HTML5TextureFormat");
-// 							}
-					}
+					InModule.AddPlatformSpecificDynamicallyLoadedModule("HTML5TargetPlatform");
 				}
-
-				// allow standalone tools to use target platform modules, without needing Engine
-				if (UEBuildConfiguration.bForceBuildTargetPlatforms)
-				{
-                    InModule.AddPlatformSpecificDynamicallyLoadedModule("HTML5TargetPlatform");
-				}
-
-                if (bBuildShaderFormats)
-                {
-                    // InModule.AddDynamicallyLoadedModule("HTML5ShaderFormat");
-                }
 			}
 		}
 
diff --git a/Engine/Source/Runtime/Engine/Engine.Build.cs b/Engine/Source/Runtime/Engine/Engine.Build.cs
--- a/Engine/Source/Runtime/Engine/Engine.Build.cs
+++ b/Engine/Source/Runtime/Engine/Engine.Build.cs
@@ -119,8 +119,7 @@ public class Engine : ModuleRules
 					    "MacTargetPlatform",
 					    "MacNoEditorTargetPlatform",
 						"MacServerTargetPlatform",
-						"MacClientTargetPlatform",
-					    "HTML5TargetPlatform"
+						"MacClientTargetPlatform"
 					}
 				);
 			}
