Remove Monolithic Direct3D switch, as this is now enabled by default.

diff --git a/Engine/Programs/UnrealBuildTool/BuildConfiguration.xml b/Engine/Programs/UnrealBuildTool/BuildConfiguration.xml
--- a/Engine/Programs/UnrealBuildTool/BuildConfiguration.xml
+++ b/Engine/Programs/UnrealBuildTool/BuildConfiguration.xml
@@ -84,7 +84,6 @@
 		<ProcessorCountMultiplier>1.0</ProcessorCountMultiplier>
 		<bTestIncludeDependencyResolveCache>False</bTestIncludeDependencyResolveCache>
 		<bUseIncludeDependencyResolveCache>True</bUseIncludeDependencyResolveCache>
-		<bUseMonolithicGraphicsDrivers>True</bUseMonolithicGraphicsDrivers>
 		<bUseFastMonoCalls>True</bUseFastMonoCalls>
 		<bDebugBuildsActuallyUseDebugCRT>False</bDebugBuildsActuallyUseDebugCRT>
 		<BaseIntermediateFolder>Intermediate/Build/</BaseIntermediateFolder>
diff --git a/Engine/Source/Programs/UnrealBuildTool/Configuration/BuildConfiguration.cs b/Engine/Source/Programs/UnrealBuildTool/Configuration/BuildConfiguration.cs
--- a/Engine/Source/Programs/UnrealBuildTool/Configuration/BuildConfiguration.cs
+++ b/Engine/Source/Programs/UnrealBuildTool/Configuration/BuildConfiguration.cs
@@ -26,12 +26,6 @@ namespace UnrealBuildTool
 		/// </summary>
 		public static int MinGameModuleSourceFilesForUnityBuild;
 
-        /// <summary>
-        /// Whether to build with Monolithic D3D Graphics drivers
-        ///	Eventually this will become mandatory
-        /// </summary>
-        public static bool bUseMonolithicGraphicsDrivers;
-
 		/// <summary>
 		/// New Monolithic Graphics drivers have optional "fast calls" replacing various D3d functions
 		/// </summary>
@@ -386,12 +380,6 @@ namespace UnrealBuildTool
 			// if we are testing the resolve cache, we require UBT to use it.
 			bUseIncludeDependencyResolveCache = true;
             
-            //IMPORTANT THIS IS THE MAIN SWITCH FOR MONO.
-            // New Monolithic Graphics drivers replace D3D and are *mostly* API compatible
-            // Eventually this will become mandatory
-            //  if this is set to true, then mono will be on by default on Dingo, and if false it will be off by default on Dingo.
-            //  This can be overridden by -monolithicdrivers or -nomonolithicdrivers in the NMAKE params.
-			bUseMonolithicGraphicsDrivers = true;
 			//IMPORTANT THIS IS THE MAIN SWITCH FOR MONO FAST CALLS
 			//  if this is set to true, then fast calls will be on by default on Dingo, and if false it will be off by default on Dingo.
 			//  This can be overridden by -fastmonocalls  or -nofastmonocalls in the NMAKE params.
diff --git a/Engine/Source/Programs/UnrealBuildTool/System/UnrealBuildTool.cs b/Engine/Source/Programs/UnrealBuildTool/System/UnrealBuildTool.cs
--- a/Engine/Source/Programs/UnrealBuildTool/System/UnrealBuildTool.cs
+++ b/Engine/Source/Programs/UnrealBuildTool/System/UnrealBuildTool.cs
@@ -1273,29 +1273,15 @@ namespace UnrealBuildTool
                 BuildConfiguration.bForceUnityBuild = true;
             }
 
-            // New Monolithic Graphics drivers replace D3D and are *mostly* API compatible
-            if (Utils.ParseCommandLineFlag(Arguments, "-monolithicdrivers", out ArgumentIndex))
-            {
-				BuildConfiguration.bUseMonolithicGraphicsDrivers = true;
-            }
-            if (Utils.ParseCommandLineFlag(Arguments, "-nomonolithicdrivers", out ArgumentIndex))
-            {
-				BuildConfiguration.bUseMonolithicGraphicsDrivers = false;
-            }
             // New Monolithic Graphics drivers have optional "fast calls" replacing various D3d functions
             if (Utils.ParseCommandLineFlag(Arguments, "-fastmonocalls", out ArgumentIndex))
             {
-                BuildConfiguration.bUseMonolithicGraphicsDrivers = true;    //can't have fast calls without mono !
                 BuildConfiguration.bUseFastMonoCalls = true;
             }
             if (Utils.ParseCommandLineFlag(Arguments, "-nofastmonocalls", out ArgumentIndex))
             {
                 BuildConfiguration.bUseFastMonoCalls = false;
             }
-			if (BuildConfiguration.bUseMonolithicGraphicsDrivers == false)
-			{
-				BuildConfiguration.bUseFastMonoCalls = false;
-			}
             
             if (Utils.ParseCommandLineFlag(Arguments, "-uniqueintermediate", out ArgumentIndex))
 			{
diff --git a/Engine/Source/Runtime/Windows/D3D11RHI/Private/D3D11Device.cpp b/Engine/Source/Runtime/Windows/D3D11RHI/Private/D3D11Device.cpp
--- a/Engine/Source/Runtime/Windows/D3D11RHI/Private/D3D11Device.cpp
+++ b/Engine/Source/Runtime/Windows/D3D11RHI/Private/D3D11Device.cpp
@@ -234,7 +234,7 @@ void FD3D11DynamicRHI::RHIGetSupportedResolution( uint32 &Width, uint32 &Height
 		DXGI_ADAPTER_DESC AdapterDesc;
 		VERIFYD3D11RESULT(Adapter->GetDesc(&AdapterDesc));
 	  
-#ifndef USE_MONOLITHIC_GRAPHICS_DRIVERS
+#ifndef PLATFORM_XBOXONE // No need for display mode enumeration on console
 		// Enumerate outputs for this adapter
 		// TODO: Cap at 1 for default output
 		for(uint32 o = 0;o < 1; o++)
@@ -282,7 +282,7 @@ void FD3D11DynamicRHI::RHIGetSupportedResolution( uint32 &Width, uint32 &Height
 
 			delete[] ModeList;
 		}
-#endif // USE_MONOLITHIC_GRAPHICS_DRIVERS
+#endif // PLATFORM_XBOXONE
 	}
 
 	check(InitializedMode);
diff --git a/Engine/Source/Runtime/Windows/D3D11RHI/Private/D3D11StateCache.cpp b/Engine/Source/Runtime/Windows/D3D11RHI/Private/D3D11StateCache.cpp
--- a/Engine/Source/Runtime/Windows/D3D11RHI/Private/D3D11StateCache.cpp
+++ b/Engine/Source/Runtime/Windows/D3D11RHI/Private/D3D11StateCache.cpp
@@ -185,7 +185,7 @@ void FD3D11StateCacheBase::VerifyCacheState()
 		Direct3DDeviceIMContext->OMGetBlendState(BlendState.GetInitReference(), BlendFactor, &SampleMask);
 
 		check(BlendState.GetReference() == CurrentBlendState);
-#ifdef USE_MONOLITHIC_GRAPHICS_DRIVERS
+#ifdef PLATFORM_XBOXONE
 		//MONO BUG WORKAROUND and a check to make sure it isn't fixed yet.
 		// mono's OMGetBlendState fails to set elements 2 and 3 (they don't change when you call OMGetBlendState)
 		check(FMemory::Memcmp(BlendFactor, CurrentBlendFactor, sizeof(float)*2) == 0);
@@ -232,7 +232,7 @@ void FD3D11StateCacheBase::VerifyCacheState()
 		for (uint32 Index = 0; Index < D3D11_IA_VERTEX_INPUT_RESOURCE_SLOT_COUNT; Index++)
 		{
 			check(VertexBuffers[Index] == CurrentVertexBuffers[Index].VertexBuffer);
-#ifndef USE_MONOLITHIC_GRAPHICS_DRIVERS
+#ifndef PLATFORM_XBOXONE
 			//Broken on mono.  will return stride = "previous stride" after call to IASetVertexBuffers(0,1,0,0,0);
 			check(Strides[Index] == CurrentVertexBuffers[Index].Stride);
 #endif
diff --git a/Engine/Source/Runtime/Windows/D3D11RHI/Private/D3D11Viewport.cpp b/Engine/Source/Runtime/Windows/D3D11RHI/Private/D3D11Viewport.cpp
--- a/Engine/Source/Runtime/Windows/D3D11RHI/Private/D3D11Viewport.cpp
+++ b/Engine/Source/Runtime/Windows/D3D11RHI/Private/D3D11Viewport.cpp
@@ -179,7 +179,7 @@ void FD3D11Viewport::Resize(uint32 InSizeX,uint32 InSizeY,bool bInIsFullscreen)
 		check(BackBuffer->GetRefCount() == 1);
 
 		// IUnknown is GONE from Mono Drivers, and is replaced with IGraphicsUnknown, making this incompatible (no COM)
-#ifndef USE_MONOLITHIC_GRAPHICS_DRIVERS
+#ifndef PLATFORM_XBOXONE
 		checkComRefCount(BackBuffer->GetResource(),1);
 		checkComRefCount(BackBuffer->GetRenderTargetView(0, -1),1);
 		checkComRefCount(BackBuffer->GetShaderResourceView(),1);
