Changed XcodeProjectTarget.Type in Xcode project generator from string to enum

diff --git a/Engine/Source/Programs/UnrealBuildTool/System/XcodeProjectFileGenerator.cs b/Engine/Source/Programs/UnrealBuildTool/System/XcodeProjectFileGenerator.cs
--- a/Engine/Source/Programs/UnrealBuildTool/System/XcodeProjectFileGenerator.cs
+++ b/Engine/Source/Programs/UnrealBuildTool/System/XcodeProjectFileGenerator.cs
@@ -33,6 +33,15 @@ namespace UnrealBuildTool
 		}
 	}
 
+	public enum XcodeTargetType
+	{
+		Project,
+		Legacy,
+		Native,
+		XCTest,
+		XcodeHelper
+	}
+
 	/// <summary>
 	/// Represents a build target
 	/// </summary>
@@ -47,7 +56,7 @@ namespace UnrealBuildTool
 		/// <param name="InTargetPlatform">Name of the target that may be built.</param>
 		/// <param name="InDependencies">Name of the target that may be built.</param>
 		/// <param name="bHasPlist">Name of the target that may be built.</param>
-		public XcodeProjectTarget(string InDisplayName, string InTargetName, string InType, string InProductName = "", UnrealTargetPlatform InTargetPlatform = UnrealTargetPlatform.Mac, bool bInIsMacOnly = false, List<XcodeTargetDependency> InDependencies = null, bool bHasPlist = false, List<XcodeFrameworkRef> InFrameworks = null)
+		public XcodeProjectTarget(string InDisplayName, string InTargetName, XcodeTargetType InType, string InProductName = "", UnrealTargetPlatform InTargetPlatform = UnrealTargetPlatform.Mac, bool bInIsMacOnly = false, List<XcodeTargetDependency> InDependencies = null, bool bHasPlist = false, List<XcodeFrameworkRef> InFrameworks = null)
 		{
 			DisplayName = InDisplayName;
 			TargetName = InTargetName;
@@ -78,7 +87,7 @@ namespace UnrealBuildTool
 
 		public string DisplayName;					// Name displayed in Xcode's UI.
 		public string TargetName;					// Actual name of the target.
-		public string Type;							// "Native", "Legacy" or "Project"
+		public XcodeTargetType Type;
 		public UnrealTargetPlatform TargetPlatform;	//Mac, IOS
 		public bool bIsMacOnly;
 		public string ProductName;
@@ -288,7 +297,7 @@ namespace UnrealBuildTool
 			// Required for iOS. This removes some "Target is not Built to Run on this Device" errors.
 			foreach (XcodeProjectTarget Target in Targets)
 			{
-				if (Target.Type == "Project" || Target == UE4XcodeHelperTarget || Target.ProductName == "")
+				if (Target.Type == XcodeTargetType.Project || Target == UE4XcodeHelperTarget || Target.ProductName == "")
 				{
 					continue;
 				}
@@ -323,6 +332,11 @@ namespace UnrealBuildTool
 			Contents.Append("/* End PBXGroup section */" + ProjectFileGenerator.NewLine + ProjectFileGenerator.NewLine);
 		}
 
+		private static bool IsXcodeTargetTypeNative(XcodeTargetType Type)
+		{
+			return Type == XcodeTargetType.Native || Type == XcodeTargetType.XcodeHelper || Type == XcodeTargetType.XCTest;
+		}
+
 		/// <summary>
 		/// Appends a target to targets section.
 		/// </summary>
@@ -330,12 +344,12 @@ namespace UnrealBuildTool
 		/// <param name="Target">Target to append</param>
 		private void AppendTarget(ref StringBuilder Contents, XcodeProjectTarget Target)
 		{
-            string XcodeTargetType = Target.Type == "XCTest" ? "Native" : Target.Type;
+			string TargetType = IsXcodeTargetTypeNative(Target.Type) ? "Native" : Target.Type.ToString();
 			Contents.Append(
 				"\t\t" + Target.Guid + " /* " + Target.DisplayName + " */ = {" + ProjectFileGenerator.NewLine +
-				"\t\t\tisa = PBX" + XcodeTargetType + "Target;" + ProjectFileGenerator.NewLine);
+				"\t\t\tisa = PBX" + TargetType + "Target;" + ProjectFileGenerator.NewLine);
 
-			if (Target.Type == "Legacy")
+			if (Target.Type == XcodeTargetType.Legacy)
 			{
 				string UProjectPath = "";
 				if (UnrealBuildTool.HasUProjectFile() && Target.TargetName.StartsWith(Path.GetFileNameWithoutExtension(UnrealBuildTool.GetUProjectFile())))
@@ -353,11 +367,11 @@ namespace UnrealBuildTool
 				Contents.Append("\t\t\tbuildArgumentsString = \"$(ACTION) " + Target.TargetName + " $(PLATFORM_NAME) $(CONFIGURATION)" + UProjectPath + "\";" + ProjectFileGenerator.NewLine);
 			}
 
-			Contents.Append("\t\t\tbuildConfigurationList = " + Target.BuildConfigGuild + " /* Build configuration list for PBX" + XcodeTargetType + "Target \"" + Target.DisplayName + "\" */;" + ProjectFileGenerator.NewLine);
+			Contents.Append("\t\t\tbuildConfigurationList = " + Target.BuildConfigGuild + " /* Build configuration list for PBX" + TargetType + "Target \"" + Target.DisplayName + "\" */;" + ProjectFileGenerator.NewLine);
 
 			Contents.Append("\t\t\tbuildPhases = (" + ProjectFileGenerator.NewLine);
 
-			if (Target.Type == "Native" || Target.Type == "XCTest")
+			if (IsXcodeTargetTypeNative(Target.Type))
 			{
 				Contents.Append("\t\t\t\t" + Target.SourcesPhaseGuid + " /* Sources */," + ProjectFileGenerator.NewLine);
 				//Contents.Append("\t\t\t\t" + Target.ResourcesPhaseGuid + " /* Resources */," + ProjectFileGenerator.NewLine);
@@ -367,7 +381,7 @@ namespace UnrealBuildTool
 
 			Contents.Append("\t\t\t);" + ProjectFileGenerator.NewLine);
 
-			if (Target.Type == "Legacy")
+			if (Target.Type == XcodeTargetType.Legacy)
 			{
 				string UE4Dir = Path.GetFullPath(Directory.GetCurrentDirectory() + "../../..");
 				if (bGeneratingRocketProjectFiles)
@@ -391,14 +405,14 @@ namespace UnrealBuildTool
 				"\t\t\t);" + ProjectFileGenerator.NewLine +
 				"\t\t\tname = \"" + Target.DisplayName + "\";" + ProjectFileGenerator.NewLine);
 
-			if (Target.Type == "Legacy")
+			if (Target.Type == XcodeTargetType.Legacy)
 			{
 				Contents.Append("\t\t\tpassBuildSettingsInEnvironment = 1;" + ProjectFileGenerator.NewLine);
 			}
 
 			Contents.Append("\t\t\tproductName = \"" + Target.DisplayName + "\";" + ProjectFileGenerator.NewLine);
 
-			if (Target.Type == "Native")
+			if (Target.Type == XcodeTargetType.XcodeHelper)
 			{
 				if (Target.DisplayName == "UE4XcodeHelper")
 				{
@@ -413,7 +427,7 @@ namespace UnrealBuildTool
 						"\t\t\tproductType = \"com.apple.product-type.application\";" + ProjectFileGenerator.NewLine);
 				}
 			}
-			if (Target.Type == "XCTest")
+			if (Target.Type == XcodeTargetType.XCTest)
 			{
 					Contents.Append(
 						"\t\t\tproductReference = " + Target.ProductGuid + " /* " + Target.ProductName + " */;" + ProjectFileGenerator.NewLine +
@@ -613,7 +627,7 @@ namespace UnrealBuildTool
 		private void AppendSingleConfig(ref StringBuilder Contents, XcodeProjectTarget Target, string ConfigName, string ConfigGuid, string PreprocessorDefinitions, string HeaderSearchPaths,
 			string EngineRelative, string GamePath, bool bIsUE4Game, bool IsAGame, bool bIsUE4Client)
 		{
-			if (Target.Type == "Project")
+			if (Target.Type == XcodeTargetType.Project)
 			{
 				AppendProjectConfig(ref Contents, ConfigName, ConfigGuid, PreprocessorDefinitions, HeaderSearchPaths);
 			}
@@ -625,13 +639,13 @@ namespace UnrealBuildTool
 				}
 				else
 				{
-					if (Target.Type == "Legacy")
+					if (Target.Type == XcodeTargetType.Legacy)
 					{
 						AppendIOSBuildConfig(ref Contents, ConfigName, ConfigGuid);
 					}
 					else
 					{
-						if (Target.Type != "XCTest")
+						if (Target.Type != XcodeTargetType.XCTest)
 						{
 							AppendIOSRunConfig(ref Contents, ConfigName, ConfigGuid, Target.TargetName, EngineRelative, GamePath, bIsUE4Game, IsAGame, bIsUE4Client);
 						}
@@ -710,8 +724,8 @@ namespace UnrealBuildTool
 		/// <param name="Target">Target for which we generate the build configuration list</param>
 		private void AppendConfigList(ref StringBuilder Contents, XcodeProjectTarget Target)
 		{
-			string XcodeTargetType = Target.Type == "XCTest" ? "Native" : Target.Type;
-			string TypeName = Target.Type == "Project" ? "PBXProject" : "PBX" + XcodeTargetType + "Target";
+			string TargetType = IsXcodeTargetTypeNative(Target.Type) ? "Native" : Target.Type.ToString();
+			string TypeName = Target.Type == XcodeTargetType.Project ? "PBXProject" : "PBX" + TargetType + "Target";
 
 			if (!bGeneratingRocketProjectFiles)
 			{
@@ -767,7 +781,7 @@ namespace UnrealBuildTool
 		{
 			foreach (XcodeProjectTarget Target in ProjectTargets)
 			{
-				if (Target.Type == "Native" || Target.Type == "XCTest")
+				if (IsXcodeTargetTypeNative(Target.Type))
 				{
 					// Generate Build references and Framework references for each Framework.
 					string FrameworkFiles = "";
@@ -788,7 +802,7 @@ namespace UnrealBuildTool
 
 					string Sources = null;
 
-					if (Target.Type == "XCTest")
+					if (Target.Type == XcodeTargetType.XCTest)
 					{
 						// Add the xctest framework.
 						FrameworkFiles += TestFrameworkFiles;
@@ -803,7 +817,7 @@ namespace UnrealBuildTool
 					AppendBuildPhase(ref PBXFrameworksBuildPhaseSection, Target.FrameworksPhaseGuid, "Frameworks", FrameworkFiles, null);
 
 					string PayloadDir = "Engine";
-					if (Target.Type == "Native" && Target.TargetName != "UE4Game")
+					if (Target.Type == XcodeTargetType.Native && Target.TargetName != "UE4Game")
 					{
 						PayloadDir = Target.TargetName;
 					}
@@ -1065,7 +1079,7 @@ namespace UnrealBuildTool
 					// @todo: Remove target platform param and merge Mac and iOS targets. For now BuildTarget knows how to build iOS, but cannot run iOS apps, so we need separate DeployTarget.
 					bool bIsMacOnly = !SupportedPlatforms.Contains(UnrealTargetPlatform.IOS);
 
-					XcodeProjectTarget BuildTarget = new XcodeProjectTarget(TargetName + " - Mac", TargetName, "Legacy", "", UnrealTargetPlatform.Mac, bIsMacOnly);
+					XcodeProjectTarget BuildTarget = new XcodeProjectTarget(TargetName + " - Mac", TargetName, XcodeTargetType.Legacy, "", UnrealTargetPlatform.Mac, bIsMacOnly);
 					if (!bGeneratingRunIOSProject)
 					{
 						ProjectTargets.Add(BuildTarget);
@@ -1082,14 +1096,14 @@ namespace UnrealBuildTool
 								FrameworkRefs.Add(new XcodeFrameworkRef(Framework));
 							}
 
-							XcodeProjectTarget IOSDeployTarget = new XcodeProjectTarget(TargetName + " - iOS", TargetName, "Native", TargetName + ".app", UnrealTargetPlatform.IOS, false, null, true, FrameworkRefs);
+							XcodeProjectTarget IOSDeployTarget = new XcodeProjectTarget(TargetName + " - iOS", TargetName, XcodeTargetType.Native, TargetName + ".app", UnrealTargetPlatform.IOS, false, null, true, FrameworkRefs);
 							ProjectTargets.Add(IOSDeployTarget);
 						}
 						else
 						{
 							XcodeContainerItemProxy ContainerProxy = new XcodeContainerItemProxy(ProjectTarget.Guid, BuildTarget.Guid, BuildTarget.DisplayName);
 							XcodeTargetDependency TargetDependency = new XcodeTargetDependency(BuildTarget.DisplayName, BuildTarget.Guid, ContainerProxy.Guid);
-							XcodeProjectTarget IOSDeployTarget = new XcodeProjectTarget(TargetName + " - iOS", TargetName, "Native", TargetName + ".app", UnrealTargetPlatform.IOS, false, new List<XcodeTargetDependency>() { TargetDependency }, true);
+							XcodeProjectTarget IOSDeployTarget = new XcodeProjectTarget(TargetName + " - iOS", TargetName, XcodeTargetType.Native, TargetName + ".app", UnrealTargetPlatform.IOS, false, new List<XcodeTargetDependency>() { TargetDependency }, true);
 							ProjectTargets.Add(IOSDeployTarget);
 							ContainerItemProxies.Add(ContainerProxy);
 							TargetDependencies.Add(TargetDependency);
@@ -1401,13 +1415,13 @@ namespace UnrealBuildTool
 			// attempt to determine targets for the project
 			List<XcodeProjectTarget> ProjectTargets = new List<XcodeProjectTarget> ();
 			// add mandatory ones
-			XcodeProjectTarget UE4ProjectTarget = new XcodeProjectTarget ("UE4", "UE4", "Project");
-			XcodeProjectTarget UE4XcodeHelperTarget = new XcodeProjectTarget ("UE4XcodeHelper", "UE4XcodeHelper", "Native", "libUE4XcodeHelper.a");
+			XcodeProjectTarget UE4ProjectTarget = new XcodeProjectTarget ("UE4", "UE4", XcodeTargetType.Project);
+			XcodeProjectTarget UE4XcodeHelperTarget = new XcodeProjectTarget ("UE4XcodeHelper", "UE4XcodeHelper", XcodeTargetType.XcodeHelper, "libUE4XcodeHelper.a");
 			ProjectTargets.AddRange(new XcodeProjectTarget[] { UE4ProjectTarget, UE4XcodeHelperTarget });
 
 			if (ProjectFilePlatform.HasFlag(XcodeProjectFilePlatform.iOS))
 			{
-				XcodeProjectTarget UE4CmdLineRunTarget = new XcodeProjectTarget("UE4CmdLineRun", "UE4CmdLineRun", "XCTest", "UE4CmdLineRun.xctest", UnrealTargetPlatform.IOS);
+				XcodeProjectTarget UE4CmdLineRunTarget = new XcodeProjectTarget("UE4CmdLineRun", "UE4CmdLineRun", XcodeTargetType.XCTest, "UE4CmdLineRun.xctest", UnrealTargetPlatform.IOS);
 				ProjectTargets.AddRange(new XcodeProjectTarget[] { UE4ProjectTarget, UE4XcodeHelperTarget, UE4CmdLineRunTarget });
 				// This GUID will be referenced by each app's test action.
 				UE4CmdLineGuid = UE4CmdLineRunTarget.Guid;
@@ -1467,11 +1481,11 @@ namespace UnrealBuildTool
 
 			foreach (XcodeProjectTarget Target in ProjectTargets)
 			{
-				if ((Target.Type == "Native" || Target.Type == "XCTest") && Target != UE4XcodeHelperTarget)
+				if ((Target.Type == XcodeTargetType.Native || Target.Type == XcodeTargetType.XCTest))
 				{
-					string FileType = Target.Type == "Native" ? "wrapper.application" : "wrapper.cfbundle";
+					string FileType = Target.Type == XcodeTargetType.Native ? "wrapper.application" : "wrapper.cfbundle";
 					string PayloadDir = "Engine";
-					if (Target.Type == "Native" && Target.TargetName != "UE4Game")
+					if (Target.Type == XcodeTargetType.Native && Target.TargetName != "UE4Game")
 					{
 						PayloadDir = Target.TargetName;
 					}
@@ -1517,11 +1531,11 @@ namespace UnrealBuildTool
 			XcodeProjectFileContent.Append(PBXFrameworksBuildPhaseSection);
 
 			XcodeProjectFileContent.Append("/* Begin PBXLegacyTarget section */" + ProjectFileGenerator.NewLine);
-			foreach (var target in ProjectTargets)
+			foreach (var Target in ProjectTargets)
 			{
-				if (target.Type == "Legacy")
+				if (Target.Type == XcodeTargetType.Legacy)
 				{
-					AppendTarget (ref XcodeProjectFileContent, target);
+					AppendTarget (ref XcodeProjectFileContent, Target);
 				}
 			}
 			XcodeProjectFileContent.Append("/* End PBXLegacyTarget section */" + ProjectFileGenerator.NewLine + ProjectFileGenerator.NewLine);
@@ -1531,7 +1545,7 @@ namespace UnrealBuildTool
 
 			foreach (XcodeProjectTarget Target in ProjectTargets)
 			{
-				if ((Target.Type == "Native" || Target.Type == "XCTest") && Target != UE4XcodeHelperTarget)
+				if ((Target.Type == XcodeTargetType.Native || Target.Type == XcodeTargetType.XCTest) && Target != UE4XcodeHelperTarget)
 				{
 					AppendTarget(ref XcodeProjectFileContent, Target);
 				}
