Fix for product type setting in Xcode project generator

diff --git a/Engine/Source/Programs/UnrealBuildTool/System/XcodeProjectFileGenerator.cs b/Engine/Source/Programs/UnrealBuildTool/System/XcodeProjectFileGenerator.cs
--- a/Engine/Source/Programs/UnrealBuildTool/System/XcodeProjectFileGenerator.cs
+++ b/Engine/Source/Programs/UnrealBuildTool/System/XcodeProjectFileGenerator.cs
@@ -414,24 +414,21 @@ namespace UnrealBuildTool
 
 			if (Target.Type == XcodeTargetType.XcodeHelper)
 			{
-				if (Target.DisplayName == "UE4XcodeHelper")
-				{
-					Contents.Append(
-						"\t\t\tproductReference = " + Target.ProductGuid + " /* " + Target.ProductName + " */;" + ProjectFileGenerator.NewLine +
-						"\t\t\tproductType = \"com.apple.product-type.library.static\";" + ProjectFileGenerator.NewLine);
-				}
-				else
-				{
-					Contents.Append(
-						"\t\t\tproductReference = " + Target.ProductGuid + " /* " + Target.ProductName + " */;" + ProjectFileGenerator.NewLine +
-						"\t\t\tproductType = \"com.apple.product-type.application\";" + ProjectFileGenerator.NewLine);
-				}
+				Contents.Append(
+					"\t\t\tproductReference = " + Target.ProductGuid + " /* " + Target.ProductName + " */;" + ProjectFileGenerator.NewLine +
+					"\t\t\tproductType = \"com.apple.product-type.library.static\";" + ProjectFileGenerator.NewLine);
 			}
-			if (Target.Type == XcodeTargetType.XCTest)
+			else if (Target.Type == XcodeTargetType.Native)
 			{
-					Contents.Append(
-						"\t\t\tproductReference = " + Target.ProductGuid + " /* " + Target.ProductName + " */;" + ProjectFileGenerator.NewLine +
-						"\t\t\tproductType = \"com.apple.product-type.bundle.unit-test\";" + ProjectFileGenerator.NewLine);
+				Contents.Append(
+					"\t\t\tproductReference = " + Target.ProductGuid + " /* " + Target.ProductName + " */;" + ProjectFileGenerator.NewLine +
+					"\t\t\tproductType = \"com.apple.product-type.application\";" + ProjectFileGenerator.NewLine);
+			}
+			else if (Target.Type == XcodeTargetType.XCTest)
+			{
+				Contents.Append(
+					"\t\t\tproductReference = " + Target.ProductGuid + " /* " + Target.ProductName + " */;" + ProjectFileGenerator.NewLine +
+					"\t\t\tproductType = \"com.apple.product-type.bundle.unit-test\";" + ProjectFileGenerator.NewLine);
 			}
 			Contents.Append("\t\t};" + ProjectFileGenerator.NewLine);
 		}
