Hide command window popups when switching SDKS.

diff --git a/Engine/Source/Programs/UnrealBuildTool/Configuration/UEBuildPlatform.cs b/Engine/Source/Programs/UnrealBuildTool/Configuration/UEBuildPlatform.cs
--- a/Engine/Source/Programs/UnrealBuildTool/Configuration/UEBuildPlatform.cs
+++ b/Engine/Source/Programs/UnrealBuildTool/Configuration/UEBuildPlatform.cs
@@ -394,6 +394,7 @@ namespace UnrealBuildTool
                     HookProcess.StartInfo.WorkingDirectory = SDKDirectory;
                     HookProcess.StartInfo.FileName = HookExe;
                     HookProcess.StartInfo.Arguments = "";
+                    HookProcess.StartInfo.WindowStyle = ProcessWindowStyle.Hidden;
 
                     //installers may require administrator access to succeed. so run as an admmin.
 					HookProcess.StartInfo.Verb = "runas";
