Small cleanup of UE4XcodeHelper use in Xcode project generator

diff --git a/Engine/Source/Programs/UnrealBuildTool/System/XcodeProjectFileGenerator.cs b/Engine/Source/Programs/UnrealBuildTool/System/XcodeProjectFileGenerator.cs
--- a/Engine/Source/Programs/UnrealBuildTool/System/XcodeProjectFileGenerator.cs
+++ b/Engine/Source/Programs/UnrealBuildTool/System/XcodeProjectFileGenerator.cs
@@ -247,9 +247,8 @@ namespace UnrealBuildTool
 		/// </summary>
 		/// <param name="Contents">StringBuilder object to append groups string to</param>
 		/// <param name="Groups">Dictionary of all project groups</param>
-		/// <param name="UE4XcodeHelperTarget">Dummy target required for syntax highlighting etc. work</param>
         /// <param name="Frameworks">Frameworks referenced</param>
-		private void AppendGroups(ref StringBuilder Contents, ref Dictionary<string, XcodeFileGroup> Groups, List<XcodeProjectTarget> Targets, XcodeProjectTarget UE4XcodeHelperTarget, List<XcodeFramework> Frameworks)
+		private void AppendGroups(ref StringBuilder Contents, ref Dictionary<string, XcodeFileGroup> Groups, List<XcodeProjectTarget> Targets, List<XcodeFramework> Frameworks)
 		{
 			Contents.Append("/* Begin PBXGroup section */" + ProjectFileGenerator.NewLine);
 
@@ -292,16 +291,13 @@ namespace UnrealBuildTool
 			Contents.Append(string.Format("\t\t{0} = {{{1}", ProductRefGroupGuid, ProjectFileGenerator.NewLine));
 			Contents.Append("\t\t\tisa = PBXGroup;" + ProjectFileGenerator.NewLine);
 			Contents.Append("\t\t\tchildren = (" + ProjectFileGenerator.NewLine);
-			Contents.Append("\t\t\t\t" + UE4XcodeHelperTarget.ProductGuid + " /* " + UE4XcodeHelperTarget.ProductName + " */," + ProjectFileGenerator.NewLine);
 
-			// Required for iOS. This removes some "Target is not Built to Run on this Device" errors.
 			foreach (XcodeProjectTarget Target in Targets)
 			{
-				if (Target.Type == XcodeTargetType.Project || Target == UE4XcodeHelperTarget || Target.ProductName == "")
+				if (IsXcodeTargetTypeNative(Target.Type))
 				{
-					continue;
+					Contents.Append("\t\t\t\t" + Target.ProductGuid + " /* " + Target.ProductName + " */," + ProjectFileGenerator.NewLine);
 				}
-				Contents.Append("\t\t\t\t" + Target.ProductGuid + " /* " + Target.ProductName + " */," + ProjectFileGenerator.NewLine);
 			}
 			Contents.Append("\t\t\t);" + ProjectFileGenerator.NewLine);
 			Contents.Append("\t\t\tname = Products;" + ProjectFileGenerator.NewLine);
@@ -1291,7 +1287,7 @@ namespace UnrealBuildTool
 			File.WriteAllText(SchemeFilePath, SchemeFileContent);
 		}
 
-		private void WriteSchemes(string XcodeProjectPath, List<XcodeProjectTarget> ProjectTargets, XcodeProjectTarget UE4ProjectTarget, XcodeProjectTarget UE4XcodeHelperTarget)
+		private void WriteSchemes(string XcodeProjectPath, List<XcodeProjectTarget> ProjectTargets)
 		{
 			// Delete obsolete schemes
 			string SchemesDir = XcodeProjectPath + "/xcuserdata/" + Environment.UserName + ".xcuserdatad/xcschemes";
@@ -1307,7 +1303,7 @@ namespace UnrealBuildTool
 			// write scheme files for targets
 			foreach (var Target in ProjectTargets)
 			{
-				if (Target == UE4ProjectTarget || Target == UE4XcodeHelperTarget)
+				if (Target.Type == XcodeTargetType.Project || Target.Type == XcodeTargetType.XcodeHelper)
 				{
 					continue;
 				}
@@ -1352,7 +1348,7 @@ namespace UnrealBuildTool
 			}
 		}
 
-		private void WriteProject(StringBuilder XcodeProjectFileContent, List<XcodeProjectTarget> ProjectTargets, XcodeProjectTarget UE4ProjectTarget, XcodeProjectTarget UE4XcodeHelperTarget)
+		private void WriteProject(StringBuilder XcodeProjectFileContent, List<XcodeProjectTarget> ProjectTargets)
 		{
 			// @todo: Windows and Mac should store the project in the same folder
 			string PathBranch = (ExternalExecution.GetRuntimePlatform() != UnrealTargetPlatform.Mac) ? "/Engine/Intermediate/IOS" : "";
@@ -1386,7 +1382,7 @@ namespace UnrealBuildTool
 				Log.TraceInformation("Skipping project file write, as it didn't change...");
 			}
 
-			WriteSchemes(XcodeProjectPath, ProjectTargets, UE4ProjectTarget, UE4XcodeHelperTarget);
+			WriteSchemes(XcodeProjectPath, ProjectTargets);
 		}
 
 		/// <summary>
@@ -1419,7 +1415,7 @@ namespace UnrealBuildTool
 			if (ProjectFilePlatform.HasFlag(XcodeProjectFilePlatform.iOS))
 			{
 				XcodeProjectTarget UE4CmdLineRunTarget = new XcodeProjectTarget("UE4CmdLineRun", "UE4CmdLineRun", XcodeTargetType.XCTest, "UE4CmdLineRun.xctest", UnrealTargetPlatform.IOS);
-				ProjectTargets.AddRange(new XcodeProjectTarget[] { UE4ProjectTarget, UE4XcodeHelperTarget, UE4CmdLineRunTarget });
+				ProjectTargets.Add(UE4CmdLineRunTarget);
 				// This GUID will be referenced by each app's test action.
 				UE4CmdLineGuid = UE4CmdLineRunTarget.Guid;
 			}
@@ -1521,7 +1517,7 @@ namespace UnrealBuildTool
 			XcodeProjectFileContent.Append(PBXContainerItemProxySection);
 			XcodeProjectFileContent.Append(PBXTargetDependencySection);
 
-			AppendGroups(ref XcodeProjectFileContent, ref Groups, ProjectTargets, UE4XcodeHelperTarget, Frameworks);
+			AppendGroups(ref XcodeProjectFileContent, ref Groups, ProjectTargets , Frameworks);
 
 			XcodeProjectFileContent.Append(PBXShellScriptBuildPhaseSection);
 			XcodeProjectFileContent.Append(PBXSourcesBuildPhaseSection);
@@ -1542,7 +1538,7 @@ namespace UnrealBuildTool
 
 			foreach (XcodeProjectTarget Target in ProjectTargets)
 			{
-				if ((Target.Type == XcodeTargetType.Native || Target.Type == XcodeTargetType.XCTest) && Target != UE4XcodeHelperTarget)
+				if (Target.Type == XcodeTargetType.Native || Target.Type == XcodeTargetType.XCTest)
 				{
 					AppendTarget(ref XcodeProjectFileContent, Target);
 				}
@@ -1570,17 +1566,15 @@ namespace UnrealBuildTool
 				"\t\t\tprojectRoot = \"\";" + ProjectFileGenerator.NewLine +
 				"\t\t\ttargets = (" + ProjectFileGenerator.NewLine
 			);
-			foreach (var target in ProjectTargets)
+
+			foreach (var Target in ProjectTargets)
 			{
-				if (target == UE4ProjectTarget || target == UE4XcodeHelperTarget)
+				if (Target != UE4ProjectTarget)
 				{
-					continue;
+					XcodeProjectFileContent.AppendLine(string.Format ("\t\t\t\t{0} /* {1} */,", Target.Guid, Target.DisplayName));
 				}
-				XcodeProjectFileContent.AppendLine(string.Format ("\t\t\t\t{0} /* {1} */,", target.Guid, target.DisplayName));
 			}
 
-			XcodeProjectFileContent.Append("\t\t\t\t" + UE4XcodeHelperTarget.Guid + " /* " + UE4XcodeHelperTarget.DisplayName + " */" + ProjectFileGenerator.NewLine);
-
 			XcodeProjectFileContent.Append(
 				"\t\t\t);" + ProjectFileGenerator.NewLine +
 				"\t\t};" + ProjectFileGenerator.NewLine +
@@ -1612,7 +1606,7 @@ namespace UnrealBuildTool
 			AppendBuildConfigs(ref XcodeProjectFileContent, UE4ProjectTarget, HeaderSearchPaths, PreprocessorDefinitionsString);
 			foreach (var Target in ProjectTargets)
 			{
-				if (Target != UE4ProjectTarget && Target != UE4XcodeHelperTarget)
+				if (Target.Type != XcodeTargetType.Project && Target.Type != XcodeTargetType.XcodeHelper)
 				{
 					AppendBuildConfigs(ref XcodeProjectFileContent, Target);
 				}
@@ -1626,7 +1620,7 @@ namespace UnrealBuildTool
 
 			foreach (var Target in ProjectTargets)
 			{
-				if (Target != UE4XcodeHelperTarget)
+				if (Target.Type != XcodeTargetType.XcodeHelper)
 				{
 					AppendConfigList(ref XcodeProjectFileContent, Target);
 				}
@@ -1641,7 +1635,7 @@ namespace UnrealBuildTool
 				"\trootObject = " + UE4ProjectTarget.Guid + " /* Project object */;" + ProjectFileGenerator.NewLine +
 				"}" + ProjectFileGenerator.NewLine);
 
-			WriteProject(XcodeProjectFileContent, ProjectTargets, UE4ProjectTarget, UE4XcodeHelperTarget);
+			WriteProject(XcodeProjectFileContent, ProjectTargets);
 
 			return true;
 		}
