Fixed engine path in server builds

diff --git a/Engine/Source/Programs/UnrealBuildTool/Configuration/UEBuildTarget.cs b/Engine/Source/Programs/UnrealBuildTool/Configuration/UEBuildTarget.cs
--- a/Engine/Source/Programs/UnrealBuildTool/Configuration/UEBuildTarget.cs
+++ b/Engine/Source/Programs/UnrealBuildTool/Configuration/UEBuildTarget.cs
@@ -1301,7 +1301,8 @@ namespace UnrealBuildTool
 
                 }
 
-                if (TargetRules.IsAGame(Rules.Type) && IsCurrentPlatform)
+				if ( (TargetRules.IsAGame(Rules.Type) || (Rules.Type == TargetRules.TargetType.Server)) 
+					&& IsCurrentPlatform)
 				{
 					// The hardcoded engine directory needs to be a relative path to match the normal EngineDir format. Not doing so breaks the network file system (TTP#315861).
 					string OutputFilePath = ExecutableBinary.Config.OutputFilePath;
