UBT: Fixing DebugGame pch and plugin intermediate paths

diff --git a/Engine/Source/Programs/UnrealBuildTool/Configuration/UEBuildTarget.cs b/Engine/Source/Programs/UnrealBuildTool/Configuration/UEBuildTarget.cs
--- a/Engine/Source/Programs/UnrealBuildTool/Configuration/UEBuildTarget.cs
+++ b/Engine/Source/Programs/UnrealBuildTool/Configuration/UEBuildTarget.cs
@@ -1943,8 +1943,9 @@ namespace UnrealBuildTool
 						if (Plugin.LoadedFrom == PluginInfo.LoadedFromType.Engine)
 						{
 							// Plugin folder is in the engine directory
+							var PluginConfiguration = Configuration == UnrealTargetConfiguration.DebugGame ? UnrealTargetConfiguration.Development : Configuration;
 							Plugin.IntermediateIncPath   = Path.GetFullPath(Path.Combine(BuildConfiguration.RelativeEnginePath, BuildConfiguration.PlatformIntermediateFolder));
-							Plugin.IntermediateBuildPath = Path.GetFullPath(Path.Combine(BuildConfiguration.RelativeEnginePath, BuildConfiguration.PlatformIntermediateFolder, AppName, Configuration.ToString()));
+							Plugin.IntermediateBuildPath = Path.GetFullPath(Path.Combine(BuildConfiguration.RelativeEnginePath, BuildConfiguration.PlatformIntermediateFolder, AppName, PluginConfiguration.ToString()));
 						}
 						else
 						{
@@ -2145,7 +2146,8 @@ namespace UnrealBuildTool
 			}
 			else
 			{
-				GlobalCompileEnvironment.Config.OutputDirectory = Path.Combine(BuildConfiguration.PlatformIntermediatePath, OutputAppName, Configuration.ToString());
+				var GlobalEnvironmentConfiguration = Configuration == UnrealTargetConfiguration.DebugGame ? UnrealTargetConfiguration.Development : Configuration;
+				GlobalCompileEnvironment.Config.OutputDirectory = Path.Combine(BuildConfiguration.PlatformIntermediatePath, OutputAppName, GlobalEnvironmentConfiguration.ToString());
 			}
 
 			// By default, shadow source files for this target in the root OutputDirectory
