[GitHub] 66 : Start plumbing for advanced build handling
#ttp 333302

diff --git a/Engine/Source/Programs/UnrealBuildTool/Configuration/BuildConfiguration.cs b/Engine/Source/Programs/UnrealBuildTool/Configuration/BuildConfiguration.cs
--- a/Engine/Source/Programs/UnrealBuildTool/Configuration/BuildConfiguration.cs
+++ b/Engine/Source/Programs/UnrealBuildTool/Configuration/BuildConfiguration.cs
@@ -409,7 +409,7 @@ namespace UnrealBuildTool
 		/// <param name="bCreateDebugInfo">True if debug info should be created</param>
 		public static void ValidateConfiguration(CPPTargetConfiguration Configuration, CPPTargetPlatform Platform, bool bCreateDebugInfo)
 		{
-			UEBuildPlatform BuildPlatform = UEBuildPlatform.GetBuildPlatformForCPPTargetPlatform(Platform);
+			var BuildPlatform = UEBuildPlatform.GetBuildPlatformForCPPTargetPlatform(Platform);
 
 			// E&C support.
 			if( bSupportEditAndContinue )
diff --git a/Engine/Source/Programs/UnrealBuildTool/Configuration/UEBuildDeploy.cs b/Engine/Source/Programs/UnrealBuildTool/Configuration/UEBuildDeploy.cs
--- a/Engine/Source/Programs/UnrealBuildTool/Configuration/UEBuildDeploy.cs
+++ b/Engine/Source/Programs/UnrealBuildTool/Configuration/UEBuildDeploy.cs
@@ -8,12 +8,34 @@ using System.Diagnostics;
 
 namespace UnrealBuildTool
 {
+	public interface IUEBuildDeploy
+	{
+		/**
+		 * Register the platform with the UEBuildDeploy class.
+		 */
+		void RegisterBuildDeploy();
+
+		/**
+		 * Prepare the target for deployment.
+		 *	
+		 * @param InTarget The target for deployment.
+		 *	
+		 * @return bool True if successful, false if not.
+		 */
+		bool PrepTargetForDeployment(UEBuildTarget InTarget);
+
+		/**
+		 * Prepare the target for deployment.
+		 */
+		bool PrepForUATPackageOrDeploy(string ProjectName, string ProjectDirectory, string ExecutablePath, string EngineDirectory, bool bForDistribution);
+	}
+
 	/// <summary>
 	///  Base class to handle deploy of a target for a given platform
 	/// </summary>
-	public abstract class UEBuildDeploy
+	public abstract class UEBuildDeploy : IUEBuildDeploy
 	{
-		static Dictionary<UnrealTargetPlatform, UEBuildDeploy> BuildDeployDictionary = new Dictionary<UnrealTargetPlatform, UEBuildDeploy>();
+		static Dictionary<UnrealTargetPlatform, IUEBuildDeploy> BuildDeployDictionary = new Dictionary<UnrealTargetPlatform, IUEBuildDeploy>();
 
 		/**
 		 *	Register the given platforms UEBuildDeploy instance
@@ -21,7 +43,7 @@ namespace UnrealBuildTool
 		 *	@param	InPlatform			The UnrealTargetPlatform to register with
 		 *	@param	InBuildDeploy		The UEBuildDeploy instance to use for the InPlatform
 		 */
-		public static void RegisterBuildDeploy(UnrealTargetPlatform InPlatform, UEBuildDeploy InBuildDeploy)
+		public static void RegisterBuildDeploy(UnrealTargetPlatform InPlatform, IUEBuildDeploy InBuildDeploy)
 		{
 			if (BuildDeployDictionary.ContainsKey(InPlatform) == true)
 			{
@@ -42,7 +64,7 @@ namespace UnrealBuildTool
 		 *	
 		 *	@return	UEBuildDeploy		The instance of the build deploy
 		 */
-		public static UEBuildDeploy GetBuildDeploy(UnrealTargetPlatform InPlatform)
+		public static IUEBuildDeploy GetBuildDeploy(UnrealTargetPlatform InPlatform)
 		{
 			if (BuildDeployDictionary.ContainsKey(InPlatform) == true)
 			{
diff --git a/Engine/Source/Programs/UnrealBuildTool/Configuration/UEBuildModule.cs b/Engine/Source/Programs/UnrealBuildTool/Configuration/UEBuildModule.cs
--- a/Engine/Source/Programs/UnrealBuildTool/Configuration/UEBuildModule.cs
+++ b/Engine/Source/Programs/UnrealBuildTool/Configuration/UEBuildModule.cs
@@ -921,7 +921,7 @@ namespace UnrealBuildTool
 		// UEBuildModule interface.
 		public override List<FileItem> Compile(CPPEnvironment GlobalCompileEnvironment, CPPEnvironment CompileEnvironment, bool bCompileMonolithic)
 		{
-			UEBuildPlatform BuildPlatform = UEBuildPlatform.GetBuildPlatformForCPPTargetPlatform(CompileEnvironment.Config.TargetPlatform);
+			var BuildPlatform = UEBuildPlatform.GetBuildPlatformForCPPTargetPlatform(CompileEnvironment.Config.TargetPlatform);
 
 			var LinkInputFiles = new List<FileItem>();
 			if( ProjectFileGenerator.bGenerateProjectFiles && IntelliSenseGatherer == null)
diff --git a/Engine/Source/Programs/UnrealBuildTool/Configuration/UEBuildPlatform.cs b/Engine/Source/Programs/UnrealBuildTool/Configuration/UEBuildPlatform.cs
--- a/Engine/Source/Programs/UnrealBuildTool/Configuration/UEBuildPlatform.cs
+++ b/Engine/Source/Programs/UnrealBuildTool/Configuration/UEBuildPlatform.cs
@@ -8,9 +8,235 @@ using System.IO;
 
 namespace UnrealBuildTool
 {
-	public abstract class UEBuildPlatform
+	public enum SDKStatus
 	{
-		public static Dictionary<UnrealTargetPlatform, UEBuildPlatform> BuildPlatformDictionary = new Dictionary<UnrealTargetPlatform, UEBuildPlatform>();
+		Valid,			// Desired SDK is installed and set up.
+		Invalid,		// Could not find the desired SDK, SDK setup failed, etc.
+		Unknown,		// Not enough info to determine SDK state. 
+	};
+
+	public interface IUEBuildPlatform
+	{
+		/**
+		 * Whether the required external SDKs are installed for this platform
+		 */
+		SDKStatus HasRequiredSDKsInstalled();
+		
+		/**
+		 * If this platform can be compiled with XGE
+		 */
+		bool CanUseXGE();
+		
+		/**
+		 * Register the platform with the UEBuildPlatform class
+		 */
+		void RegisterBuildPlatform();
+		
+		/**
+		 * Retrieve the CPPTargetPlatform for the given UnrealTargetPlatform
+		 *
+		 * @param InUnrealTargetPlatform The UnrealTargetPlatform being build
+		 * 
+		 * @return CPPTargetPlatform The CPPTargetPlatform to compile for
+		 */
+		CPPTargetPlatform GetCPPTargetPlatform(UnrealTargetPlatform InUnrealTargetPlatform);
+		
+		/**
+		 * Get the extension to use for the given binary type
+		 * 
+		 * @param InBinaryType The binary type being built
+		 * 
+		 * @return string The binary extension (i.e. 'exe' or 'dll')
+		 */
+		string GetBinaryExtension(UEBuildBinaryType InBinaryType);
+		
+		/**
+		 * Get the extension to use for debug info for the given binary type
+		 * 
+		 * @param InBinaryType The binary type being built
+		 * 
+		 * @return string The debug info extension (i.e. 'pdb')
+		 */
+		string GetDebugInfoExtension( UEBuildBinaryType InBinaryType );
+		
+		/**
+		 * Whether incremental linking should be used
+		 * 
+		 * @param InPlatform The CPPTargetPlatform being built
+		 * @param InConfiguration The CPPTargetConfiguration being built
+		 * 
+		 * @return bool true if incremental linking should be used, false if not
+		 */
+		bool ShouldUseIncrementalLinking(CPPTargetPlatform Platform, CPPTargetConfiguration Configuration);
+		
+		/**
+		 * Whether PDB files should be used
+		 * 
+		 * @param InPlatform The CPPTargetPlatform being built
+		 * @param InConfiguration The CPPTargetConfiguration being built
+		 * @param bInCreateDebugInfo true if debug info is getting create, false if not
+		 * 
+		 * @return bool true if PDB files should be used, false if not
+		 */
+		bool ShouldUsePDBFiles(CPPTargetPlatform Platform, CPPTargetConfiguration Configuration, bool bCreateDebugInfo);
+		
+		/**
+		 * Whether PCH files should be used
+		 * 
+		 * @param InPlatform The CPPTargetPlatform being built
+		 * @param InConfiguration The CPPTargetConfiguration being built
+		 * 
+		 * @return bool true if PCH files should be used, false if not
+		 */
+		bool ShouldUsePCHFiles(CPPTargetPlatform Platform, CPPTargetConfiguration Configuration);
+		
+		/**
+		 * Whether the editor should be built for this platform or not
+		 * 
+		 * @param InPlatform The UnrealTargetPlatform being built
+		 * @param InConfiguration The UnrealTargetConfiguration being built
+		 * @return bool true if the editor should be built, false if not
+		 */
+		bool ShouldNotBuildEditor(UnrealTargetPlatform InPlatform, UnrealTargetConfiguration InConfiguration);
+		
+		/**
+		 * Whether this build should support ONLY cooked data or not
+		 * 
+		 * @param InPlatform The UnrealTargetPlatform being built
+		 * @param InConfiguration The UnrealTargetConfiguration being built
+		 * @return bool true if the editor should be built, false if not
+		 */
+		bool BuildRequiresCookedData(UnrealTargetPlatform InPlatform, UnrealTargetConfiguration InConfiguration);
+		
+		/**
+		 * Whether the platform requires the extra UnityCPPWriter
+		 * This is used to add an extra file for UBT to get the #include dependencies from
+		 * 
+		 * @return bool true if it is required, false if not
+		 */
+		bool RequiresExtraUnityCPPWriter();
+		
+		/**
+		 * Whether this platform should build a monolithic binary
+		 */
+		bool ShouldCompileMonolithicBinary(UnrealTargetPlatform InPlatform);
+		
+		/**
+		 * Get a list of extra modules the platform requires.
+		 * This is to allow undisclosed platforms to add modules they need without exposing information about the platform.
+		 * 
+		 * @param Target The target being build
+		 * @param BuildTarget The UEBuildTarget getting build
+		 * @param PlatformExtraModules OUTPUT the list of extra modules the platform needs to add to the target
+		 */
+		void GetExtraModules(TargetInfo Target, UEBuildTarget BuildTarget, ref List<string> PlatformExtraModules);
+		
+		/**
+		 * Modify the newly created module passed in for this platform.
+		 * This is not required - but allows for hiding details of a
+		 * particular platform.
+		 * 
+		 * @param InModule The newly loaded module
+		 * @param Target The target being build
+		 */
+		void ModifyNewlyLoadedModule(UEBuildModule InModule, TargetInfo Target);
+		
+		/**
+		 * Setup the target environment for building
+		 * 
+		 * @param InBuildTarget The target being built
+		 */
+		void SetUpEnvironment(UEBuildTarget InBuildTarget);
+		
+		/**
+		 * Allow the platform to set an optional architecture
+		 */
+		string GetActiveArchitecture();
+		
+		/**
+		 * Setup the configuration environment for building
+		 * 
+		 * @param InBuildTarget The target being built
+		 */
+		void SetUpConfigurationEnvironment(UEBuildTarget InBuildTarget);
+		
+		/**
+		 * Whether this platform should create debug information or not
+		 * 
+		 * @param InPlatform The UnrealTargetPlatform being built
+		 * @param InConfiguration The UnrealTargetConfiguration being built
+		 * 
+		 * @return bool true if debug info should be generated, false if not
+		 */
+		bool ShouldCreateDebugInfo(UnrealTargetPlatform Platform, UnrealTargetConfiguration Configuration);
+		
+		/**
+		 * Gives the platform a chance to 'override' the configuration settings 
+		 * that are overridden on calls to RunUBT.
+		 * 
+		 * @param InPlatform The UnrealTargetPlatform being built
+		 * @param InConfiguration The UnrealTargetConfiguration being built
+		 * 
+		 * @return bool true if debug info should be generated, false if not
+		 */
+		void ResetBuildConfiguration(UnrealTargetPlatform InPlatform, UnrealTargetConfiguration InConfiguration);
+		
+		/**
+		 * Validate the UEBuildConfiguration for this platform
+		 * This is called BEFORE calling UEBuildConfiguration to allow setting 
+		 * various fields used in that function such as CompileLeanAndMean...
+		 */
+		void ValidateUEBuildConfiguration();
+		
+		/**
+		 * Validate configuration for this platform
+		 * NOTE: This function can/will modify BuildConfiguration!
+		 * 
+		 * @param InPlatform The CPPTargetPlatform being built
+		 * @param InConfiguration The CPPTargetConfiguration being built
+		 * @param bInCreateDebugInfo true if debug info is getting create, false if not
+		 */
+		void ValidateBuildConfiguration(CPPTargetConfiguration Configuration, CPPTargetPlatform Platform, bool bCreateDebugInfo);
+		
+		/**
+		 * Return whether this platform has uniquely named binaries across multiple games
+		 */
+		bool HasUniqueBinaries();
+		
+		/**
+		 * Return whether we wish to have this platform's binaries in our builds
+		 */
+		bool IsBuildRequired();
+		
+		/**
+		 * Return whether we wish to have this platform's binaries in our CIS tests
+		 */
+		bool IsCISRequired();
+		
+		/// <summary>
+		/// Whether the build platform requires deployment prep
+		/// </summary>
+		/// <returns></returns>
+		bool RequiresDeployPrepAfterCompile();
+		
+		/**
+		 * Return all valid configurations for this platform
+		 * 
+		 *  Typically, this is always Debug, Development, and Shipping - but Test is a likely future addition for some platforms
+		 */
+		List<UnrealTargetConfiguration> GetConfigurations(UnrealTargetPlatform InUnrealTargetPlatform, bool bIncludeDebug);
+		
+		/**
+		 * Setup the binaries for this specific platform.
+		 * 
+		 * @param InBuildTarget The target being built
+		 */
+		void SetupBinaries(UEBuildTarget InBuildTarget);
+	}
+
+	public abstract class UEBuildPlatform : IUEBuildPlatform
+	{
+		public static Dictionary<UnrealTargetPlatform, IUEBuildPlatform> BuildPlatformDictionary = new Dictionary<UnrealTargetPlatform, IUEBuildPlatform>();
 
 		// a mapping of a group to the platforms in the group (ie, Microsoft contains Win32 and Win64)
 		static Dictionary<UnrealPlatformGroup, List<UnrealTargetPlatform>> PlatformGroupDictionary = new Dictionary<UnrealPlatformGroup, List<UnrealTargetPlatform>>();
@@ -31,13 +257,6 @@ namespace UnrealBuildTool
             Uninstall
         };
 
-		public enum SDKStatus
-		{
-			Valid,			// Desired SDK is installed and set up.
-			Invalid,		// Could not find the desired SDK, SDK setup failed, etc.
-			Unknown,		// Not enough info to determine SDK state. 
-		};
-
 		public static readonly string SDKRootEnvVar = "UE_SDKS_ROOT";
 
         /**
@@ -118,14 +337,14 @@ namespace UnrealBuildTool
 		}
 
 		/**
-		 *	Retrieve the UEBuildPlatform instance for the given TargetPlatform
+		 *	Retrieve the IUEBuildPlatform instance for the given TargetPlatform
 		 *	
 		 *	@param	InPlatform			The UnrealTargetPlatform being built
 		 *	@param	bInAllowFailure		If true, do not throw an exception and return null
 		 *	
 		 *	@return	UEBuildPlatform		The instance of the build platform
 		 */
-		public static UEBuildPlatform GetBuildPlatform(UnrealTargetPlatform InPlatform, bool bInAllowFailure = false)
+		public static IUEBuildPlatform GetBuildPlatform(UnrealTargetPlatform InPlatform, bool bInAllowFailure = false)
 		{
 			if (BuildPlatformDictionary.ContainsKey(InPlatform) == true)
 			{
@@ -139,14 +358,14 @@ namespace UnrealBuildTool
 		}
 
 		/**
-		 *	Retrieve the UEBuildPlatform instance for the given CPPTargetPlatform
+		 *	Retrieve the IUEBuildPlatform instance for the given CPPTargetPlatform
 		 *	
 		 *	@param	InPlatform			The CPPTargetPlatform being built
 		 *	@param	bInAllowFailure		If true, do not throw an exception and return null
 		 *	
 		 *	@return	UEBuildPlatform		The instance of the build platform
 		 */
-		public static UEBuildPlatform GetBuildPlatformForCPPTargetPlatform(CPPTargetPlatform InPlatform, bool bInAllowFailure = false)
+		public static IUEBuildPlatform GetBuildPlatformForCPPTargetPlatform(CPPTargetPlatform InPlatform, bool bInAllowFailure = false)
 		{
 			UnrealTargetPlatform UTPlatform = UEBuildTarget.CPPTargetPlatformToUnrealTargetPlatform(InPlatform);
 			if (BuildPlatformDictionary.ContainsKey(UTPlatform) == true)
@@ -171,7 +390,7 @@ namespace UnrealBuildTool
          */
         public static void PlatformModifyNewlyLoadedModule(UEBuildModule InModule, TargetInfo Target, UnrealTargetPlatform Only = UnrealTargetPlatform.Unknown)
 		{
-			foreach (KeyValuePair<UnrealTargetPlatform, UEBuildPlatform> PlatformEntry in BuildPlatformDictionary)
+			foreach (var PlatformEntry in BuildPlatformDictionary)
 			{
                 if (Only == UnrealTargetPlatform.Unknown || PlatformEntry.Key == Only || PlatformEntry.Key == Target.Platform)
                 {
@@ -664,7 +883,7 @@ namespace UnrealBuildTool
 		public static bool PlatformRequiresMonolithicBuilds(UnrealTargetPlatform InPlatform, UnrealTargetConfiguration InConfiguration)
 		{
 			// Some platforms require monolithic builds...
-			UEBuildPlatform BuildPlatform = GetBuildPlatform(InPlatform, true);
+			var BuildPlatform = GetBuildPlatform(InPlatform, true);
 			if (BuildPlatform != null)
 			{
 				return BuildPlatform.ShouldCompileMonolithicBinary(InPlatform);
diff --git a/Engine/Source/Programs/UnrealBuildTool/Configuration/UEBuildTarget.cs b/Engine/Source/Programs/UnrealBuildTool/Configuration/UEBuildTarget.cs
--- a/Engine/Source/Programs/UnrealBuildTool/Configuration/UEBuildTarget.cs
+++ b/Engine/Source/Programs/UnrealBuildTool/Configuration/UEBuildTarget.cs
@@ -541,7 +541,7 @@ namespace UnrealBuildTool
 		{
 			//@todo. Allow failure to get build platform here?
 			bool bPlatformRequiresMonolithic = false;
-			UEBuildPlatform BuildPlatform = UEBuildPlatform.GetBuildPlatform(InPlatform, true);
+			var BuildPlatform = UEBuildPlatform.GetBuildPlatform(InPlatform, true);
 			if (BuildPlatform != null)
 			{
 				bPlatformRequiresMonolithic = BuildPlatform.ShouldCompileMonolithicBinary(InPlatform);
@@ -1027,7 +1027,7 @@ namespace UnrealBuildTool
 			}
 
 			UnrealTargetPlatform TargetPlatform = CPPTargetPlatformToUnrealTargetPlatform( Platform );
-			UEBuildPlatform BuildPlatform = UEBuildPlatform.GetBuildPlatform( TargetPlatform );
+			var BuildPlatform = UEBuildPlatform.GetBuildPlatform( TargetPlatform );
 
 
 			// Iterate over all the binaries, and add the relevant info to the manifest
@@ -1812,7 +1812,7 @@ namespace UnrealBuildTool
 			UnrealTargetConfiguration Configuration, UEBuildBinaryType BinaryType, TargetRules.TargetType? TargetType, bool bIsRocketModule, PluginInfo PluginInfo, string AppName)
 		{
 			// Determine the binary extension for the platform and binary type.
-			UEBuildPlatform BuildPlatform = UEBuildPlatform.GetBuildPlatform(Platform);
+			var BuildPlatform = UEBuildPlatform.GetBuildPlatform(Platform);
 			string BinaryExtension = BuildPlatform.GetBinaryExtension(BinaryType);
 
 			UnrealTargetConfiguration LocalConfig = Configuration;
@@ -1923,7 +1923,7 @@ namespace UnrealBuildTool
 		/** Sets up the modules for the target. */
 		protected virtual void SetupModules()
 		{
-			UEBuildPlatform BuildPlatform = UEBuildPlatform.GetBuildPlatform(Platform);
+			var BuildPlatform = UEBuildPlatform.GetBuildPlatform(Platform);
 			List<string> PlatformExtraModules = new List<string>();
 			BuildPlatform.GetExtraModules(new TargetInfo(Platform, Configuration, Rules.Type), this, ref PlatformExtraModules);
 			ExtraModuleNames.AddRange(PlatformExtraModules);			
@@ -2040,9 +2040,9 @@ namespace UnrealBuildTool
 		/** Sets up the global compile and link environment for the target. */
 		public virtual void SetupGlobalEnvironment()
 		{
-			UEBuildPlatform BuildPlatform = UEBuildPlatform.GetBuildPlatform(Platform);
+			var BuildPlatform = UEBuildPlatform.GetBuildPlatform(Platform);
 
-			UEToolChain ToolChain = UEToolChain.GetPlatformToolChain(BuildPlatform.GetCPPTargetPlatform(Platform));
+			var ToolChain = UEToolChain.GetPlatformToolChain(BuildPlatform.GetCPPTargetPlatform(Platform));
 			ToolChain.SetUpGlobalEnvironment();
 
 			// Allow each target type (Game/Editor/Server) to set a default global environment state
@@ -2275,7 +2275,7 @@ namespace UnrealBuildTool
 
 		void SetUpPlatformEnvironment()
 		{
-			UEBuildPlatform BuildPlatform = UEBuildPlatform.GetBuildPlatform(Platform);
+			var BuildPlatform = UEBuildPlatform.GetBuildPlatform(Platform);
 
 			CPPTargetPlatform MainCompilePlatform = BuildPlatform.GetCPPTargetPlatform(Platform);
 
diff --git a/Engine/Source/Programs/UnrealBuildTool/IOS/IOSToolChain.cs b/Engine/Source/Programs/UnrealBuildTool/IOS/IOSToolChain.cs
--- a/Engine/Source/Programs/UnrealBuildTool/IOS/IOSToolChain.cs
+++ b/Engine/Source/Programs/UnrealBuildTool/IOS/IOSToolChain.cs
@@ -992,7 +992,7 @@ namespace UnrealBuildTool
 					if (Directory.Exists (Project))
 					{
 						// ensure the plist, entitlements, and provision files are properly copied
-						UEBuildDeploy DeployHandler = UEBuildDeploy.GetBuildDeploy(Target.Platform);
+						var DeployHandler = UEBuildDeploy.GetBuildDeploy(Target.Platform);
 						if (DeployHandler != null)
 						{
 							DeployHandler.PrepTargetForDeployment(Target);
@@ -1180,7 +1180,7 @@ namespace UnrealBuildTool
 					}
 					StubGenerateProcess.StartInfo.Arguments += " -config " + Target.Configuration + " -mac " + RemoteServerName;
 
-					UEBuildPlatform BuildPlatform = UEBuildPlatform.GetBuildPlatform(Target.Platform);
+					var BuildPlatform = UEBuildPlatform.GetBuildPlatform(Target.Platform);
 					string Architecture = BuildPlatform.GetActiveArchitecture();
 					if (Architecture != "")
 					{
diff --git a/Engine/Source/Programs/UnrealBuildTool/Linux/UEBuildLinux.cs b/Engine/Source/Programs/UnrealBuildTool/Linux/UEBuildLinux.cs
--- a/Engine/Source/Programs/UnrealBuildTool/Linux/UEBuildLinux.cs
+++ b/Engine/Source/Programs/UnrealBuildTool/Linux/UEBuildLinux.cs
@@ -89,7 +89,7 @@ namespace UnrealBuildTool.Linux
 				return;
 			}
 
-			if ((ProjectFileGenerator.bGenerateProjectFiles == true) || (HasRequiredSDKsInstalled() == UEBuildPlatform.SDKStatus.Valid))
+			if ((ProjectFileGenerator.bGenerateProjectFiles == true) || (HasRequiredSDKsInstalled() == SDKStatus.Valid))
             {
                 bool bRegisterBuildPlatform = true;
 
diff --git a/Engine/Source/Programs/UnrealBuildTool/System/CPPHeaders.cs b/Engine/Source/Programs/UnrealBuildTool/System/CPPHeaders.cs
--- a/Engine/Source/Programs/UnrealBuildTool/System/CPPHeaders.cs
+++ b/Engine/Source/Programs/UnrealBuildTool/System/CPPHeaders.cs
@@ -340,7 +340,7 @@ namespace UnrealBuildTool
 			++CPPEnvironment.TotalDirectIncludeCacheMisses;
 
 			// Get the adjusted filename
-			UEBuildPlatform BuildPlatform = UEBuildPlatform.GetBuildPlatformForCPPTargetPlatform( TargetPlatform );
+			var BuildPlatform = UEBuildPlatform.GetBuildPlatformForCPPTargetPlatform( TargetPlatform );
 			string FileToRead = CPPFile.AbsolutePath;
 
 			if (BuildPlatform.RequiresExtraUnityCPPWriter() == true &&
diff --git a/Engine/Source/Programs/UnrealBuildTool/System/ProjectFileGenerator.cs b/Engine/Source/Programs/UnrealBuildTool/System/ProjectFileGenerator.cs
--- a/Engine/Source/Programs/UnrealBuildTool/System/ProjectFileGenerator.cs
+++ b/Engine/Source/Programs/UnrealBuildTool/System/ProjectFileGenerator.cs
@@ -1274,7 +1274,7 @@ namespace UnrealBuildTool
 				if (bInProjectPlatformsList && (IncludeAllPlatforms || IsRequiredPlatform))
 				{
 					// If there is a build platform present, add it to the SupportedPlatforms list
-					UEBuildPlatform BuildPlatform = UEBuildPlatform.GetBuildPlatform( Platform, true );
+					var BuildPlatform = UEBuildPlatform.GetBuildPlatform( Platform, true );
 					if( BuildPlatform != null )
 					{
 						if (UnrealBuildTool.IsValidPlatform(Platform))
diff --git a/Engine/Source/Programs/UnrealBuildTool/System/RulesCompiler.cs b/Engine/Source/Programs/UnrealBuildTool/System/RulesCompiler.cs
--- a/Engine/Source/Programs/UnrealBuildTool/System/RulesCompiler.cs
+++ b/Engine/Source/Programs/UnrealBuildTool/System/RulesCompiler.cs
@@ -38,7 +38,7 @@ namespace UnrealBuildTool
 			Type = InitType;
 
 			// get the platform's architecture
-			UEBuildPlatform BuildPlatform = UEBuildPlatform.GetBuildPlatform(Platform);
+			var BuildPlatform = UEBuildPlatform.GetBuildPlatform(Platform);
 			Architecture = BuildPlatform.GetActiveArchitecture();
 		}
 
diff --git a/Engine/Source/Programs/UnrealBuildTool/System/UEPlatformProjectGenerator.cs b/Engine/Source/Programs/UnrealBuildTool/System/UEPlatformProjectGenerator.cs
--- a/Engine/Source/Programs/UnrealBuildTool/System/UEPlatformProjectGenerator.cs
+++ b/Engine/Source/Programs/UnrealBuildTool/System/UEPlatformProjectGenerator.cs
@@ -24,7 +24,7 @@ namespace UnrealBuildTool
 		public static void RegisterPlatformProjectGenerator(UnrealTargetPlatform InPlatform, UEPlatformProjectGenerator InProjectGenerator)
 		{
 			// Make sure the build platform is legal
-			UEBuildPlatform BuildPlatform = UEBuildPlatform.GetBuildPlatform(InPlatform, true);
+			var BuildPlatform = UEBuildPlatform.GetBuildPlatform(InPlatform, true);
 			if (BuildPlatform != null)
 			{
 				if (ProjectGeneratorDictionary.ContainsKey(InPlatform) == true)
diff --git a/Engine/Source/Programs/UnrealBuildTool/System/Unity.cs b/Engine/Source/Programs/UnrealBuildTool/System/Unity.cs
--- a/Engine/Source/Programs/UnrealBuildTool/System/Unity.cs
+++ b/Engine/Source/Programs/UnrealBuildTool/System/Unity.cs
@@ -114,7 +114,7 @@ namespace UnrealBuildTool
 		{
 			var ToolChain = UEToolChain.GetPlatformToolChain(CompileEnvironment.Config.TargetPlatform);
 
-			UEBuildPlatform BuildPlatform = UEBuildPlatform.GetBuildPlatformForCPPTargetPlatform(CompileEnvironment.Config.TargetPlatform);
+			var BuildPlatform = UEBuildPlatform.GetBuildPlatformForCPPTargetPlatform(CompileEnvironment.Config.TargetPlatform);
 
 			// Figure out size of all input files combined. We use this to determine whether to use larger unity threshold or not.
 			long TotalBytesInCPPFiles = CPPFiles.Sum(F => F.Info.Length);
diff --git a/Engine/Source/Programs/UnrealBuildTool/System/UnrealBuildTool.cs b/Engine/Source/Programs/UnrealBuildTool/System/UnrealBuildTool.cs
--- a/Engine/Source/Programs/UnrealBuildTool/System/UnrealBuildTool.cs
+++ b/Engine/Source/Programs/UnrealBuildTool/System/UnrealBuildTool.cs
@@ -229,7 +229,7 @@ namespace UnrealBuildTool
 			{
 				if (Predicate(Platform))
 				{
-					UEBuildPlatform BuildPlatform = UEBuildPlatform.GetBuildPlatform(Platform, true);
+					var BuildPlatform = UEBuildPlatform.GetBuildPlatform(Platform, true);
 					if ((bCheckValidity == false) || (BuildPlatform != null))
 					{
 						OutPlatforms.Add(Platform);
@@ -395,10 +395,10 @@ namespace UnrealBuildTool
                 {
                     if (CheckType.IsClass && !CheckType.IsAbstract)
                     {
-                        if (CheckType.IsSubclassOf(typeof(UEBuildPlatform)))
+						if (Utils.ImplementsInterface<IUEBuildPlatform>(CheckType))
                         {
                             Log.TraceVerbose("    Registering build platform: {0}", CheckType.ToString());
-                            UEBuildPlatform TempInst = (UEBuildPlatform)(UBTAssembly.CreateInstance(CheckType.FullName, true));
+                            var TempInst = (UEBuildPlatform)(UBTAssembly.CreateInstance(CheckType.FullName, true));
                             TempInst.RegisterBuildPlatform();
                         }
                     }
@@ -408,16 +408,16 @@ namespace UnrealBuildTool
 				{
 					if (CheckType.IsClass && !CheckType.IsAbstract)
 					{
-						if (CheckType.IsSubclassOf(typeof(UEToolChain)))
+						if (Utils.ImplementsInterface<IUEToolChain>(CheckType))
 						{
 							Log.TraceVerbose("    Registering tool chain    : {0}", CheckType.ToString());
-							UEToolChain TempInst = (UEToolChain)(UBTAssembly.CreateInstance(CheckType.FullName, true));
+							var TempInst = (UEToolChain)(UBTAssembly.CreateInstance(CheckType.FullName, true));
 							TempInst.RegisterToolChain();
 						}
-						else if (CheckType.IsSubclassOf(typeof(UEBuildDeploy)))
+						else if (Utils.ImplementsInterface<IUEBuildDeploy>(CheckType))
 						{
 							Log.TraceVerbose("    Registering build deploy  : {0}", CheckType.ToString());
-							UEBuildDeploy TempInst = (UEBuildDeploy)(UBTAssembly.CreateInstance(CheckType.FullName, true));
+							var TempInst = (UEBuildDeploy)(UBTAssembly.CreateInstance(CheckType.FullName, true));
 							TempInst.RegisterBuildDeploy();
 						}
 						else if (CheckType.IsSubclassOf(typeof(UEPlatformProjectGenerator)))
@@ -909,7 +909,7 @@ namespace UnrealBuildTool
 							}
 							else
 							{
-								UEBuildPlatform BuildPlatform = UEBuildPlatform.GetBuildPlatform(CheckPlatform, true);
+								var BuildPlatform = UEBuildPlatform.GetBuildPlatform(CheckPlatform, true);
 								if (BuildPlatform != null)
 								{
 									// Setup environment wasn't called, so set the flag
@@ -923,7 +923,7 @@ namespace UnrealBuildTool
 							if ((Result == ECompilationResult.Succeeded) && (BuildConfiguration.bDeployAfterCompile == true) && (BuildConfiguration.bXGEExport == false) &&
 								(UEBuildConfiguration.bGenerateManifest == false) && (UEBuildConfiguration.bCleanProject == false))
 							{
-								UEBuildDeploy DeployHandler = UEBuildDeploy.GetBuildDeploy(CheckPlatform);
+								var DeployHandler = UEBuildDeploy.GetBuildDeploy(CheckPlatform);
 								if (DeployHandler != null)
 								{
 									// We need to be able to identify the Target.Type we can derive it from the Arguments.
@@ -1036,7 +1036,7 @@ namespace UnrealBuildTool
 					ResetConfiguration = UnrealTargetConfiguration.Development;
 				}
 			}
-			UEBuildPlatform BuildPlatform = UEBuildPlatform.GetBuildPlatform(ResetPlatform);
+			var BuildPlatform = UEBuildPlatform.GetBuildPlatform(ResetPlatform);
 
 			// now that we have the platform, we can set the intermediate path to include the platform/architecture name
 			BuildConfiguration.PlatformIntermediateFolder = Path.Combine(BuildConfiguration.BaseIntermediateFolder, ResetPlatform.ToString() + BuildPlatform.GetActiveArchitecture());
@@ -1162,7 +1162,7 @@ namespace UnrealBuildTool
 							ActionsToExecute.Count
 							);
 
-					UEToolChain ToolChain = UEToolChain.GetPlatformToolChain(BuildPlatform.GetCPPTargetPlatform(ResetPlatform));
+					var ToolChain = UEToolChain.GetPlatformToolChain(BuildPlatform.GetCPPTargetPlatform(ResetPlatform));
 					ToolChain.PreBuildSync();
 
 					// Execute the actions.
diff --git a/Engine/Source/Programs/UnrealBuildTool/System/Utils.cs b/Engine/Source/Programs/UnrealBuildTool/System/Utils.cs
--- a/Engine/Source/Programs/UnrealBuildTool/System/Utils.cs
+++ b/Engine/Source/Programs/UnrealBuildTool/System/Utils.cs
@@ -884,6 +884,17 @@ namespace UnrealBuildTool
 		{
 			return Path.GetDirectoryName(GetExecutingAssemblyLocation());
 		}
+
+		/// <summary>
+		/// Checks if given type implements given interface.
+		/// </summary>
+		/// <typeparam name="InterfaceType">Interface to check.</typeparam>
+		/// <param name="TestType">Type to check.</param>
+		/// <returns>True if TestType implements InterfaceType. False otherwise.</returns>
+		public static bool ImplementsInterface<InterfaceType>(Type TestType)
+		{
+			return Array.IndexOf(TestType.GetInterfaces(), typeof(InterfaceType)) != -1;
+		}
 	}
 
 	/// <summary>
diff --git a/Engine/Source/Programs/UnrealBuildTool/System/VCProject.cs b/Engine/Source/Programs/UnrealBuildTool/System/VCProject.cs
--- a/Engine/Source/Programs/UnrealBuildTool/System/VCProject.cs
+++ b/Engine/Source/Programs/UnrealBuildTool/System/VCProject.cs
@@ -177,13 +177,13 @@ namespace UnrealBuildTool
 				return false;
 			}
 
-			UEBuildPlatform BuildPlatform = UEBuildPlatform.GetBuildPlatform(Platform, true);
+			var BuildPlatform = UEBuildPlatform.GetBuildPlatform(Platform, true);
 			if (BuildPlatform == null)
 			{
 				return false;
 			}
 
-			if (BuildPlatform.HasRequiredSDKsInstalled() != UEBuildPlatform.SDKStatus.Valid)
+			if (BuildPlatform.HasRequiredSDKsInstalled() != SDKStatus.Valid)
 			{
 				return false;
 			}
@@ -415,8 +415,8 @@ namespace UnrealBuildTool
 						{
 							continue;
 						}
-						UEBuildPlatform BuildPlatform = UEBuildPlatform.GetBuildPlatform(Platform, true);
-						if ((BuildPlatform != null) && (BuildPlatform.HasRequiredSDKsInstalled() == UEBuildPlatform.SDKStatus.Valid))
+						var BuildPlatform = UEBuildPlatform.GetBuildPlatform(Platform, true);
+						if ((BuildPlatform != null) && (BuildPlatform.HasRequiredSDKsInstalled() == SDKStatus.Valid))
 						{
 							// Now go through all of the target types for this project
 							if( ProjectTargets.Count == 0 )
@@ -942,7 +942,7 @@ namespace UnrealBuildTool
 					var UBTConfigurationName = IsStubProject ? StubProjectConfigurationName : Configuration.ToString();
 
 					// Setup output path
-					UEBuildPlatform BuildPlatform = UEBuildPlatform.GetBuildPlatform(Platform);
+					var BuildPlatform = UEBuildPlatform.GetBuildPlatform(Platform);
 
 					// Figure out if this is a monolithic build
 					bool bShouldCompileMonolithic = BuildPlatform.ShouldCompileMonolithicBinary(Platform);
@@ -989,7 +989,7 @@ namespace UnrealBuildTool
 					}
 					NMakePath += BuildPlatform.GetActiveArchitecture();
 					NMakePath += BuildPlatform.GetBinaryExtension(UEBuildBinaryType.Executable);
-					NMakePath = BuildPlatform.ModifyNMakeOutput(NMakePath);
+					NMakePath = (BuildPlatform as UEBuildPlatform).ModifyNMakeOutput(NMakePath);
 
 					string PathStrings = (ProjGenerator != null) ? ProjGenerator.GetVisualStudioPathsEntries(Platform, Configuration, TargetRulesObject.Type, TargetFilePath, ProjectFilePath, NMakePath) : "";
 					if (string.IsNullOrEmpty(PathStrings) || (PathStrings.Contains("<IntDir>") == false))
diff --git a/Engine/Source/Programs/UnrealBuildTool/System/VCProjectFileGenerator.cs b/Engine/Source/Programs/UnrealBuildTool/System/VCProjectFileGenerator.cs
--- a/Engine/Source/Programs/UnrealBuildTool/System/VCProjectFileGenerator.cs
+++ b/Engine/Source/Programs/UnrealBuildTool/System/VCProjectFileGenerator.cs
@@ -181,11 +181,11 @@ namespace UnrealBuildTool
 			// version of Visual Studio.  This is their chance to override that.
 			foreach( var SupportedPlatform in SupportedPlatforms )
 			{
-				UEBuildPlatform BuildPlatform = UEBuildPlatform.GetBuildPlatform( SupportedPlatform, true );				
+				var BuildPlatform = UEBuildPlatform.GetBuildPlatform( SupportedPlatform, true );				
 				if( BuildPlatform != null )
 				{
 					// Don't worry about platforms that we're missing SDKs for
-					if (BuildPlatform.HasRequiredSDKsInstalled() == UEBuildPlatform.SDKStatus.Valid)
+					if (BuildPlatform.HasRequiredSDKsInstalled() == SDKStatus.Valid)
 					{						
 						// Make sure Visual Studio 2013 project files will work...
 						if( ProjectFileFormat == VCProjectFileFormat.VisualStudio2013 )
diff --git a/Engine/Source/Programs/UnrealBuildTool/System/XcodeProjectFileGenerator.cs b/Engine/Source/Programs/UnrealBuildTool/System/XcodeProjectFileGenerator.cs
--- a/Engine/Source/Programs/UnrealBuildTool/System/XcodeProjectFileGenerator.cs
+++ b/Engine/Source/Programs/UnrealBuildTool/System/XcodeProjectFileGenerator.cs
@@ -676,7 +676,7 @@ namespace UnrealBuildTool
 					IsAGame = true;
 					if (Target.TargetPlatform == UnrealTargetPlatform.IOS)
 					{
-						UEToolChain Toolchain = UEToolChain.GetPlatformToolChain(CPPTargetPlatform.IOS);
+						var Toolchain = UEToolChain.GetPlatformToolChain(CPPTargetPlatform.IOS);
 						GamePath = Toolchain.ConvertPath(Path.GetFullPath(GameFolder));
 					}
 					break;
@@ -695,7 +695,7 @@ namespace UnrealBuildTool
 				}
 				if (Target.TargetPlatform == UnrealTargetPlatform.IOS)
 				{
-					UEToolChain Toolchain = UEToolChain.GetPlatformToolChain(CPPTargetPlatform.IOS);
+					var Toolchain = UEToolChain.GetPlatformToolChain(CPPTargetPlatform.IOS);
 					EngineRelative = Toolchain.ConvertPath(EngineRelative);
 				}
 			}
diff --git a/Engine/Source/Programs/UnrealBuildTool/ToolChain/UEToolChain.cs b/Engine/Source/Programs/UnrealBuildTool/ToolChain/UEToolChain.cs
--- a/Engine/Source/Programs/UnrealBuildTool/ToolChain/UEToolChain.cs
+++ b/Engine/Source/Programs/UnrealBuildTool/ToolChain/UEToolChain.cs
@@ -9,7 +9,35 @@ using Microsoft.Win32;
 
 namespace UnrealBuildTool
 {
-	public abstract class UEToolChain
+	public interface IUEToolChain
+	{
+		void RegisterToolChain();
+		
+		CPPOutput CompileCPPFiles(CPPEnvironment CompileEnvironment, List<FileItem> SourceFiles, string ModuleName);
+		
+		CPPOutput CompileRCFiles(CPPEnvironment Environment, List<FileItem> RCFiles);
+		
+		FileItem LinkFiles(LinkEnvironment LinkEnvironment, bool bBuildImportLibraryOnly);
+		
+		void CompileCSharpProject(CSharpEnvironment CompileEnvironment, string ProjectFileName, string DestinationFile);
+		
+		/** Converts the passed in path from UBT host to compiler native format. */
+		String ConvertPath(String OriginalPath);
+		
+		/// <summary>
+		/// Called immediately after UnrealHeaderTool is executed to generated code for all UObjects modules.  Only is called if UnrealHeaderTool was actually run in this session.
+		/// </summary>
+		/// <param name="UObjectModules">List of UObject modules we generated code for.</param>
+		void PostCodeGeneration(UEBuildTarget Target, UHTManifest Manifest);
+		
+		void PreBuildSync();
+		
+		void PostBuildSync(UEBuildTarget Target);
+		
+		void SetUpGlobalEnvironment();
+	}
+
+	public abstract class UEToolChain : IUEToolChain
 	{
 		static Dictionary<CPPTargetPlatform, UEToolChain> CPPToolChainDictionary = new Dictionary<CPPTargetPlatform, UEToolChain>();
 
@@ -27,7 +55,7 @@ namespace UnrealBuildTool
 			}
 		}
 
-		public static UEToolChain GetPlatformToolChain(CPPTargetPlatform InPlatform)
+		public static IUEToolChain GetPlatformToolChain(CPPTargetPlatform InPlatform)
 		{
 			if (CPPToolChainDictionary.ContainsKey(InPlatform) == true)
 			{
