UBT: More detailed module types. ScriptPlugins can now ingore modules by their type.

diff --git a/Engine/Plugins/Script/ScriptGeneratorPlugin/Source/ScriptGeneratorPlugin/Private/ScriptCodeGeneratorBase.cpp b/Engine/Plugins/Script/ScriptGeneratorPlugin/Source/ScriptGeneratorPlugin/Private/ScriptCodeGeneratorBase.cpp
--- a/Engine/Plugins/Script/ScriptGeneratorPlugin/Source/ScriptGeneratorPlugin/Private/ScriptCodeGeneratorBase.cpp
+++ b/Engine/Plugins/Script/ScriptGeneratorPlugin/Source/ScriptGeneratorPlugin/Private/ScriptCodeGeneratorBase.cpp
@@ -157,20 +157,6 @@ bool FScriptCodeGeneratorBase::CanExportClass(UClass* Class)
 		(Class->ClassFlags & (CLASS_RequiredAPI | CLASS_MinimalAPI)) && // Don't export classes that don't export DLL symbols
 		!ExportedClasses.Contains(Class->GetFName()); // Don't export classes that have already been exported
 
-	if (bCanExport)
-	{
-		// Only export functions from selected modules
-		static struct FSupportedModules
-		{
-			TArray<FString> SupportedScriptModules;
-			FSupportedModules()
-			{
-				GConfig->GetArray(TEXT("Plugins"), TEXT("ScriptSupportedModules"), SupportedScriptModules, GEngineIni);
-			}
-		} SupportedModules;
-
-		bCanExport = SupportedModules.SupportedScriptModules.Contains(Class->GetOutermost()->GetName());
-	}
 	return bCanExport;
 }
 
diff --git a/Engine/Plugins/Script/ScriptGeneratorPlugin/Source/ScriptGeneratorPlugin/Private/ScriptGeneratorPlugin.cpp b/Engine/Plugins/Script/ScriptGeneratorPlugin/Source/ScriptGeneratorPlugin/Private/ScriptGeneratorPlugin.cpp
--- a/Engine/Plugins/Script/ScriptGeneratorPlugin/Source/ScriptGeneratorPlugin/Private/ScriptGeneratorPlugin.cpp
+++ b/Engine/Plugins/Script/ScriptGeneratorPlugin/Source/ScriptGeneratorPlugin/Private/ScriptGeneratorPlugin.cpp
@@ -18,6 +18,7 @@ class FScriptGeneratorPlugin : public IScriptGeneratorPlugin
 
 	/** IScriptGeneratorPlugin interface */
 	virtual FString GetGeneratedCodeModuleName() const OVERRIDE { return TEXT("ScriptPlugin"); }
+	virtual bool ShouldExportClassesForModule(const FString& ModuleName, EBuildModuleType::Type ModuleType) const;
 	virtual bool SupportsTarget(const FString& TargetName) const OVERRIDE { return true; }
 	virtual void Initialize(const FString& RootLocalPath, const FString& RootBuildPath, const FString& OutputDirectory) OVERRIDE;
 	virtual void ExportClass(UClass* Class, const FString& SourceHeaderFilename, const FString& GeneratedHeaderFilename, bool bHasChanged) OVERRIDE;
@@ -48,6 +49,25 @@ void FScriptGeneratorPlugin::Initialize(const FString& RootLocalPath, const FStr
 	UE_LOG(LogScriptGenerator, Log, TEXT("Output directory: %s"), *OutputDirectory);
 }
 
+bool FScriptGeneratorPlugin::ShouldExportClassesForModule(const FString& ModuleName, EBuildModuleType::Type ModuleType) const
+{ 
+	bool bCanExport = (ModuleType == EBuildModuleType::Runtime || ModuleType == EBuildModuleType::Game);
+	if (bCanExport)
+	{
+		// Only export functions from selected modules
+		static struct FSupportedModules
+		{
+			TArray<FString> SupportedScriptModules;
+			FSupportedModules()
+			{
+				GConfig->GetArray(TEXT("Plugins"), TEXT("ScriptSupportedModules"), SupportedScriptModules, GEngineIni);
+			}
+		} SupportedModules;
+		bCanExport = SupportedModules.SupportedScriptModules.Num() == 0 || SupportedModules.SupportedScriptModules.Contains(ModuleName);
+	}
+	return bCanExport;
+}
+
 void FScriptGeneratorPlugin::ExportClass(UClass* Class, const FString& SourceHeaderFilename, const FString& GeneratedHeaderFilename, bool bHasChanged)
 {
 	CodeGenerator->ExportClass(Class, SourceHeaderFilename, GeneratedHeaderFilename, bHasChanged);
diff --git a/Engine/Programs/UnrealHeaderTool/Config/DefaultEngine.ini b/Engine/Programs/UnrealHeaderTool/Config/DefaultEngine.ini
--- a/Engine/Programs/UnrealHeaderTool/Config/DefaultEngine.ini
+++ b/Engine/Programs/UnrealHeaderTool/Config/DefaultEngine.ini
@@ -2,6 +2,6 @@
 +ProgramEnabledPlugins="ScriptGeneratorPlugin"
 +ScriptPlugins="ScriptGeneratorPlugin"
 
-+ScriptSupportedModules=/Script/CoreUObject
-+ScriptSupportedModules=/Script/Engine
-+ScriptSupportedModules=/Script/ScriptPlugin
\ No newline at end of file
++ScriptSupportedModules=CoreUObject
++ScriptSupportedModules=Engine
++ScriptSupportedModules=ScriptPlugin
\ No newline at end of file
diff --git a/Engine/Source/Programs/UnrealBuildTool/Configuration/UEBuildModule.cs b/Engine/Source/Programs/UnrealBuildTool/Configuration/UEBuildModule.cs
--- a/Engine/Source/Programs/UnrealBuildTool/Configuration/UEBuildModule.cs
+++ b/Engine/Source/Programs/UnrealBuildTool/Configuration/UEBuildModule.cs
@@ -11,16 +11,66 @@ using System.Xml;
 
 namespace UnrealBuildTool
 {
-	/** Type of module. */
+	/** Type of module. Mirrored in UHT as EBuildModuleType */
 	public enum UEBuildModuleType
 	{
-		EngineModule,
-		GameModule,
+		Unknown,
+		Runtime,
+		Developer,
+		Editor,
+		ThirdParty,
+		Plugin,
+		Program,
+		Game,
+	}
+	public static class UEBuildModuleTypeExtensions
+	{
+		public static bool IsEngineModule(this UEBuildModuleType ModuleType)
+		{
+			return ModuleType != UEBuildModuleType.Game;
+		}
 	}
 	
 	/** A unit of code compilation and linking. */
 	public abstract class UEBuildModule
 	{
+		/** Known module folders */
+		public static readonly string RuntimeFolder = String.Format("{0}Runtime{0}", Path.DirectorySeparatorChar);
+		public static readonly string DeveloperFolder = String.Format("{0}Developer{0}", Path.DirectorySeparatorChar);
+		public static readonly string EditorFolder = String.Format("{0}Editor{0}", Path.DirectorySeparatorChar);
+		public static readonly string PluginFolder = String.Format("Plugins{0}", Path.DirectorySeparatorChar);
+		public static readonly string ThirdPartyFolder = String.Format("{0}ThirdParty{0}", Path.DirectorySeparatorChar);
+		public static readonly string ProgramsFolder = String.Format("{0}Programs{0}", Path.DirectorySeparatorChar);
+
+		public static UEBuildModuleType GetEngineModuleTypeBasedOnLocation(UEBuildModuleType ModuleType, string ModuleFileRelativeToEngineDirectory)
+		{
+			if (ModuleFileRelativeToEngineDirectory.IndexOf(UEBuildModule.RuntimeFolder, StringComparison.InvariantCultureIgnoreCase) >= 0)
+			{
+				ModuleType = UEBuildModuleType.Runtime;
+			}
+			else if (ModuleFileRelativeToEngineDirectory.IndexOf(UEBuildModule.DeveloperFolder, StringComparison.InvariantCultureIgnoreCase) >= 0)
+			{
+				ModuleType = UEBuildModuleType.Developer;
+			}
+			else if (ModuleFileRelativeToEngineDirectory.IndexOf(UEBuildModule.EditorFolder, StringComparison.InvariantCultureIgnoreCase) >= 0)
+			{
+				ModuleType = UEBuildModuleType.Editor;
+			}
+			else if (ModuleFileRelativeToEngineDirectory.StartsWith(UEBuildModule.PluginFolder, StringComparison.InvariantCultureIgnoreCase))
+			{
+				ModuleType = UEBuildModuleType.Plugin;
+			}
+			else if (ModuleFileRelativeToEngineDirectory.IndexOf(UEBuildModule.ThirdPartyFolder, StringComparison.InvariantCultureIgnoreCase) >= 0)
+			{
+				ModuleType = UEBuildModuleType.ThirdParty;
+			}
+			else if (ModuleFileRelativeToEngineDirectory.IndexOf(UEBuildModule.ProgramsFolder, StringComparison.InvariantCultureIgnoreCase) >= 0)
+			{
+				ModuleType = UEBuildModuleType.Program;
+			}
+			return ModuleType;
+		}
+
 		/** Converts an optional string list parameter to a well-defined list. */
 		protected static List<string> ListFromOptionalEnumerableStringParameter(IEnumerable<string> InEnumerableStrings)
 		{
@@ -1476,7 +1526,7 @@ namespace UnrealBuildTool
 			Result.Config.OutputDirectory                        = Path.Combine(Binary.Config.IntermediateDirectory, Name);
 
 			// Switch the optimization flag if we're building a game module
-			if (Target.Configuration == UnrealTargetConfiguration.DebugGame && Type == UEBuildModuleType.GameModule)
+			if (Target.Configuration == UnrealTargetConfiguration.DebugGame && Type == UEBuildModuleType.Game)
 			{
 				Result.Config.TargetConfiguration = CPPTargetConfiguration.Debug;
 			}
@@ -1608,6 +1658,7 @@ namespace UnrealBuildTool
 			CachedModuleUHTInfo = new UHTModuleInfo {
 				ModuleName                  = this.Name,
 				ModuleDirectory             = this.ModuleDirectory,
+				ModuleType					= this.Type.ToString(),
 				PublicUObjectClassesHeaders = _AllClassesHeaders    .ToList(),
 				PublicUObjectHeaders        = _PublicUObjectHeaders .ToList(),
 				PrivateUObjectHeaders       = _PrivateUObjectHeaders.ToList()
diff --git a/Engine/Source/Programs/UnrealBuildTool/Configuration/UEBuildTarget.cs b/Engine/Source/Programs/UnrealBuildTool/Configuration/UEBuildTarget.cs
--- a/Engine/Source/Programs/UnrealBuildTool/Configuration/UEBuildTarget.cs
+++ b/Engine/Source/Programs/UnrealBuildTool/Configuration/UEBuildTarget.cs
@@ -1546,7 +1546,7 @@ namespace UnrealBuildTool
 				UEBuildModuleCPP NewModule = new UEBuildModuleCPP(
 					InTarget: this,
 					InName: LinkerFixupsName,
-					InType: UEBuildModuleType.GameModule,
+					InType: UEBuildModuleType.Game,
 					InModuleDirectory:FakeModuleDirectory,
 					InOutputDirectory: GlobalCompileEnvironment.Config.OutputDirectory,
 					InIntelliSenseGatherer: null,
@@ -2344,8 +2344,9 @@ namespace UnrealBuildTool
 
 				// Making an assumption here that any project file path that contains '../../'
 				// is NOT from the engine and therefore must be an application-specific module.
-				bool IsGameModule = false;
-				string ApplicationOutputPath = "";
+				// Get the type of module we're creating
+				var ModuleType = UEBuildModuleType.Unknown;
+				var ApplicationOutputPath = "";
 				var ModuleFileRelativeToEngineDirectory = Utils.MakePathRelativeTo(ModuleFileName, ProjectFileGenerator.EngineRelativePath);
 				if (ModuleFileRelativeToEngineDirectory.StartsWith("..") || Path.IsPathRooted(ModuleFileRelativeToEngineDirectory))
 				{
@@ -2354,7 +2355,7 @@ namespace UnrealBuildTool
 					if (SourceIndex != -1)
 					{
 						ApplicationOutputPath = ApplicationOutputPath.Substring(0, SourceIndex + 1);
-						IsGameModule = true;
+						ModuleType = UEBuildModuleType.Game;
 					}
 					else
 					{
@@ -2362,6 +2363,16 @@ namespace UnrealBuildTool
 							ModuleName, ModuleFileName);
 					}
 				}
+				else
+				{
+					ModuleType = UEBuildModule.GetEngineModuleTypeBasedOnLocation(ModuleType, ModuleFileRelativeToEngineDirectory);
+					if (ModuleType == UEBuildModuleType.Unknown)
+					{
+						throw new BuildException("Unable to determine module type for {0}", ModuleFileName);
+					}
+				}
+
+				var IsGameModule = !ModuleType.IsEngineModule();
 
 				// Get the base directory for paths referenced by the module. If the module's under the UProject source directory use that, otherwise leave it relative to the Engine source directory.
 				string ProjectSourcePath = UnrealBuildTool.GetUProjectSourcePath();
@@ -2421,9 +2432,6 @@ namespace UnrealBuildTool
 					}
 				}
 
-				// Get the type of module we're creating
-				UEBuildModuleType ModuleType = IsGameModule ? UEBuildModuleType.GameModule : UEBuildModuleType.EngineModule;
-
 				// Now, go ahead and create the module builder instance
 				if((Module = InstantiateModule(RulesObject, ModuleName, ModuleType, ModuleDirectory, ApplicationOutputPath, IntelliSenseGatherer, ModuleSourceFiles)) == null)
 				{
diff --git a/Engine/Source/Programs/UnrealBuildTool/System/ExternalExecution.cs b/Engine/Source/Programs/UnrealBuildTool/System/ExternalExecution.cs
--- a/Engine/Source/Programs/UnrealBuildTool/System/ExternalExecution.cs
+++ b/Engine/Source/Programs/UnrealBuildTool/System/ExternalExecution.cs
@@ -31,6 +31,9 @@ namespace UnrealBuildTool
 		/** Module base directory */
 		public string ModuleDirectory;
 
+		/** Module type */
+		public string ModuleType;
+
 		/** Public UObject headers found in the Classes directory (legacy) */
 		public List<FileItem> PublicUObjectClassesHeaders;
 
@@ -57,6 +60,7 @@ namespace UnrealBuildTool
 		public struct Module
 		{
 			public string       Name;
+			public string		ModuleType;
 			public string       BaseDirectory;
 			public string       IncludeBase;     // The include path which all UHT-generated includes should be relative to
 			public string       OutputDirectory;
@@ -82,6 +86,7 @@ namespace UnrealBuildTool
 
 			Modules = ModuleInfo.Select(Info => new Module{
 				Name                     = Info.ModuleName,
+				ModuleType				 = Info.ModuleType,
 				BaseDirectory            = Info.ModuleDirectory,
 				IncludeBase              = Info.ModuleDirectory,
 				OutputDirectory          = UEBuildModuleCPP.GetGeneratedCodeDirectoryForModule(Target, Info.ModuleDirectory, Info.ModuleName),
diff --git a/Engine/Source/Programs/UnrealHeaderTool/Private/CodeGenerator.cpp b/Engine/Source/Programs/UnrealHeaderTool/Private/CodeGenerator.cpp
--- a/Engine/Source/Programs/UnrealHeaderTool/Private/CodeGenerator.cpp
+++ b/Engine/Source/Programs/UnrealHeaderTool/Private/CodeGenerator.cpp
@@ -8,9 +8,9 @@
 #include "NativeClassExporter.h"
 #include "HeaderParser.h"
 #include "ClassMaps.h"
+#include "IScriptGeneratorPluginInterface.h"
 #include "Manifest.h"
 #include "StringUtils.h"
-#include "IScriptGeneratorPluginInterface.h"
 
 /////////////////////////////////////////////////////
 // Globals
@@ -4664,7 +4664,7 @@ ECompilationResult::Type UnrealHeaderTool_Main(const FString& ModuleInfoFilename
 			{
 				if (UPackage* Package = Cast<UPackage>(StaticFindObjectFast(UPackage::StaticClass(), NULL, FName(*Module.LongPackageName), false, false)))
 				{
-					Result = FHeaderParser::ParseAllHeadersInside(GWarn, Package, Module.SaveExportedHeaders, ScriptPlugins);
+					Result = FHeaderParser::ParseAllHeadersInside(GWarn, Package, Module, ScriptPlugins);
 					if (Result != ECompilationResult::Succeeded)
 					{
 						++NumFailures;
diff --git a/Engine/Source/Programs/UnrealHeaderTool/Private/HeaderParser.cpp b/Engine/Source/Programs/UnrealHeaderTool/Private/HeaderParser.cpp
--- a/Engine/Source/Programs/UnrealHeaderTool/Private/HeaderParser.cpp
+++ b/Engine/Source/Programs/UnrealHeaderTool/Private/HeaderParser.cpp
@@ -6581,32 +6581,30 @@ FString FHeaderParser::RequireExactlyOneSpecifierValue(const FPropertySpecifier&
 }
 
 // Exports the class to all vailable plugins
-void ExportClassToScriptPlugins(UClass* Class, TArray<IScriptGeneratorPluginInterface*>& ScriptPlugins)
+void ExportClassToScriptPlugins(UClass* Class, const FManifestModule& Module, IScriptGeneratorPluginInterface& ScriptPlugin)
 {
-	for (auto Plugin : ScriptPlugins)
-	{
-		auto ClassHeaderInfo = GClassGeneratedFileMap.FindRef(Class);
-		Plugin->ExportClass(Class, ClassHeaderInfo.SourceFilename, ClassHeaderInfo.GeneratedFilename, ClassHeaderInfo.bHasChanged);
-	}
+	auto ClassHeaderInfo = GClassGeneratedFileMap.FindRef(Class);
+	ScriptPlugin.ExportClass(Class, ClassHeaderInfo.SourceFilename, ClassHeaderInfo.GeneratedFilename, ClassHeaderInfo.bHasChanged);
 }
+
 // Exports class tree to all available plugins
-void ExportClassTreeToScriptPlugins(const FClassTree* Node, TArray<IScriptGeneratorPluginInterface*>& ScriptPlugins)
+void ExportClassTreeToScriptPlugins(const FClassTree* Node, const FManifestModule& Module, IScriptGeneratorPluginInterface& ScriptPlugin)
 {
 	for (int32 ChildIndex = 0; ChildIndex < Node->NumChildren(); ++ChildIndex)
 	{
 		auto ChildNode = Node->GetChild(ChildIndex);
-		ExportClassToScriptPlugins(ChildNode->GetClass(), ScriptPlugins);
+		ExportClassToScriptPlugins(ChildNode->GetClass(), Module, ScriptPlugin);
 	}
 
 	for (int32 ChildIndex = 0; ChildIndex < Node->NumChildren(); ++ChildIndex)
 	{
 		auto ChildNode = Node->GetChild(ChildIndex);
-		ExportClassTreeToScriptPlugins(ChildNode, ScriptPlugins);
+		ExportClassTreeToScriptPlugins(ChildNode, Module, ScriptPlugin);
 	}
 }
 
 // Parse all headers for classes that are inside CurrentPackage.
-ECompilationResult::Type FHeaderParser::ParseAllHeadersInside(FFeedbackContext* Warn, UPackage* CurrentPackage, bool bAllowSaveExportedHeaders, TArray<IScriptGeneratorPluginInterface*>& ScriptPlugins)
+ECompilationResult::Type FHeaderParser::ParseAllHeadersInside(FFeedbackContext* Warn, UPackage* CurrentPackage, const FManifestModule& Module, TArray<IScriptGeneratorPluginInterface*>& ScriptPlugins)
 {
 	// Disable loading of objects outside of this package (or more exactly, objects which aren't UFields, CDO, or templates)
 	TGuardValue<bool> AutoRestoreVerifyObjectRefsFlag(GVerifyObjectReferencesOnly, true);
@@ -6643,7 +6641,7 @@ ECompilationResult::Type FHeaderParser::ParseAllHeadersInside(FFeedbackContext*
 			// from the feedback context.
 			Warn->SetContext(NULL);
 
-			ExportNativeHeaders(CurrentPackage, AllClasses, bAllowSaveExportedHeaders);
+			ExportNativeHeaders(CurrentPackage, AllClasses, Module.SaveExportedHeaders);
 
 			// Done with header generation
 			if (HeaderParser.LinesParsed > 0)
@@ -6669,8 +6667,14 @@ ECompilationResult::Type FHeaderParser::ParseAllHeadersInside(FFeedbackContext*
 	if (Result == ECompilationResult::Succeeded && ScriptPlugins.Num())
 	{
 		auto RootNode = &AllClasses.GetClassTree();
-		ExportClassToScriptPlugins(RootNode->GetClass(), ScriptPlugins);
-		ExportClassTreeToScriptPlugins(&AllClasses.GetClassTree(), ScriptPlugins);
+		for (auto Plugin : ScriptPlugins)
+		{
+			if (Plugin->ShouldExportClassesForModule(Module.Name, Module.ModuleType))
+			{
+				ExportClassToScriptPlugins(RootNode->GetClass(), Module, *Plugin);
+				ExportClassTreeToScriptPlugins(RootNode, Module, *Plugin);
+			}
+		}
 	}
 
 	return Result;
diff --git a/Engine/Source/Programs/UnrealHeaderTool/Private/HeaderParser.h b/Engine/Source/Programs/UnrealHeaderTool/Private/HeaderParser.h
--- a/Engine/Source/Programs/UnrealHeaderTool/Private/HeaderParser.h
+++ b/Engine/Source/Programs/UnrealHeaderTool/Private/HeaderParser.h
@@ -80,7 +80,7 @@ class FHeaderParser : public FBaseParser, public FContextSupplier
 {
 public:
 	// Parse all headers for classes that are inside LimitOuter.
-	static ECompilationResult::Type ParseAllHeadersInside(FFeedbackContext* Warn, UPackage* LimitOuter, bool bAllowSaveExportedHeaders, TArray<class IScriptGeneratorPluginInterface*>& ScriptPlugins);
+	static ECompilationResult::Type ParseAllHeadersInside(FFeedbackContext* Warn, UPackage* LimitOuter, const FManifestModule& Module, TArray<class IScriptGeneratorPluginInterface*>& ScriptPlugins);
 
 	// Performs a preliminary parse of the text in the specified buffer, pulling out:
 	//   Class name and parent class name
diff --git a/Engine/Source/Programs/UnrealHeaderTool/Private/Manifest.cpp b/Engine/Source/Programs/UnrealHeaderTool/Private/Manifest.cpp
--- a/Engine/Source/Programs/UnrealHeaderTool/Private/Manifest.cpp
+++ b/Engine/Source/Programs/UnrealHeaderTool/Private/Manifest.cpp
@@ -1,6 +1,7 @@
 // Copyright 1998-2014 Epic Games, Inc. All Rights Reserved.
 
 #include "UnrealHeaderTool.h"
+#include "IScriptGeneratorPluginInterface.h"
 #include "Manifest.h"
 #include "Json.h"
 
@@ -106,6 +107,10 @@ FManifest FManifest::LoadFromFile(const FString& Filename)
 		GetJsonFieldValue(KnownModule.PCH,                       ModuleObj, TEXT("PCH"),                      *Outer);
 		GetJsonFieldValue(KnownModule.GeneratedCPPFilenameBase,  ModuleObj, TEXT("GeneratedCPPFilenameBase"), *Outer);
 
+		FString ModuleTypeText;
+		GetJsonFieldValue(ModuleTypeText, ModuleObj, TEXT("ModuleType"), *Outer);
+		KnownModule.ModuleType = EBuildModuleType::Parse(*ModuleTypeText);
+
 		KnownModule.LongPackageName = FPackageName::ConvertToLongScriptPackageName(*KnownModule.Name);
 
 		// Convert relative paths
@@ -134,6 +139,14 @@ FManifest FManifest::LoadFromFile(const FString& Filename)
 		KnownModule.PublicUObjectHeaders       .Sort();
 		KnownModule.PrivateUObjectHeaders      .Sort();
 
+		UE_LOG(LogCompile, Log, TEXT("  %s"), *KnownModule.Name);
+		UE_LOG(LogCompile, Log, TEXT("  .BaseDirectory=%s"), *KnownModule.BaseDirectory);
+		UE_LOG(LogCompile, Log, TEXT("  .IncludeBase=%s"), *KnownModule.IncludeBase);
+		UE_LOG(LogCompile, Log, TEXT("  .GeneratedIncludeDirectory=%s"), *KnownModule.GeneratedIncludeDirectory);
+		UE_LOG(LogCompile, Log, TEXT("  .SaveExportedHeaders=%s"), KnownModule.SaveExportedHeaders ? TEXT("True") : TEXT("False"));
+		UE_LOG(LogCompile, Log, TEXT("  .GeneratedCPPFilenameBase=%s"), *KnownModule.GeneratedCPPFilenameBase);
+		UE_LOG(LogCompile, Log, TEXT("  .ModuleType=%s"), *ModuleTypeText);
+
 		++ModuleIndex;
 	}
 
diff --git a/Engine/Source/Programs/UnrealHeaderTool/Private/Manifest.h b/Engine/Source/Programs/UnrealHeaderTool/Private/Manifest.h
--- a/Engine/Source/Programs/UnrealHeaderTool/Private/Manifest.h
+++ b/Engine/Source/Programs/UnrealHeaderTool/Private/Manifest.h
@@ -12,6 +12,9 @@ struct FManifestModule
 	/** The name of the module */
 	FString Name;
 
+	/** Module type */
+	EBuildModuleType::Type ModuleType;
+
 	/** Long package name for this module's UObject class */
 	FString LongPackageName;
 
diff --git a/Engine/Source/Programs/UnrealHeaderTool/Public/IScriptGeneratorPluginInterface.h b/Engine/Source/Programs/UnrealHeaderTool/Public/IScriptGeneratorPluginInterface.h
--- a/Engine/Source/Programs/UnrealHeaderTool/Public/IScriptGeneratorPluginInterface.h
+++ b/Engine/Source/Programs/UnrealHeaderTool/Public/IScriptGeneratorPluginInterface.h
@@ -3,6 +3,74 @@
 
 #include "ModuleManager.h"
 
+/** Build module type, mirrored in UEBuildModule.cs, enum UEBUildModuletype */
+struct EBuildModuleType
+{
+	enum Type
+	{
+		Unknown,
+		Runtime,
+		Developer,
+		Editor,
+		ThirdParty,
+		Plugin,
+		Program,
+		Game,
+		// NOTE: If you add a new value, make sure to update the ToString() method below!
+
+		Max
+	};
+
+
+	/**
+	* Converts a string literal into EModuleType::Type value
+	*
+	* @param	The string to convert to EModuleType::Type
+	*
+	* @return	The enum value corresponding to the name
+	*/
+	inline static const EBuildModuleType::Type Parse(const TCHAR* Value)
+	{
+		if (FCString::Stricmp(Value, TEXT("Unknown")) == 0)
+		{
+			return Unknown;
+		}
+		else if (FCString::Stricmp(Value, TEXT("Runtime")) == 0)
+		{
+			return Runtime;
+		}
+		else if (FCString::Stricmp(Value, TEXT("Developer")) == 0)
+		{
+			return Developer;
+		}
+		else if (FCString::Stricmp(Value, TEXT("Editor")) == 0)
+		{
+			return Editor;
+		}
+		else if (FCString::Stricmp(Value, TEXT("ThirdParty")) == 0)
+		{
+			return ThirdParty;
+		}
+		else if (FCString::Stricmp(Value, TEXT("Plugin")) == 0)
+		{
+			return Plugin;
+		}
+		else if (FCString::Stricmp(Value, TEXT("Program")) == 0)
+		{
+			return Program;
+		}
+		else if (FCString::Stricmp(Value, TEXT("Game")) == 0)
+		{
+			return Game;
+		}
+		else
+		{
+			FError::Throwf(TEXT("Unrecognized EBuildModuleType name: %s"), Value);
+			return Unknown;
+		}
+	}
+};
+
 /**
  * The public interface to script generator plugins.
  */
@@ -14,6 +82,8 @@ public:
 	virtual FString GetGeneratedCodeModuleName() const = 0;
 	/** Returns true if this plugin supports exporting scripts for the specified target. This should handle game as well as editor target names */
 	virtual bool SupportsTarget(const FString& TargetName) const = 0;
+	/** Returns true if this plugin supports exporting scripts for the specified module */
+	virtual bool ShouldExportClassesForModule(const FString& ModuleName, EBuildModuleType::Type ModuleType) const = 0;
 	/** Initializes this plugin with build information */
 	virtual void Initialize(const FString& RootLocalPath, const FString& RootBuildPath, const FString& OutputDirectory) = 0;
 	/** Exports a single class. May be called multiple times for the same class (as UHT processes the entire hierarchy inside modules. */
