[GitHub] 111 : Remote tool chain game project fixes
#ttp 333984

diff --git a/Engine/Source/Programs/UnrealBuildTool/ToolChain/RemoteToolChain.cs b/Engine/Source/Programs/UnrealBuildTool/ToolChain/RemoteToolChain.cs
--- a/Engine/Source/Programs/UnrealBuildTool/ToolChain/RemoteToolChain.cs
+++ b/Engine/Source/Programs/UnrealBuildTool/ToolChain/RemoteToolChain.cs
@@ -33,10 +33,15 @@ namespace UnrealBuildTool
 		/** The path (on the Mac) to the your particular development directory, where files will be copied to from the PC */
 		public static string UserDevRootMac = "/UE4/Builds/";
 
-		/** The directory that this local branch is in, without drive information (strip off X:\ from X:\UE4\iOS) */
-		public static string BranchDirectory = Environment.MachineName + "\\" + Path.GetFullPath(".\\").Substring(3);
-		public static string BranchDirectoryMac = BranchDirectory.Replace("\\", "/");
+		/** The directory that the engine is in, without drive information (strip off X:\ from X:\UE4\iOS) */
+		public static string UnrealEngineDirectory = Environment.MachineName + "\\" + Path.GetFullPath("..\\..\\").Substring(3);
+		public static string UnrealEngineDirectoryMac = UnrealEngineDirectory.Replace("\\", "/");
 
+		/** The directory that the game project is in. */
+		public static string ProjectDirectory = UnrealBuildTool.HasUProjectFile()
+			? Environment.MachineName + "\\" + Path.GetFullPath(UnrealBuildTool.GetUProjectPath()).Substring(3)
+			: "";
+		public static string ProjectDirectoryMac = ProjectDirectory.Replace("\\","/");
 
         /** Substrings that indicate a line contains an error */
         protected static List<string> ErrorMessageTokens;
@@ -53,9 +58,9 @@ namespace UnrealBuildTool
             ErrorMessageTokens.Add("IPP ERROR");
             ErrorMessageTokens.Add("System.Net.Sockets.SocketException");
 
-			BranchDirectory = BranchDirectory.Replace("Engine\\Binaries\\DotNET", "");
-			BranchDirectory = BranchDirectory.Replace("Engine\\Source\\", "");
-			BranchDirectoryMac = BranchDirectory.Replace("\\", "/");
+			UnrealEngineDirectory = UnrealEngineDirectory.Replace("Engine\\Binaries\\DotNET", "");
+			UnrealEngineDirectory = UnrealEngineDirectory.Replace("Engine\\Source\\", "");
+			UnrealEngineDirectoryMac = UnrealEngineDirectoryMac.Replace("\\", "/");
         }
 
         // Do any one-time, global initialization for the tool chain
@@ -140,7 +145,7 @@ namespace UnrealBuildTool
                 string MacPath = GetMacDevSrcRoot();
 
                 // In the case of paths from the PC to the Mac over a UNC path, peel off the possible roots
-                string StrippedPath = OriginalPath.Replace(BranchDirectory, "");
+				string StrippedPath = OriginalPath.Replace(UnrealEngineDirectory, "");
 
                 // Now, reduce the path down to just relative to UE4 and add file location
                 MacPath += "../../" + StrippedPath.Substring(RootDirectoryLocation(StrippedPath));
@@ -175,13 +180,18 @@ namespace UnrealBuildTool
             // By default, assume the filename is already stripped down and the root is at zero
             int RootDirLocation = 0;
 
-            string UBTRootPath = Path.GetFullPath(AppDomain.CurrentDomain.BaseDirectory + "..\\..\\..\\");
-            if (LocalPath.ToUpperInvariant().Contains(UBTRootPath.ToUpperInvariant()))
-            {
-                // If the file is a full path name and rooted at the same location as UBT,
-                // use that location as the root and simply return the length
-                RootDirLocation = UBTRootPath.Length;
-            }
+			string LocalUnrealEngineRootPath = Path.GetFullPath(AppDomain.CurrentDomain.BaseDirectory + "..\\..\\..\\");
+			string LocalProjectRootPath = Path.GetFullPath(UnrealBuildTool.GetUProjectPath());
+			if (LocalPath.ToUpperInvariant().Contains(LocalUnrealEngineRootPath.ToUpperInvariant()))
+			{
+				// If the file is a full path name and rooted at the same location as UBT,
+				// use that location as the root and simply return the length
+				RootDirLocation = LocalUnrealEngineRootPath.Length;
+			}
+			else if (LocalPath.ToUpperInvariant().Contains(LocalProjectRootPath.ToUpperInvariant()))
+			{
+				RootDirLocation = LocalProjectRootPath.Length;
+			}
 
             return RootDirLocation;
         }
@@ -535,7 +545,7 @@ namespace UnrealBuildTool
         {
             if (ExternalExecution.GetRuntimePlatform() != UnrealTargetPlatform.Mac)
             {
-                return UserDevRootMac + BranchDirectoryMac + "Engine/Source/";
+				return UserDevRootMac + UnrealEngineDirectoryMac + "Engine/Source/";
             }
             else
             {
