[GitHub] 131 : Improve cloning of Environments
#ttp 334941

diff --git a/Engine/Source/Programs/UnrealBuildTool/Configuration/UEBuildBinary.cs b/Engine/Source/Programs/UnrealBuildTool/Configuration/UEBuildBinary.cs
--- a/Engine/Source/Programs/UnrealBuildTool/Configuration/UEBuildBinary.cs
+++ b/Engine/Source/Programs/UnrealBuildTool/Configuration/UEBuildBinary.cs
@@ -468,8 +468,8 @@ namespace UnrealBuildTool
 
 			var OutputFiles = new List<FileItem>();
 
-			var BinaryCompileEnvironment = new CPPEnvironment(CompileEnvironment);
-			var BinaryLinkEnvironment = new LinkEnvironment(LinkEnvironment);
+			var BinaryCompileEnvironment = CompileEnvironment.DeepCopy();
+			var BinaryLinkEnvironment = LinkEnvironment.DeepCopy();
 
 			// Process each module that is linked into the binary.
 			var BinaryDependencies = new List<UEBuildBinary>();
@@ -595,7 +595,7 @@ namespace UnrealBuildTool
 				if (BinaryLinkEnvironment.Config.CanProduceAdditionalConsoleApp && UEBuildConfiguration.bBuildEditor)
 				{					
 					// Produce additional binary but link it as a console app
-					LinkEnvironment ConsoleAppLinkEvironment = new LinkEnvironment(BinaryLinkEnvironment);					
+					var ConsoleAppLinkEvironment = BinaryLinkEnvironment.DeepCopy();
 					ConsoleAppLinkEvironment.Config.bIsBuildingConsoleApplication = true;
 					ConsoleAppLinkEvironment.Config.WindowsEntryPointOverride = "WinMainCRTStartup";		// For WinMain() instead of "main()" for Launch module
 					ConsoleAppLinkEvironment.Config.OutputFilePath = GetAdditionalConsoleAppPath(ConsoleAppLinkEvironment.Config.OutputFilePath);
diff --git a/Engine/Source/Programs/UnrealBuildTool/Configuration/UEBuildModule.cs b/Engine/Source/Programs/UnrealBuildTool/Configuration/UEBuildModule.cs
--- a/Engine/Source/Programs/UnrealBuildTool/Configuration/UEBuildModule.cs
+++ b/Engine/Source/Programs/UnrealBuildTool/Configuration/UEBuildModule.cs
@@ -1279,7 +1279,8 @@ namespace UnrealBuildTool
 				{
 					// Setup a new compile environment for this module's source files.  It's pretty much the exact same as the
 					// module's compile environment, except that it will include a PCH file.
-					var ModulePCHCompileEnvironment = new CPPEnvironment( ModuleCompileEnvironment );
+
+					var ModulePCHCompileEnvironment = ModuleCompileEnvironment.DeepCopy();
 					ModulePCHCompileEnvironment.Config.PrecompiledHeaderAction          = PrecompiledHeaderAction.Include;
 					ModulePCHCompileEnvironment.Config.PrecompiledHeaderIncludeFilename = ModulePCHEnvironment.PrecompiledHeaderIncludeFilename.AbsolutePath;
 					ModulePCHCompileEnvironment.Config.PCHHeaderNameInCode              = ModulePCHEnvironment.PCHHeaderNameInCode;
@@ -1512,7 +1513,7 @@ namespace UnrealBuildTool
 		/// <returns>The new module compile environment.</returns>
 		public CPPEnvironment CreateModuleCompileEnvironment(CPPEnvironment BaseCompileEnvironment)
 		{
-			var Result = new CPPEnvironment(BaseCompileEnvironment);
+			var Result = BaseCompileEnvironment.DeepCopy();
 
 			// Override compile environment
 			Result.Config.bFasterWithoutUnity                    = bFasterWithoutUnity;
@@ -1924,7 +1925,7 @@ namespace UnrealBuildTool
 		// UEBuildModule interface.
 		public override List<FileItem> Compile( CPPEnvironment GlobalCompileEnvironment, CPPEnvironment CompileEnvironment, bool bCompileMonolithic )
 		{
-			var ModuleCLREnvironment = new CPPEnvironment(CompileEnvironment);
+			var ModuleCLREnvironment = CompileEnvironment.DeepCopy();
 
 			// Setup the module environment for the project CLR mode
 			ModuleCLREnvironment.Config.CLRMode = CPPCLRMode.CLREnabled;
diff --git a/Engine/Source/Programs/UnrealBuildTool/Mac/MacToolChain.cs b/Engine/Source/Programs/UnrealBuildTool/Mac/MacToolChain.cs
--- a/Engine/Source/Programs/UnrealBuildTool/Mac/MacToolChain.cs
+++ b/Engine/Source/Programs/UnrealBuildTool/Mac/MacToolChain.cs
@@ -1152,7 +1152,7 @@ namespace UnrealBuildTool
 
 			// We need to know what third party dylibs would be copied to the bundle
 			var Modules = Binary.GetAllDependencyModules(bIncludeDynamicallyLoaded: false, bForceCircular: false);
-			var BinaryLinkEnvironment = new LinkEnvironment(Binary.Target.GlobalLinkEnvironment);
+			var BinaryLinkEnvironment = Binary.Target.GlobalLinkEnvironment.DeepCopy();
 			var BinaryDependencies = new List<UEBuildBinary>();
 			var LinkEnvironmentVisitedModules = new Dictionary<UEBuildModule, bool>();
 			foreach (var Module in Modules)
diff --git a/Engine/Source/Programs/UnrealBuildTool/System/CPPEnvironment.cs b/Engine/Source/Programs/UnrealBuildTool/System/CPPEnvironment.cs
--- a/Engine/Source/Programs/UnrealBuildTool/System/CPPEnvironment.cs
+++ b/Engine/Source/Programs/UnrealBuildTool/System/CPPEnvironment.cs
@@ -261,7 +261,7 @@ namespace UnrealBuildTool
 		{}
 
 		/** Copy constructor. */
-		public CPPEnvironment(CPPEnvironment InCopyEnvironment)
+		protected CPPEnvironment(CPPEnvironment InCopyEnvironment)
 		{
 			PrecompiledHeaderFile = InCopyEnvironment.PrecompiledHeaderFile;
 			PrivateAssemblyDependencies.AddRange(InCopyEnvironment.PrivateAssemblyDependencies);
@@ -310,5 +310,13 @@ namespace UnrealBuildTool
 			PrivateAssemblyDependencies.Add( NewPrivateAssembly );
 		}
 
+		/// <summary>
+		/// Performs a deep copy of this CPPEnvironment object.
+		/// </summary>
+		/// <returns>Copied new CPPEnvironment object.</returns>
+		public virtual CPPEnvironment DeepCopy()
+		{
+			return new CPPEnvironment(this);
+		}
 	};
 }
diff --git a/Engine/Source/Programs/UnrealBuildTool/System/LinkEnvironment.cs b/Engine/Source/Programs/UnrealBuildTool/System/LinkEnvironment.cs
--- a/Engine/Source/Programs/UnrealBuildTool/System/LinkEnvironment.cs
+++ b/Engine/Source/Programs/UnrealBuildTool/System/LinkEnvironment.cs
@@ -166,7 +166,7 @@ namespace UnrealBuildTool
 		}
 
 		/** Copy constructor. */
-		public LinkEnvironment(LinkEnvironment InCopyEnvironment)
+		protected LinkEnvironment(LinkEnvironment InCopyEnvironment)
 		{
 			InputFiles.AddRange(InCopyEnvironment.InputFiles);
 			InputLibraries.AddRange(InCopyEnvironment.InputLibraries);
@@ -179,5 +179,14 @@ namespace UnrealBuildTool
 		{
 			return UEToolChain.GetPlatformToolChain(Config.TargetPlatform).LinkFiles(this, bBuildImportLibraryOnly);
 		}
+
+		/// <summary>
+		/// Performs a deep copy of this LinkEnvironment object.
+		/// </summary>
+		/// <returns>Copied new LinkEnvironment object.</returns>
+		public virtual LinkEnvironment DeepCopy ()
+		{
+			return new LinkEnvironment(this);
+		}
 	}
 }
diff --git a/Engine/Source/Programs/UnrealBuildTool/System/Utils.cs b/Engine/Source/Programs/UnrealBuildTool/System/Utils.cs
--- a/Engine/Source/Programs/UnrealBuildTool/System/Utils.cs
+++ b/Engine/Source/Programs/UnrealBuildTool/System/Utils.cs
@@ -73,7 +73,7 @@ namespace UnrealBuildTool
 				FileItem DummyPCH = FileItem.CreateIntermediateTextFile(DummyPath, TmpPCHHeaderContents);
 
 				// Create a new C++ environment that is used to create the PCH.
-				CPPEnvironment ProjectPCHEnvironment = new CPPEnvironment(ProjectCPPEnvironment);
+				var ProjectPCHEnvironment = ProjectCPPEnvironment.DeepCopy();
 				ProjectPCHEnvironment.Config.PrecompiledHeaderAction = PrecompiledHeaderAction.Create;
 				ProjectPCHEnvironment.Config.PrecompiledHeaderIncludeFilename = PrecompiledHeaderIncludeFilename.AbsolutePath;
 				ProjectPCHEnvironment.Config.PCHHeaderNameInCode = PCHHeaderNameInCode;
