[GitHub] 141 : Move project validation to own method
#ttp 335085

diff --git a/Engine/Source/Programs/UnrealBuildTool/System/ProjectFileGenerator.cs b/Engine/Source/Programs/UnrealBuildTool/System/ProjectFileGenerator.cs
--- a/Engine/Source/Programs/UnrealBuildTool/System/ProjectFileGenerator.cs
+++ b/Engine/Source/Programs/UnrealBuildTool/System/ProjectFileGenerator.cs
@@ -1000,55 +1000,63 @@ namespace UnrealBuildTool
 			Folder.SubFolders.AddRange( SubFoldersToAdd );
 
 			// After everything has been collapsed, do a bit of data validation
+			Validate(Folder, ParentMasterProjectFolderPath);
+		}
+
+		/// <summary>
+		/// Validate the specified Folder. Default implementation requires
+		/// for project file names to be unique!
+		/// </summary>
+		/// <param name="Folder">Folder.</param>
+		/// <param name="MasterProjectFolderPath">Parent master project folder path.</param>
+		protected virtual void Validate(MasterProjectFolder Folder, string MasterProjectFolderPath)
+		{
+			foreach (var CurChildProject in Folder.ChildProjects)
 			{
-				foreach( var CurChildProject in Folder.ChildProjects )
+				foreach (var OtherChildProject in Folder.ChildProjects)
 				{
-					foreach( var OtherChildProject in Folder.ChildProjects )
+					if (CurChildProject != OtherChildProject)
 					{
-						if( CurChildProject != OtherChildProject )
+						if (Utils.GetFilenameWithoutAnyExtensions(CurChildProject.ProjectFilePath).Equals(Utils.GetFilenameWithoutAnyExtensions(OtherChildProject.ProjectFilePath), StringComparison.InvariantCultureIgnoreCase))
 						{
-							if (Utils.GetFilenameWithoutAnyExtensions(CurChildProject.ProjectFilePath).Equals(Utils.GetFilenameWithoutAnyExtensions(OtherChildProject.ProjectFilePath), StringComparison.InvariantCultureIgnoreCase))
-							{
-								throw new BuildException( "Detected collision between two project files with the same path in the same master project folder, " + OtherChildProject.ProjectFilePath + " and " + CurChildProject.ProjectFilePath + " (master project folder: " + MasterProjectFolderPath + ")" );
-							}
+							throw new BuildException("Detected collision between two project files with the same path in the same master project folder, " + OtherChildProject.ProjectFilePath + " and " + CurChildProject.ProjectFilePath + " (master project folder: " + MasterProjectFolderPath + ")");
 						}
 					}
 				}
+			}
 
-				foreach( var SubFolder in Folder.SubFolders )
+			foreach (var SubFolder in Folder.SubFolders)
+			{
+				// If the parent folder already has a child project or file item with the same name as this sub-folder, then
+				// that's considered an error (it should never have been allowed to have a folder name that collided
+				// with project file names or file items, as that's not supported in Visual Studio.)
+				foreach (var CurChildProject in Folder.ChildProjects)
 				{
-					// If the parent folder already has a child project or file item with the same name as this sub-folder, then
-					// that's considered an error (it should never have been allowed to have a folder name that collided
-					// with project file names or file items, as that's not supported in Visual Studio.)
-					foreach( var CurChildProject in Folder.ChildProjects )
+					if (Utils.GetFilenameWithoutAnyExtensions(CurChildProject.ProjectFilePath).Equals(SubFolder.FolderName, StringComparison.InvariantCultureIgnoreCase))
 					{
-						if (Utils.GetFilenameWithoutAnyExtensions(CurChildProject.ProjectFilePath).Equals(SubFolder.FolderName, StringComparison.InvariantCultureIgnoreCase))
-						{
-							throw new BuildException( "Detected collision between a master project sub-folder " + SubFolder.FolderName + " and a project within the outer folder " + CurChildProject.ProjectFilePath + " (master project folder: " + MasterProjectFolderPath + ")" );
-						}
+						throw new BuildException("Detected collision between a master project sub-folder " + SubFolder.FolderName + " and a project within the outer folder " + CurChildProject.ProjectFilePath + " (master project folder: " + MasterProjectFolderPath + ")");
 					}
-					foreach( var CurFile in Folder.Files )
+				}
+				foreach (var CurFile in Folder.Files)
+				{
+					if (Path.GetFileName(CurFile).Equals(SubFolder.FolderName, StringComparison.InvariantCultureIgnoreCase))
 					{
-						if( Path.GetFileName( CurFile ).Equals( SubFolder.FolderName, StringComparison.InvariantCultureIgnoreCase ) )
-						{
-							throw new BuildException( "Detected collision between a master project sub-folder " + SubFolder.FolderName + " and a file within the outer folder " + CurFile + " (master project folder: " + MasterProjectFolderPath + ")" );
-						}
+						throw new BuildException("Detected collision between a master project sub-folder " + SubFolder.FolderName + " and a file within the outer folder " + CurFile + " (master project folder: " + MasterProjectFolderPath + ")");
 					}
-					foreach( var CurFolder in Folder.SubFolders )
+				}
+				foreach (var CurFolder in Folder.SubFolders)
+				{
+					if (CurFolder != SubFolder)
 					{
-						if( CurFolder != SubFolder )
+						if (CurFolder.FolderName.Equals(SubFolder.FolderName, StringComparison.InvariantCultureIgnoreCase))
 						{
-							if( CurFolder.FolderName.Equals( SubFolder.FolderName, StringComparison.InvariantCultureIgnoreCase ) )
-							{
-								throw new BuildException( "Detected collision between a master project sub-folder " + SubFolder.FolderName + " and a sibling folder " + CurFolder.FolderName + " (master project folder: " + MasterProjectFolderPath + ")" );
-							}
+							throw new BuildException("Detected collision between a master project sub-folder " + SubFolder.FolderName + " and a sibling folder " + CurFolder.FolderName + " (master project folder: " + MasterProjectFolderPath + ")");
 						}
 					}
 				}
 			}
 		}
 
-
 		/// <summary>
 		/// Adds UnrealBuildTool to the master project
 		/// </summary>
