[GitHub] 78 : Extract WriteTargetInfo.
#ttp 333495

diff --git a/Engine/Source/Programs/UnrealBuildTool/Configuration/UEBuildTarget.cs b/Engine/Source/Programs/UnrealBuildTool/Configuration/UEBuildTarget.cs
--- a/Engine/Source/Programs/UnrealBuildTool/Configuration/UEBuildTarget.cs
+++ b/Engine/Source/Programs/UnrealBuildTool/Configuration/UEBuildTarget.cs
@@ -1434,25 +1434,33 @@ namespace UnrealBuildTool
 			if (BuildConfiguration.WriteTargetInfoPath != null)
 			{
 				Log.TraceInformation("Writing build environment to " + BuildConfiguration.WriteTargetInfoPath + "...");
+				WriteTargetInfo();
+				OutputItems.Clear();
+			}
+
+			return ECompilationResult.Succeeded;
+		}
+
+		/// <summary>
+		/// Writes target info.
+		/// </summary>
+		protected virtual void WriteTargetInfo()
+		{
+			Log.TraceInformation("Writing build environment to " + BuildConfiguration.WriteTargetInfoPath + "...");
 
-				XmlWriterSettings Settings = new XmlWriterSettings();
-				Settings.Indent = true;
+			XmlWriterSettings Settings = new XmlWriterSettings();
+			Settings.Indent = true;
 
-				using (XmlWriter Writer = XmlWriter.Create(BuildConfiguration.WriteTargetInfoPath, Settings))
+			using (XmlWriter Writer = XmlWriter.Create(BuildConfiguration.WriteTargetInfoPath, Settings))
+			{
+				Writer.WriteStartElement("target");
+				Writer.WriteAttributeString("name", AppName);
+				foreach (UEBuildBinary Binary in AppBinaries)
 				{
-					Writer.WriteStartElement("target");
-					Writer.WriteAttributeString("name", AppName);
-					foreach (UEBuildBinary Binary in AppBinaries)
-					{
-						Binary.WriteBuildEnvironment(GlobalCompileEnvironment, GlobalLinkEnvironment, Writer);
-					}
-					Writer.WriteEndElement();
+					Binary.WriteBuildEnvironment(GlobalCompileEnvironment, GlobalLinkEnvironment, Writer);
 				}
-
-				OutputItems.Clear();
+				Writer.WriteEndElement();
 			}
-
-			return ECompilationResult.Succeeded;
 		}
 
 		/// <summary>
