UBT BuildConfiguration.xml file hasn't been added to UE4 project if it was generated in the same run. Fixed.

diff --git a/Engine/Source/Programs/UnrealBuildTool/Utilities/XmlConfigLoader.cs b/Engine/Source/Programs/UnrealBuildTool/Utilities/XmlConfigLoader.cs
--- a/Engine/Source/Programs/UnrealBuildTool/Utilities/XmlConfigLoader.cs
+++ b/Engine/Source/Programs/UnrealBuildTool/Utilities/XmlConfigLoader.cs
@@ -240,6 +240,7 @@ namespace UnrealBuildTool
 				}
 
 				CreateUserXmlConfigTemplate(PossibleConfigLocation.FSLocation);
+				PossibleConfigLocation.bExists = true;
 			}
 		}
 
