UBT CIS Fix.

diff --git a/Engine/Source/Programs/UnrealBuildTool/Utilities/XmlConfigLoader.cs b/Engine/Source/Programs/UnrealBuildTool/Utilities/XmlConfigLoader.cs
--- a/Engine/Source/Programs/UnrealBuildTool/Utilities/XmlConfigLoader.cs
+++ b/Engine/Source/Programs/UnrealBuildTool/Utilities/XmlConfigLoader.cs
@@ -168,7 +168,7 @@ namespace UnrealBuildTool
 			public bool bCreateIfDoesNotExist { get; private set; }
 
 			// Tells if config file exists in this location.
-			public bool bExists { get; private set; }
+			public bool bExists { get; set; }
 
 			public XmlConfigLocation(string[] FSLocations, string IDEFolderName, bool bCreateIfDoesNotExist = false)
 			{
