UBT BuildConfiguration.xml generate and add to project fix refactor. Now it avoids a possible error when template generation fails, but the file is still added to a project.

diff --git a/Engine/Source/Programs/UnrealBuildTool/Utilities/XmlConfigLoader.cs b/Engine/Source/Programs/UnrealBuildTool/Utilities/XmlConfigLoader.cs
--- a/Engine/Source/Programs/UnrealBuildTool/Utilities/XmlConfigLoader.cs
+++ b/Engine/Source/Programs/UnrealBuildTool/Utilities/XmlConfigLoader.cs
@@ -168,7 +168,7 @@ namespace UnrealBuildTool
 			public bool bCreateIfDoesNotExist { get; private set; }
 
 			// Tells if config file exists in this location.
-			public bool bExists { get; set; }
+			public bool bExists { get; private set; }
 
 			public XmlConfigLocation(string[] FSLocations, string IDEFolderName, bool bCreateIfDoesNotExist = false)
 			{
@@ -184,6 +184,33 @@ namespace UnrealBuildTool
 				: this(new string[] { FSLocation }, IDEFolderName, bCreateIfDoesNotExist)
 			{
 			}
+
+			/// <summary>
+			/// Creates template file in the FS location.
+			/// </summary>
+			public void CreateUserXmlConfigTemplate()
+			{
+				try
+				{
+					Directory.CreateDirectory(Path.GetDirectoryName(FSLocation));
+
+					var FilePath = Path.Combine(FSLocation);
+
+					const string TemplateContent =
+						"<Configuration>\n" +
+						"	<BuildConfiguration>\n" +
+						"	</BuildConfiguration>\n" +
+						"</Configuration>\n";
+
+					File.WriteAllText(FilePath, TemplateContent);
+					bExists = true;
+				}
+				catch (Exception)
+				{
+					// Ignore quietly.
+				}
+			}
+
 		}
 
 		public static readonly XmlConfigLocation[] ConfigLocationHierarchy;
@@ -239,37 +266,10 @@ namespace UnrealBuildTool
 					continue;
 				}
 
-				CreateUserXmlConfigTemplate(PossibleConfigLocation.FSLocation);
-				PossibleConfigLocation.bExists = true;
+				PossibleConfigLocation.CreateUserXmlConfigTemplate();
 			}
 		}
 
-
-		/// <summary>
-		/// Creates template file in one of the user locations if it not
-		/// </summary>
-		/// <param name="Location">Location of the xml config file to create.</param>
-		private static void CreateUserXmlConfigTemplate(string Location)
-		{
-			try
-			{
-				Directory.CreateDirectory(Path.GetDirectoryName(Location));
-
-				var FilePath = Path.Combine(Location);
-
-				const string TemplateContent =
-					"<Configuration>\n" +
-					"	<BuildConfiguration>\n" +
-					"	</BuildConfiguration>\n" +
-					"</Configuration>\n";
-
-				File.WriteAllText(FilePath, TemplateContent);
-			}
-			catch (Exception)
-			{
-				// Ignore quietly.
-			}
-		}
 		/// <summary>
 		/// Sets values of this class with values from given XML file.
 		/// </summary>
