Updated UBT's WriteFileIfChanged() to not use a BOM marker on Mac

diff --git a/Engine/Source/Programs/UnrealBuildTool/System/ProjectFileGenerator.cs b/Engine/Source/Programs/UnrealBuildTool/System/ProjectFileGenerator.cs
--- a/Engine/Source/Programs/UnrealBuildTool/System/ProjectFileGenerator.cs
+++ b/Engine/Source/Programs/UnrealBuildTool/System/ProjectFileGenerator.cs
@@ -1966,7 +1966,7 @@ namespace UnrealBuildTool
                     // TODO(sbc): See if we can just drop the Encoding.UTF8 argument on all
                     // platforms.  In this case UTF8 encoding will still be used but without the
                     // BOM, which is, AFAICT, desirable in almost all cases.
-                    if (ExternalExecution.GetRuntimePlatform() == UnrealTargetPlatform.Linux)
+					if (ExternalExecution.GetRuntimePlatform() == UnrealTargetPlatform.Linux || ExternalExecution.GetRuntimePlatform() == UnrealTargetPlatform.Mac)
                         File.WriteAllText(FileName, NewFileContents);
                     else
                         File.WriteAllText(FileName, NewFileContents, Encoding.UTF8);
