UE4: Removed bCompileNetworkProfiler. USE_NETWORK_PROFILER is now true whenever STATS is also true. Also, removed some platform-specific code in NetworkProfiler.cpp

diff --git a/Engine/Programs/UnrealBuildTool/BuildConfiguration.xml b/Engine/Programs/UnrealBuildTool/BuildConfiguration.xml
--- a/Engine/Programs/UnrealBuildTool/BuildConfiguration.xml
+++ b/Engine/Programs/UnrealBuildTool/BuildConfiguration.xml
@@ -89,7 +89,6 @@
 		<BaseIntermediateFolder>Intermediate/Build/</BaseIntermediateFolder>
 	</BuildConfiguration>
 	<UEBuildConfiguration>
-		<bCompileNetworkProfiler>False</bCompileNetworkProfiler>
 		<bCompileSimplygon>True</bCompileSimplygon>
 		<bCompileRecast>True</bCompileRecast>
 		<bCompileSpeedTree>True</bCompileSpeedTree>
diff --git a/Engine/Source/Programs/BlankProgram/BlankProgram.Target.cs b/Engine/Source/Programs/BlankProgram/BlankProgram.Target.cs
--- a/Engine/Source/Programs/BlankProgram/BlankProgram.Target.cs
+++ b/Engine/Source/Programs/BlankProgram/BlankProgram.Target.cs
@@ -37,8 +37,6 @@ public class BlankProgramTarget : TargetRules
 		ref CPPEnvironmentConfiguration OutCPPEnvironmentConfiguration
 		)
 	{
-		UEBuildConfiguration.bCompileNetworkProfiler = false;
-
 		// Lean and mean
 		UEBuildConfiguration.bCompileLeanAndMeanUE = true;
 
diff --git a/Engine/Source/Programs/CrashReporter/CrashReportClient/CrashReportClient.Target.cs b/Engine/Source/Programs/CrashReporter/CrashReportClient/CrashReportClient.Target.cs
--- a/Engine/Source/Programs/CrashReporter/CrashReportClient/CrashReportClient.Target.cs
+++ b/Engine/Source/Programs/CrashReporter/CrashReportClient/CrashReportClient.Target.cs
@@ -50,8 +50,6 @@ public class CrashReportClientTarget : TargetRules
 		ref CPPEnvironmentConfiguration OutCPPEnvironmentConfiguration
 		)
 	{
-		UEBuildConfiguration.bCompileNetworkProfiler = false;
-
 		UEBuildConfiguration.bCompileLeanAndMeanUE = true;
 
 		// Don't need editor
diff --git a/Engine/Source/Programs/CrashReporter/MinidumpDiagnostics/MinidumpDiagnostics.Target.cs b/Engine/Source/Programs/CrashReporter/MinidumpDiagnostics/MinidumpDiagnostics.Target.cs
--- a/Engine/Source/Programs/CrashReporter/MinidumpDiagnostics/MinidumpDiagnostics.Target.cs
+++ b/Engine/Source/Programs/CrashReporter/MinidumpDiagnostics/MinidumpDiagnostics.Target.cs
@@ -42,8 +42,6 @@ public class MinidumpDiagnosticsTarget : TargetRules
 	{
 		OutCPPEnvironmentConfiguration.Definitions.Add("WITH_DATABASE_SUPPORT=1");
 
-		UEBuildConfiguration.bCompileNetworkProfiler = false;
-
 		UEBuildConfiguration.bCompileLeanAndMeanUE = true;
 
 		// Don't need editor
diff --git a/Engine/Source/Programs/IOS/UnrealLaunchDaemon/UnrealLaunchDaemon.Target.cs b/Engine/Source/Programs/IOS/UnrealLaunchDaemon/UnrealLaunchDaemon.Target.cs
--- a/Engine/Source/Programs/IOS/UnrealLaunchDaemon/UnrealLaunchDaemon.Target.cs
+++ b/Engine/Source/Programs/IOS/UnrealLaunchDaemon/UnrealLaunchDaemon.Target.cs
@@ -46,7 +46,6 @@ public class UnrealLaunchDaemonTarget : TargetRules
 		ref CPPEnvironmentConfiguration OutCPPEnvironmentConfiguration
 		)
 	{
-		UEBuildConfiguration.bCompileNetworkProfiler = false;
 		UEBuildConfiguration.bCompileLeanAndMeanUE = true;
 		UEBuildConfiguration.bBuildEditor = false;
 		UEBuildConfiguration.bCompileAgainstEngine = false;
diff --git a/Engine/Source/Programs/Mac/UE4EditorServices/UE4EditorServices.Target.cs b/Engine/Source/Programs/Mac/UE4EditorServices/UE4EditorServices.Target.cs
--- a/Engine/Source/Programs/Mac/UE4EditorServices/UE4EditorServices.Target.cs
+++ b/Engine/Source/Programs/Mac/UE4EditorServices/UE4EditorServices.Target.cs
@@ -36,7 +36,6 @@ public class UE4EditorServicesTarget : TargetRules
 		ref CPPEnvironmentConfiguration OutCPPEnvironmentConfiguration
 		)
 	{
-		UEBuildConfiguration.bCompileNetworkProfiler = false;
 		UEBuildConfiguration.bCompileLeanAndMeanUE = true;
 		BuildConfiguration.bUseMallocProfiler = false;
 
diff --git a/Engine/Source/Programs/ShaderCompileWorker/ShaderCompileWorker.Target.cs b/Engine/Source/Programs/ShaderCompileWorker/ShaderCompileWorker.Target.cs
--- a/Engine/Source/Programs/ShaderCompileWorker/ShaderCompileWorker.Target.cs
+++ b/Engine/Source/Programs/ShaderCompileWorker/ShaderCompileWorker.Target.cs
@@ -48,7 +48,6 @@ public class ShaderCompileWorkerTarget : TargetRules
 		)
 	{
 		// Turn off various third party features we don't need
-		UEBuildConfiguration.bCompileNetworkProfiler = false;
 
 		// Currently we force Lean and Mean mode
 		UEBuildConfiguration.bCompileLeanAndMeanUE = true;
diff --git a/Engine/Source/Programs/SlateViewer/SlateViewer.Target.cs b/Engine/Source/Programs/SlateViewer/SlateViewer.Target.cs
--- a/Engine/Source/Programs/SlateViewer/SlateViewer.Target.cs
+++ b/Engine/Source/Programs/SlateViewer/SlateViewer.Target.cs
@@ -50,8 +50,6 @@ public class SlateViewerTarget : TargetRules
 		ref CPPEnvironmentConfiguration OutCPPEnvironmentConfiguration
 		)
 	{
-		UEBuildConfiguration.bCompileNetworkProfiler = false;
-
 		UEBuildConfiguration.bCompileLeanAndMeanUE = true;
 
 		// Don't need editor
diff --git a/Engine/Source/Programs/SymbolDebugger/SymbolDebugger.Target.cs b/Engine/Source/Programs/SymbolDebugger/SymbolDebugger.Target.cs
--- a/Engine/Source/Programs/SymbolDebugger/SymbolDebugger.Target.cs
+++ b/Engine/Source/Programs/SymbolDebugger/SymbolDebugger.Target.cs
@@ -46,8 +46,6 @@ public class SymbolDebuggerTarget : TargetRules
 	{
 		OutCPPEnvironmentConfiguration.Definitions.Add("WITH_DATABASE_SUPPORT=1");
 
-		UEBuildConfiguration.bCompileNetworkProfiler = false;
-
 		UEBuildConfiguration.bCompileLeanAndMeanUE = true;
 
 		// Don't need editor
diff --git a/Engine/Source/Programs/UnrealBuildTool/Android/UEBuildAndroid.cs b/Engine/Source/Programs/UnrealBuildTool/Android/UEBuildAndroid.cs
--- a/Engine/Source/Programs/UnrealBuildTool/Android/UEBuildAndroid.cs
+++ b/Engine/Source/Programs/UnrealBuildTool/Android/UEBuildAndroid.cs
@@ -121,7 +121,6 @@ namespace UnrealBuildTool
 			UEBuildConfiguration.bBuildEditor = false;
 			UEBuildConfiguration.bBuildDeveloperTools = false;
 			UEBuildConfiguration.bCompileSimplygon = false;
-			UEBuildConfiguration.bCompileNetworkProfiler = false;
 
 			UEBuildConfiguration.bCompileRecast = true;
 
diff --git a/Engine/Source/Programs/UnrealBuildTool/Configuration/UEBuildClient.cs b/Engine/Source/Programs/UnrealBuildTool/Configuration/UEBuildClient.cs
--- a/Engine/Source/Programs/UnrealBuildTool/Configuration/UEBuildClient.cs
+++ b/Engine/Source/Programs/UnrealBuildTool/Configuration/UEBuildClient.cs
@@ -79,15 +79,6 @@ namespace UnrealBuildTool
             ref CPPEnvironmentConfiguration OutCPPEnvironmentConfiguration
             )
         {
-            if (Target.Platform == UnrealTargetPlatform.Mac)
-            {
-                UEBuildConfiguration.bCompileNetworkProfiler = false;
-            }
-            else
-            {
-                UEBuildConfiguration.bCompileNetworkProfiler = true;
-            }
-
             UEBuildConfiguration.bCompileLeanAndMeanUE = true;
 
             // Do not include the editor
diff --git a/Engine/Source/Programs/UnrealBuildTool/Configuration/UEBuildConfiguration.cs b/Engine/Source/Programs/UnrealBuildTool/Configuration/UEBuildConfiguration.cs
--- a/Engine/Source/Programs/UnrealBuildTool/Configuration/UEBuildConfiguration.cs
+++ b/Engine/Source/Programs/UnrealBuildTool/Configuration/UEBuildConfiguration.cs
@@ -16,9 +16,6 @@ namespace UnrealBuildTool
         /** Whether to include ICU unicode/i18n support in core */
         public static bool bCompileICU;
 
-		/** Whether we should compile the network profiler or not.*/
-		public static bool bCompileNetworkProfiler;
-
 		/** Whether to build a stripped down version of the game specifically for dedicated server. */
 		public static bool bBuildDedicatedServer;
 
@@ -119,7 +116,6 @@ namespace UnrealBuildTool
 			// Currently, WITH_PHYSX is forced to true in Engine.h (as it isn't defined anywhere by the builder)
 			bCompilePhysX = true;
 			bCompileAPEX = true;
-			bCompileNetworkProfiler = false;
 			bBuildDedicatedServer = false;
 			bBuildEditor = true;
 			bBuildRequiresCookedData = false;
@@ -179,7 +175,6 @@ namespace UnrealBuildTool
 				bBuildEditor = false;
 				bBuildDeveloperTools = false;
 				bCompileSimplygon = false;
-				bCompileNetworkProfiler = false;
 				bCompileSpeedTree = false;
 			}
 		}
diff --git a/Engine/Source/Programs/UnrealBuildTool/Configuration/UEBuildEditor.cs b/Engine/Source/Programs/UnrealBuildTool/Configuration/UEBuildEditor.cs
--- a/Engine/Source/Programs/UnrealBuildTool/Configuration/UEBuildEditor.cs
+++ b/Engine/Source/Programs/UnrealBuildTool/Configuration/UEBuildEditor.cs
@@ -74,15 +74,6 @@ namespace UnrealBuildTool
 			ref CPPEnvironmentConfiguration OutCPPEnvironmentConfiguration
 			)
 		{
-			if (Target.Platform == UnrealTargetPlatform.Mac)
-			{
-				UEBuildConfiguration.bCompileNetworkProfiler = false;
-			}
-			else
-			{
-				UEBuildConfiguration.bCompileNetworkProfiler = true;
-			}
-
 			UEBuildConfiguration.bCompileLeanAndMeanUE = false;
 
 			// Do not include the editor
diff --git a/Engine/Source/Programs/UnrealBuildTool/Configuration/UEBuildGame.cs b/Engine/Source/Programs/UnrealBuildTool/Configuration/UEBuildGame.cs
--- a/Engine/Source/Programs/UnrealBuildTool/Configuration/UEBuildGame.cs
+++ b/Engine/Source/Programs/UnrealBuildTool/Configuration/UEBuildGame.cs
@@ -94,15 +94,6 @@ namespace UnrealBuildTool
 			ref CPPEnvironmentConfiguration OutCPPEnvironmentConfiguration
 			)
 		{
-			if (Target.Platform == UnrealTargetPlatform.Mac)
-			{
-				UEBuildConfiguration.bCompileNetworkProfiler = false;
-			}
-			else
-			{
-				UEBuildConfiguration.bCompileNetworkProfiler = true;
-			}
-
 			UEBuildConfiguration.bCompileLeanAndMeanUE = true;
 
 			// Do not include the editor
diff --git a/Engine/Source/Programs/UnrealBuildTool/Configuration/UEBuildServer.cs b/Engine/Source/Programs/UnrealBuildTool/Configuration/UEBuildServer.cs
--- a/Engine/Source/Programs/UnrealBuildTool/Configuration/UEBuildServer.cs
+++ b/Engine/Source/Programs/UnrealBuildTool/Configuration/UEBuildServer.cs
@@ -74,14 +74,6 @@ namespace UnrealBuildTool
 		{
 			// Dedicated server
 			UEBuildConfiguration.bBuildDedicatedServer = true;
-			if (Target.Platform == UnrealTargetPlatform.Mac)
-			{
-				UEBuildConfiguration.bCompileNetworkProfiler = false;
-			}
-			else
-			{
-				UEBuildConfiguration.bCompileNetworkProfiler = true;
-			}
 
 			UEBuildConfiguration.bCompileLeanAndMeanUE = true;
 
diff --git a/Engine/Source/Programs/UnrealBuildTool/IOS/UEBuildIOS.cs b/Engine/Source/Programs/UnrealBuildTool/IOS/UEBuildIOS.cs
--- a/Engine/Source/Programs/UnrealBuildTool/IOS/UEBuildIOS.cs
+++ b/Engine/Source/Programs/UnrealBuildTool/IOS/UEBuildIOS.cs
@@ -160,7 +160,6 @@ namespace UnrealBuildTool
 			UEBuildConfiguration.bBuildDeveloperTools = false;
 			UEBuildConfiguration.bCompileAPEX = false;
 			UEBuildConfiguration.bCompileSimplygon = false;
-			UEBuildConfiguration.bCompileNetworkProfiler = false;
 			UEBuildConfiguration.bBuildDeveloperTools = false;
 
 			// we currently don't have any simulator libs for PhysX
diff --git a/Engine/Source/Programs/UnrealBuildTool/Linux/UEBuildLinux.cs b/Engine/Source/Programs/UnrealBuildTool/Linux/UEBuildLinux.cs
--- a/Engine/Source/Programs/UnrealBuildTool/Linux/UEBuildLinux.cs
+++ b/Engine/Source/Programs/UnrealBuildTool/Linux/UEBuildLinux.cs
@@ -209,7 +209,6 @@ namespace UnrealBuildTool.Linux
             UEBuildConfiguration.bBuildEditor = false;
             UEBuildConfiguration.bBuildDeveloperTools = false;
             UEBuildConfiguration.bCompileSimplygon = false;
-            UEBuildConfiguration.bCompileNetworkProfiler = false;
 
             // Don't stop compilation at first error...
             BuildConfiguration.bStopXGECompilationAfterErrors = true;
diff --git a/Engine/Source/Programs/UnrealBuildTool/Mac/UEBuildMac.cs b/Engine/Source/Programs/UnrealBuildTool/Mac/UEBuildMac.cs
--- a/Engine/Source/Programs/UnrealBuildTool/Mac/UEBuildMac.cs
+++ b/Engine/Source/Programs/UnrealBuildTool/Mac/UEBuildMac.cs
@@ -145,7 +145,6 @@ namespace UnrealBuildTool
 		public override void ResetBuildConfiguration(UnrealTargetPlatform InPlatform, UnrealTargetConfiguration InConfiguration)
 		{
 			UEBuildConfiguration.bCompileSimplygon = false;
-			UEBuildConfiguration.bCompileNetworkProfiler = false;
 			UEBuildConfiguration.bCompileICU = true;
 		}
 
diff --git a/Engine/Source/Programs/UnrealFileServer/UnrealFileServer.Target.cs b/Engine/Source/Programs/UnrealFileServer/UnrealFileServer.Target.cs
--- a/Engine/Source/Programs/UnrealFileServer/UnrealFileServer.Target.cs
+++ b/Engine/Source/Programs/UnrealFileServer/UnrealFileServer.Target.cs
@@ -37,8 +37,6 @@ public class UnrealFileServerTarget : TargetRules
 		ref CPPEnvironmentConfiguration OutCPPEnvironmentConfiguration
 		)
 	{
-		UEBuildConfiguration.bCompileNetworkProfiler = false;
-
 		// Lean and mean
 		UEBuildConfiguration.bCompileLeanAndMeanUE = true;
 
diff --git a/Engine/Source/Programs/UnrealFrontend/UnrealFrontend.Target.cs b/Engine/Source/Programs/UnrealFrontend/UnrealFrontend.Target.cs
--- a/Engine/Source/Programs/UnrealFrontend/UnrealFrontend.Target.cs
+++ b/Engine/Source/Programs/UnrealFrontend/UnrealFrontend.Target.cs
@@ -43,7 +43,6 @@ public class UnrealFrontendTarget : TargetRules
 		UEBuildConfiguration.bCompileAgainstEngine = false;
 		UEBuildConfiguration.bCompileAgainstCoreUObject = true;
 		UEBuildConfiguration.bCompileLeanAndMeanUE = true;
-		UEBuildConfiguration.bCompileNetworkProfiler = false;
 		UEBuildConfiguration.bForceBuildTargetPlatforms = true;
 		UEBuildConfiguration.bCompileWithStatsWithoutEngine = true;
 		UEBuildConfiguration.bCompileWithPluginSupport = true;
diff --git a/Engine/Source/Programs/UnrealHeaderTool/UnrealHeaderTool.Target.cs b/Engine/Source/Programs/UnrealHeaderTool/UnrealHeaderTool.Target.cs
--- a/Engine/Source/Programs/UnrealHeaderTool/UnrealHeaderTool.Target.cs
+++ b/Engine/Source/Programs/UnrealHeaderTool/UnrealHeaderTool.Target.cs
@@ -42,8 +42,6 @@ public class UnrealHeaderToolTarget : TargetRules
 		ref CPPEnvironmentConfiguration OutCPPEnvironmentConfiguration
 		)
 	{
-		UEBuildConfiguration.bCompileNetworkProfiler = false;
-
 		// Lean and mean
 		UEBuildConfiguration.bCompileLeanAndMeanUE = true;
 
diff --git a/Engine/Source/Programs/UnrealPak/UnrealPak.Target.cs b/Engine/Source/Programs/UnrealPak/UnrealPak.Target.cs
--- a/Engine/Source/Programs/UnrealPak/UnrealPak.Target.cs
+++ b/Engine/Source/Programs/UnrealPak/UnrealPak.Target.cs
@@ -36,7 +36,6 @@ public class UnrealPakTarget : TargetRules
 		ref CPPEnvironmentConfiguration OutCPPEnvironmentConfiguration
 		)
 	{
-		UEBuildConfiguration.bCompileNetworkProfiler = false;
 		UEBuildConfiguration.bCompileLeanAndMeanUE = true;
 		BuildConfiguration.bUseMallocProfiler = false;
 
diff --git a/Engine/Source/Programs/UnrealSync/UnrealSync.Target.cs b/Engine/Source/Programs/UnrealSync/UnrealSync.Target.cs
--- a/Engine/Source/Programs/UnrealSync/UnrealSync.Target.cs
+++ b/Engine/Source/Programs/UnrealSync/UnrealSync.Target.cs
@@ -37,8 +37,6 @@ public class UnrealSyncTarget : TargetRules
 		ref CPPEnvironmentConfiguration OutCPPEnvironmentConfiguration
 		)
 	{
-		UEBuildConfiguration.bCompileNetworkProfiler = false;
-
 		// Lean and mean
 		UEBuildConfiguration.bCompileLeanAndMeanUE = true;
 
diff --git a/Engine/Source/Programs/UnrealVersionSelector/UnrealVersionSelector.Target.cs b/Engine/Source/Programs/UnrealVersionSelector/UnrealVersionSelector.Target.cs
--- a/Engine/Source/Programs/UnrealVersionSelector/UnrealVersionSelector.Target.cs
+++ b/Engine/Source/Programs/UnrealVersionSelector/UnrealVersionSelector.Target.cs
@@ -36,7 +36,6 @@ public class UnrealVersionSelectorTarget : TargetRules
 		ref CPPEnvironmentConfiguration OutCPPEnvironmentConfiguration
 		)
 	{
-		UEBuildConfiguration.bCompileNetworkProfiler = false;
 		UEBuildConfiguration.bCompileLeanAndMeanUE = true;
 		BuildConfiguration.bUseMallocProfiler = false;
 
diff --git a/Engine/Source/Runtime/Core/Public/Misc/Build.h b/Engine/Source/Runtime/Core/Public/Misc/Build.h
--- a/Engine/Source/Runtime/Core/Public/Misc/Build.h
+++ b/Engine/Source/Runtime/Core/Public/Misc/Build.h
@@ -231,4 +231,7 @@
  **/
 #ifndef LOOKING_FOR_PERF_ISSUES
 	#define LOOKING_FOR_PERF_ISSUES (0 && !(UE_BUILD_SHIPPING))
-#endif
\ No newline at end of file
+#endif
+
+/** Enable the use of the network profiler as long as we are a build that includes stats */
+#define USE_NETWORK_PROFILER         STATS
\ No newline at end of file
diff --git a/Engine/Source/Runtime/Engine/Classes/Landscape/Landscape.h b/Engine/Source/Runtime/Engine/Classes/Landscape/Landscape.h
--- a/Engine/Source/Runtime/Engine/Classes/Landscape/Landscape.h
+++ b/Engine/Source/Runtime/Engine/Classes/Landscape/Landscape.h
@@ -35,7 +35,7 @@ class ALandscape : public ALandscapeProxy
 	// End AActor Interface
 
 	// Begin ALandscapeProxy Interface
-	virtual ALandscape* GetLandscapeActor() OVERRIDE;
+	ENGINE_API virtual ALandscape* GetLandscapeActor() OVERRIDE;
 #if WITH_EDITOR
 	virtual UMaterialInterface* GetLandscapeMaterial() const OVERRIDE;
 	virtual UMaterialInterface* GetLandscapeHoleMaterial() const OVERRIDE;
diff --git a/Engine/Source/Runtime/Engine/Engine.Build.cs b/Engine/Source/Runtime/Engine/Engine.Build.cs
--- a/Engine/Source/Runtime/Engine/Engine.Build.cs
+++ b/Engine/Source/Runtime/Engine/Engine.Build.cs
@@ -170,11 +170,6 @@ public class Engine : ModuleRules
 
 		SetupModulePhysXAPEXSupport(Target);
 
-		if (UEBuildConfiguration.bCompileNetworkProfiler)
-		{
-			Definitions.Add("USE_NETWORK_PROFILER=1");
-		}
-
 		if ((Target.Platform == UnrealTargetPlatform.Win64) ||
 			(Target.Platform == UnrealTargetPlatform.Win32))
 		{
diff --git a/Engine/Source/Runtime/Engine/Private/NetworkProfiler.cpp b/Engine/Source/Runtime/Engine/Private/NetworkProfiler.cpp
--- a/Engine/Source/Runtime/Engine/Private/NetworkProfiler.cpp
+++ b/Engine/Source/Runtime/Engine/Private/NetworkProfiler.cpp
@@ -17,10 +17,6 @@
 #define NETWORK_PROFILER_TRACK_RAW_NETWORK_DATA		0
 
 #include "Net/UnrealNetwork.h"
-#if PLATFORM_MAC
-	#include <netinet/in.h>
-	#include <arpa/inet.h>
-#endif
 #include "IPAddress.h"
 #include "Net/NetworkProfiler.h"
 
