fix for bad path being passed to GetFullPath
#ue4
#ios on pc

diff --git a/Engine/Source/Programs/UnrealBuildTool/ToolChain/RemoteToolChain.cs b/Engine/Source/Programs/UnrealBuildTool/ToolChain/RemoteToolChain.cs
--- a/Engine/Source/Programs/UnrealBuildTool/ToolChain/RemoteToolChain.cs
+++ b/Engine/Source/Programs/UnrealBuildTool/ToolChain/RemoteToolChain.cs
@@ -181,14 +181,15 @@ namespace UnrealBuildTool
             int RootDirLocation = 0;
 
 			string LocalUnrealEngineRootPath = Path.GetFullPath(AppDomain.CurrentDomain.BaseDirectory + "..\\..\\..\\");
-			string LocalProjectRootPath = Path.GetFullPath(UnrealBuildTool.GetUProjectPath());
+			string UProjectPath = UnrealBuildTool.GetUProjectPath();
+			string LocalProjectRootPath = String.IsNullOrEmpty(UProjectPath) ? "" : Path.GetFullPath(UProjectPath);
 			if (LocalPath.ToUpperInvariant().Contains(LocalUnrealEngineRootPath.ToUpperInvariant()))
 			{
 				// If the file is a full path name and rooted at the same location as UBT,
 				// use that location as the root and simply return the length
 				RootDirLocation = LocalUnrealEngineRootPath.Length;
 			}
-			else if (LocalPath.ToUpperInvariant().Contains(LocalProjectRootPath.ToUpperInvariant()))
+			else if (!String.IsNullOrEmpty(UProjectPath) && LocalPath.ToUpperInvariant().Contains(LocalProjectRootPath.ToUpperInvariant()))
 			{
 				RootDirLocation = LocalProjectRootPath.Length;
 			}
