Fix for hot reload.

#ttp CRITICAL: Regression: EDITOR: Game Hot Recompile Fails

#change UEBuildModule.SourceFiles renamed to SourceFilesToBuild.
#add UEBuildModule.SourceFilesFound which are used for PCH determination.
#change Source file categorization by extension moved into its own function.
#fix Source files always discovered (except during project generation) and passed to the module so that the found versions can always be found, but the built versions only exist if we need them to.
#change Spelling error fixed.

diff --git a/Engine/Source/Programs/UnrealBuildTool/Configuration/UEBuildModule.cs b/Engine/Source/Programs/UnrealBuildTool/Configuration/UEBuildModule.cs
--- a/Engine/Source/Programs/UnrealBuildTool/Configuration/UEBuildModule.cs
+++ b/Engine/Source/Programs/UnrealBuildTool/Configuration/UEBuildModule.cs
@@ -806,8 +806,11 @@ namespace UnrealBuildTool
 			}
 		}
 
-		/** A list of the absolute paths of source files in the module. */
-		public readonly SourceFilesClass SourceFiles = new SourceFilesClass();
+		/** A list of the absolute paths of source files to be built in this module. */
+		public readonly SourceFilesClass SourceFilesToBuild = new SourceFilesClass();
+
+		/** A list of the source files that were found for the module. */
+		public readonly SourceFilesClass SourceFilesFound = new SourceFilesClass();
 
 		/** The preprocessor definitions used to compile this module's private implementation. */
 		List<string> Definitions;
@@ -859,6 +862,43 @@ namespace UnrealBuildTool
 		/** The processed dependencies for the class */
 		public ProcessedDependenciesClass ProcessedDependencies = null;
 
+		/** Categorizes source files into per-extension buckets */
+		private static void CategorizeSourceFiles(IEnumerable<FileItem> InSourceFiles, SourceFilesClass OutSourceFiles)
+		{
+			foreach (var SourceFile in InSourceFiles)
+			{
+				string Extension = Path.GetExtension(SourceFile.AbsolutePath).ToUpperInvariant();
+				if (!SourceFile.bExists)
+				{
+					OutSourceFiles.MissingFiles.Add(SourceFile);
+				}
+				else if (Extension == ".CPP")
+				{
+					OutSourceFiles.CPPFiles.Add(SourceFile);
+				}
+				else if (Extension == ".C")
+				{
+					OutSourceFiles.CFiles.Add(SourceFile);
+				}
+				else if (Extension == ".CC")
+				{
+					OutSourceFiles.CCFiles.Add(SourceFile);
+				}
+				else if (Extension == ".MM" || Extension == ".M")
+				{
+					OutSourceFiles.MMFiles.Add(SourceFile);
+				}
+				else if (Extension == ".RC")
+				{
+					OutSourceFiles.RCFiles.Add(SourceFile);
+				}
+				else
+				{
+					OutSourceFiles.OtherFiles.Add(SourceFile);
+				}
+			}
+		}
+
 		public UEBuildModuleCPP(
 			UEBuildTarget InTarget,
 			string InName,
@@ -891,7 +931,8 @@ namespace UnrealBuildTool
 			bool InFasterWithoutUnity,
 			int InMinFilesUsingPrecompiledHeaderOverride,
 			bool InEnableExceptions,
-			bool InEnableInlining
+			bool InEnableInlining,
+			bool bInBuildSourceFiles
 			)
 			: base(	InTarget,
 					InName, 
@@ -919,37 +960,10 @@ namespace UnrealBuildTool
 		{
 			IntelliSenseGatherer = InIntelliSenseGatherer;
 
-			foreach (var SourceFile in InSourceFiles)
+			CategorizeSourceFiles(InSourceFiles, SourceFilesFound);
+			if (bInBuildSourceFiles)
 			{
-				string Extension = Path.GetExtension(SourceFile.AbsolutePath).ToUpperInvariant();
-				if (!SourceFile.bExists)
-				{
-					SourceFiles.MissingFiles.Add(SourceFile);
-				}
-				else if (Extension == ".CPP")
-				{
-					SourceFiles.CPPFiles.Add(SourceFile);
-				}
-				else if (Extension == ".C")
-				{
-					SourceFiles.CFiles.Add(SourceFile);
-				}
-				else if (Extension == ".CC")
-				{
-					SourceFiles.CCFiles.Add(SourceFile);
-				}
-				else if (Extension == ".MM" || Extension == ".M")
-				{
-					SourceFiles.MMFiles.Add(SourceFile);
-				}
-				else if (Extension == ".RC")
-				{
-					SourceFiles.RCFiles.Add(SourceFile);
-				}
-				else
-				{
-					SourceFiles.OtherFiles.Add(SourceFile);
-				}
+				CategorizeSourceFiles(InSourceFiles, SourceFilesToBuild);
 			}
 
 			Definitions = ListFromOptionalEnumerableStringParameter(InDefinitions);
@@ -1000,12 +1014,12 @@ namespace UnrealBuildTool
 			}
 
 			// Throw an error if the module's source file list referenced any non-existent files.
-			if (SourceFiles.MissingFiles.Count > 0)
+			if (SourceFilesToBuild.MissingFiles.Count > 0)
 			{
 				throw new BuildException(
 					"UBT ERROR: Module \"{0}\" references non-existent files:\n{1} (perhaps a file was added to the project but not checked in)",
 					Name,
-					string.Join("\n", SourceFiles.MissingFiles.Select(M => M.AbsolutePath))
+					string.Join("\n", SourceFilesToBuild.MissingFiles.Select(M => M.AbsolutePath))
 				);
 			}
 
@@ -1016,11 +1030,11 @@ namespace UnrealBuildTool
 			if (Binary.Config.Type == UEBuildBinaryType.DynamicLinkLibrary)
 			{
 				// Add default PCLaunch.rc file if this module has no own resource file specified
-				if (SourceFiles.RCFiles.Count <= 0)
+				if (SourceFilesToBuild.RCFiles.Count <= 0)
 				{
 					string DefRC = Path.Combine(Directory.GetCurrentDirectory(), "Runtime/Launch/Resources/Windows/PCLaunch.rc");
 					FileItem Item = FileItem.GetItemByFullPath(DefRC);
-					SourceFiles.RCFiles.Add(Item);
+					SourceFilesToBuild.RCFiles.Add(Item);
 				}
 			}
 
@@ -1050,7 +1064,7 @@ namespace UnrealBuildTool
 			{
 				bModuleUsesUnityBuild = false;
 			}
-			else if( !BuildConfiguration.bForceUnityBuild && IsGameModule && SourceFiles.CPPFiles.Count < BuildConfiguration.MinGameModuleSourceFilesForUnityBuild )
+			else if( !BuildConfiguration.bForceUnityBuild && IsGameModule && SourceFilesToBuild.CPPFiles.Count < BuildConfiguration.MinGameModuleSourceFilesForUnityBuild )
 			{
 				// Game modules with only a small number of source files are usually better off having faster iteration times
 				// on single source file changes, so we forcibly disable unity build for those modules
@@ -1088,7 +1102,7 @@ namespace UnrealBuildTool
 
 				// Determine what potential precompiled header is used by each source file.
 				double SharedPCHTotalTime = 0.0;
-				foreach( var CPPFile in SourceFiles.CPPFiles )
+				foreach( var CPPFile in SourceFilesToBuild.CPPFiles )
 				{
 					if (bUseSharedPCHFiles)
 					{
@@ -1212,13 +1226,13 @@ namespace UnrealBuildTool
 
 				// If there was one header that was included first by enough C++ files, use it as the precompiled header.
 				// Only use precompiled headers for projects with enough files to make the PCH creation worthwhile.
-				if( SharedPCHHeaderFile != null || SourceFiles.CPPFiles.Count >= MinFilesUsingPrecompiledHeader )
+				if( SharedPCHHeaderFile != null || SourceFilesToBuild.CPPFiles.Count >= MinFilesUsingPrecompiledHeader )
 				{
 					FileItem PCHToUse;
 
 					if( SharedPCHHeaderFile != null )
 					{
-						ModulePCHEnvironment = ApplySharedPCH(GlobalCompileEnvironment, CompileEnvironment, ModuleCompileEnvironment, SourceFiles.CPPFiles, ref SharedPCHHeaderFile);
+						ModulePCHEnvironment = ApplySharedPCH(GlobalCompileEnvironment, CompileEnvironment, ModuleCompileEnvironment, SourceFilesToBuild.CPPFiles, ref SharedPCHHeaderFile);
 						if (ModulePCHEnvironment != null)
 						{
 							// @todo SharedPCH: Ideally we would exhaustively check for a compatible compile environment (definitions, imports/exports, etc)
@@ -1238,7 +1252,7 @@ namespace UnrealBuildTool
 					if (PCHToUse != null)
 					{
 						// Update all CPPFiles to point to the PCH
-						foreach (var CPPFile in SourceFiles.CPPFiles)
+						foreach (var CPPFile in SourceFilesToBuild.CPPFiles)
 						{
 							CPPFile.PCHHeaderNameInCode              = PCHToUse.AbsolutePath;
 							CPPFile.PrecompiledHeaderIncludeFilename = PCHToUse.AbsolutePath;
@@ -1255,7 +1269,7 @@ namespace UnrealBuildTool
 							PCHHeaderFile = SharedPCHHeaderFile;
 							PCHModuleName = SharedPCHModuleName;
 						}
-						var PCHHeaderNameInCode = SourceFiles.CPPFiles[ 0 ].PCHHeaderNameInCode;
+						var PCHHeaderNameInCode = SourceFilesToBuild.CPPFiles[ 0 ].PCHHeaderNameInCode;
 
 						ModulePCHEnvironment = new PrecompileHeaderEnvironment( PCHModuleName, PCHHeaderNameInCode, PCHHeaderFile, ModuleCompileEnvironment.Config.CLRMode, ModuleCompileEnvironment.Config.OptimizeCode );
 
@@ -1292,7 +1306,7 @@ namespace UnrealBuildTool
 						ModulePCHCompileEnvironment.Config.bForceIncludePrecompiledHeader = true;
 					}
 
-					var CPPFilesToBuild = SourceFiles.CPPFiles;
+					var CPPFilesToBuild = SourceFilesToBuild.CPPFiles;
 					if (bModuleUsesUnityBuild)
 					{
 						// unity files generated for only the set of files which share the same PCH environment
@@ -1354,9 +1368,9 @@ namespace UnrealBuildTool
 				}
 			}
 
-			if( !bWasModuleCodeCompiled && SourceFiles.CPPFiles.Count > 0 )
+			if( !bWasModuleCodeCompiled && SourceFilesToBuild.CPPFiles.Count > 0 )
 			{
-				var CPPFilesToCompile = SourceFiles.CPPFiles;
+				var CPPFilesToCompile = SourceFilesToBuild.CPPFiles;
 				if (bModuleUsesUnityBuild)
 				{
 					CPPFilesToCompile = Unity.GenerateUnityCPPs( CPPFilesToCompile, ModuleCompileEnvironment, Name );
@@ -1372,13 +1386,13 @@ namespace UnrealBuildTool
 			}
 
 			// Compile C files directly.
-			LinkInputFiles.AddRange(ModuleCompileEnvironment.CompileFiles(SourceFiles.CFiles, Name).ObjectFiles);
+			LinkInputFiles.AddRange(ModuleCompileEnvironment.CompileFiles(SourceFilesToBuild.CFiles, Name).ObjectFiles);
 
 			// Compile CC files directly.
-			LinkInputFiles.AddRange(ModuleCompileEnvironment.CompileFiles(SourceFiles.CCFiles, Name).ObjectFiles);
+			LinkInputFiles.AddRange(ModuleCompileEnvironment.CompileFiles(SourceFilesToBuild.CCFiles, Name).ObjectFiles);
 
 			// Compile MM files directly.
-			LinkInputFiles.AddRange(ModuleCompileEnvironment.CompileFiles(SourceFiles.MMFiles, Name).ObjectFiles);
+			LinkInputFiles.AddRange(ModuleCompileEnvironment.CompileFiles(SourceFilesToBuild.MMFiles, Name).ObjectFiles);
 
 			// If we're building Rocket, generate a static library for this module
 			if(RedistStaticLibraryPath != null)
@@ -1400,7 +1414,7 @@ namespace UnrealBuildTool
 			}
 
 			// Compile RC files.
-			LinkInputFiles.AddRange(ModuleCompileEnvironment.CompileRCFiles(SourceFiles.RCFiles).ObjectFiles);
+			LinkInputFiles.AddRange(ModuleCompileEnvironment.CompileRCFiles(SourceFilesToBuild.RCFiles).ObjectFiles);
 
 			return LinkInputFiles;
 		}
@@ -1454,7 +1468,7 @@ namespace UnrealBuildTool
 				return;
 
 			FileItem UniquePCH = null;
-			foreach( var CPPFile in SourceFiles.CPPFiles )
+			foreach( var CPPFile in SourceFilesFound.CPPFiles )
 			{
 				// Find headers used by the source file.
 				var PCH = ProcessDependencies(CPPFile, ModuleCompileEnvironment);
@@ -1909,7 +1923,8 @@ namespace UnrealBuildTool
 			bool InFasterWithoutUnity,
 			int InMinFilesUsingPrecompiledHeaderOverride,
 			bool InEnableExceptions,
-			bool InEnableInlining
+			bool InEnableInlining,
+			bool bInBuildSourceFiles
 			)
 		: base(InTarget,InName,InType,InModuleDirectory,InOutputDirectory,InIntelliSenseGatherer,
 			InSourceFiles,InPublicIncludePaths,InPublicSystemIncludePaths,InDefinitions,
@@ -1917,7 +1932,7 @@ namespace UnrealBuildTool
 			InPrivateIncludePaths,InPrivateIncludePathModuleNames,InPrivateDependencyModuleNames,
             InCircularlyReferencedDependentModules, InDynamicallyLoadedModuleNames, InPlatformSpecificDynamicallyLoadedModuleNames, InOptimizeCode,
 			InAllowSharedPCH, InSharedPCHHeaderFile, InUseRTTI, InEnableBufferSecurityChecks, InFasterWithoutUnity, InMinFilesUsingPrecompiledHeaderOverride,
-			InEnableExceptions, InEnableInlining)
+			InEnableExceptions, InEnableInlining, bInBuildSourceFiles)
 		{
 			PrivateAssemblyReferences = ListFromOptionalEnumerableStringParameter(InPrivateAssemblyReferences);
 		}
diff --git a/Engine/Source/Programs/UnrealBuildTool/Configuration/UEBuildTarget.cs b/Engine/Source/Programs/UnrealBuildTool/Configuration/UEBuildTarget.cs
--- a/Engine/Source/Programs/UnrealBuildTool/Configuration/UEBuildTarget.cs
+++ b/Engine/Source/Programs/UnrealBuildTool/Configuration/UEBuildTarget.cs
@@ -822,7 +822,7 @@ namespace UnrealBuildTool
 		}
 
 		/// <summary>
-		/// Cleans UnralHeaderTool
+		/// Cleans UnrealHeaderTool
 		/// </summary>
 		private void CleanUnrealHeaderTool()
 		{
@@ -1400,7 +1400,7 @@ namespace UnrealBuildTool
 
 						// If we've got this far and there are no source files then it's likely we're running Rocket and ignoring
 						// engine files, so we don't need a .generated.cpp either
-						if (DependencyModuleCPP.SourceFiles.Count != 0)
+						if (DependencyModuleCPP.SourceFilesToBuild.Count != 0)
 						{
 							DependencyModuleCPP.AutoGenerateCppInfo = new UEBuildModuleCPP.AutoGenerateCppInfoClass( UHTModuleInfo.GeneratedCPPFilenameBase + ".cpp" );
 						}
@@ -1583,7 +1583,8 @@ namespace UnrealBuildTool
 					InFasterWithoutUnity: true,
 					InMinFilesUsingPrecompiledHeaderOverride: 0,
 					InEnableExceptions: false,
-					InEnableInlining: true
+					InEnableInlining: true,
+					bInBuildSourceFiles: true
 					);
 
 				// Now bind this new module to the executable binary so it will link the plugin libs correctly
@@ -2341,8 +2342,6 @@ namespace UnrealBuildTool
 				TargetInfo TargetInfo = GetTargetInfo();
 
 				// Create the module!  (It will be added to our hash table in its constructor)
-				// Figure out whether we need to build this module
-				bool bBuildFiles = OnlyModules.Count == 0 || OnlyModules.Any(x => x.OnlyModuleName == ModuleName);
 
 				// @todo projectfiles: Cross-platform modules can appear here during project generation, but they may have already
 				//   been filtered out by the project generator.  This causes the projects to not be added to directories properly.
@@ -2403,8 +2402,13 @@ namespace UnrealBuildTool
 					RulesCompiler.AddDefaultIncludePathsToModuleRules(this, ModuleName, ModuleFileName, ModuleFileRelativeToEngineDirectory, IsGameModule: IsGameModule, RulesObject: ref RulesObject);
 				}
 
+				// Figure out whether we need to build this module
+				// We don't care about actual source files when generating projects, as these are discovered separately
+				bool bDiscoverFiles = !ProjectFileGenerator.bGenerateProjectFiles;
+				bool bBuildFiles    = bDiscoverFiles && (OnlyModules.Count == 0 || OnlyModules.Any(x => x.OnlyModuleName == ModuleName));
+
 				IntelliSenseGatherer IntelliSenseGatherer = null;
-				List<FileItem> ModuleSourceFiles = new List<FileItem>();
+				List<FileItem> FoundSourceFiles = new List<FileItem>();
 				if (RulesObject.Type == ModuleRules.ModuleType.CPlusPlus || RulesObject.Type == ModuleRules.ModuleType.CPlusPlusCLR)
 				{
 					ProjectFile ProjectFile = null;
@@ -2413,35 +2417,33 @@ namespace UnrealBuildTool
 						IntelliSenseGatherer = ProjectFile;
 					}
 
-					if (!ProjectFileGenerator.bGenerateProjectFiles && bBuildFiles)	// We don't care about actual source files when generating projects, as these are discovered separately
+					// So all we care about are the game module and/or plugins.
+					//@todo Rocket: This assumes plugins that have source will be under the game folder...
+					if (bDiscoverFiles && (!UnrealBuildTool.RunningRocket() || Utils.IsFileUnderDirectory(ModuleFileName, UnrealBuildTool.GetUProjectPath())))
 					{
-						// So all we care about are the game module and/or plugins.
-						//@todo Rocket: This assumes plugins that have source will be under the game folder...
-						if (!UnrealBuildTool.RunningRocket() || Utils.IsFileUnderDirectory(ModuleFileName, UnrealBuildTool.GetUProjectPath()))
-						{
-							var SourceFilePaths = new List<string>();
+						var SourceFilePaths = new List<string>();
 
-							if (ProjectFile != null)
-							{
-								foreach (var SourceFile in ProjectFile.SourceFiles)
-								{
-									SourceFilePaths.Add(SourceFile.FilePath);
-								}
-							}
-							else
+						if (ProjectFile != null)
+						{
+							foreach (var SourceFile in ProjectFile.SourceFiles)
 							{
-								// Don't have a project file for this module with the source file names cached already, so find the source files ourselves
-								SourceFilePaths = SourceFileSearch.FindModuleSourceFiles(
-									ModuleRulesFile: ModuleFileName,
-									ExcludeNoRedistFiles: false);
+								SourceFilePaths.Add(SourceFile.FilePath);
 							}
-							ModuleSourceFiles = GetCPlusPlusFilesToBuild(SourceFilePaths, ModuleDirectory, Platform);
 						}
+						else
+						{
+							// Don't have a project file for this module with the source file names cached already, so find the source files ourselves
+							SourceFilePaths = SourceFileSearch.FindModuleSourceFiles(
+								ModuleRulesFile: ModuleFileName,
+								ExcludeNoRedistFiles: false);
+						}
+						FoundSourceFiles = GetCPlusPlusFilesToBuild(SourceFilePaths, ModuleDirectory, Platform);
 					}
 				}
 
 				// Now, go ahead and create the module builder instance
-				if((Module = InstantiateModule(RulesObject, ModuleName, ModuleType, ModuleDirectory, ApplicationOutputPath, IntelliSenseGatherer, ModuleSourceFiles)) == null)
+				Module = InstantiateModule(RulesObject, ModuleName, ModuleType, ModuleDirectory, ApplicationOutputPath, IntelliSenseGatherer, FoundSourceFiles, bBuildFiles);
+				if(Module == null)
 				{
 					throw new BuildException("Unrecognized module type specified by 'Rules' object {0}", RulesObject.ToString());
 				}
@@ -2462,13 +2464,14 @@ namespace UnrealBuildTool
 		}
 
 		protected virtual UEBuildModule InstantiateModule(
-			ModuleRules RulesObject,
-			string ModuleName,
-			UEBuildModuleType ModuleType,
-			string ModuleDirectory,
-			string ApplicationOutputPath,
+			ModuleRules          RulesObject,
+			string               ModuleName,
+			UEBuildModuleType    ModuleType,
+			string               ModuleDirectory,
+			string               ApplicationOutputPath,
 			IntelliSenseGatherer IntelliSenseGatherer,
-			List<FileItem> ModuleSourceFiles)
+			List<FileItem>       ModuleSourceFiles,
+			bool                 bBuildSourceFiles)
 		{
 			switch (RulesObject.Type)
 			{
@@ -2505,7 +2508,8 @@ namespace UnrealBuildTool
 							InFasterWithoutUnity: RulesObject.bFasterWithoutUnity,
 							InMinFilesUsingPrecompiledHeaderOverride: RulesObject.MinFilesUsingPrecompiledHeaderOverride,
 							InEnableExceptions: RulesObject.bEnableExceptions,
-							InEnableInlining: RulesObject.bEnableInlining
+							InEnableInlining: RulesObject.bEnableInlining,
+							bInBuildSourceFiles: bBuildSourceFiles
 						);
 
 				case ModuleRules.ModuleType.CPlusPlusCLR:
@@ -2542,7 +2546,8 @@ namespace UnrealBuildTool
 							InFasterWithoutUnity: RulesObject.bFasterWithoutUnity,
 							InMinFilesUsingPrecompiledHeaderOverride: RulesObject.MinFilesUsingPrecompiledHeaderOverride,
 							InEnableExceptions: RulesObject.bEnableExceptions,
-							InEnableInlining: RulesObject.bEnableInlining
+							InEnableInlining: RulesObject.bEnableInlining,
+							bInBuildSourceFiles : bBuildSourceFiles
 						);
 
 				case ModuleRules.ModuleType.External:
