Linux: do not export overriden operators new/delete.

This is needed to avoid other shared objects (particularly Steam) picking them up and possibly crashing due to incorrect usage. This has already happened in the wild and the issue was reported on UDN.
Controlling visibility by attributes does not work due to those operators being defined with default visibility in system headers (same issue: http://lists.apple.com/archives/xcode-users/2006/Jan/msg00696.html)

diff --git a/Engine/Source/Programs/UnrealBuildTool/Linux/LinuxToolChain.cs b/Engine/Source/Programs/UnrealBuildTool/Linux/LinuxToolChain.cs
--- a/Engine/Source/Programs/UnrealBuildTool/Linux/LinuxToolChain.cs
+++ b/Engine/Source/Programs/UnrealBuildTool/Linux/LinuxToolChain.cs
@@ -603,6 +603,30 @@ namespace UnrealBuildTool.Linux
             LinkAction.CommandArguments += LinkEnvironment.Config.AdditionalArguments;
             LinkAction.CommandArguments.Replace("\\", "/");
 
+            // prepare a linker script
+            string LinkerScriptPath = Path.Combine(LinkEnvironment.Config.LocalShadowDirectory, "remove-sym.ldscript");
+            if (!Directory.Exists(LinkEnvironment.Config.LocalShadowDirectory))
+            {
+                Directory.CreateDirectory(LinkEnvironment.Config.LocalShadowDirectory);
+            }
+            if (File.Exists(LinkerScriptPath))
+            {
+                File.Delete(LinkerScriptPath);
+            }
+
+            using (StreamWriter Writer = File.CreateText(LinkerScriptPath))
+            {
+                Writer.WriteLine("UE4 {");
+                Writer.WriteLine("  global: *;");
+                Writer.WriteLine("  local: _Znwm;");
+                Writer.WriteLine("         _Znam;");
+                Writer.WriteLine("         _ZdaPv;");
+                Writer.WriteLine("         _ZdlPv;");
+                Writer.WriteLine("};");
+            };
+
+            LinkAction.CommandArguments += string.Format(" -Wl,--version-script={0}", LinkerScriptPath);
+
             // Only execute linking on the local PC.
             LinkAction.bCanExecuteRemotely = false;
 
