UBT: initial changes to support building Linux tools.

diff --git a/Engine/Source/Programs/UnrealBuildTool/Linux/UEBuildLinux.cs b/Engine/Source/Programs/UnrealBuildTool/Linux/UEBuildLinux.cs
--- a/Engine/Source/Programs/UnrealBuildTool/Linux/UEBuildLinux.cs
+++ b/Engine/Source/Programs/UnrealBuildTool/Linux/UEBuildLinux.cs
@@ -201,15 +201,8 @@ namespace UnrealBuildTool.Linux
         public override void ValidateUEBuildConfiguration()
         {
             BuildConfiguration.bUseUnityBuild = true;
-
-            UEBuildConfiguration.bCompileLeanAndMeanUE = true;
-            UEBuildConfiguration.bCompilePhysX = true;
             UEBuildConfiguration.bCompileAPEX = false;
 
-            UEBuildConfiguration.bBuildEditor = false;
-            UEBuildConfiguration.bBuildDeveloperTools = false;
-            UEBuildConfiguration.bCompileSimplygon = false;
-
             // Don't stop compilation at first error...
             BuildConfiguration.bStopXGECompilationAfterErrors = true;
 
@@ -230,30 +223,6 @@ namespace UnrealBuildTool.Linux
             return true;
         }
 
-        /**
-         *	Whether the editor should be built for this platform or not
-         *	
-         *	@param	InPlatform		The UnrealTargetPlatform being built
-         *	@param	InConfiguration	The UnrealTargetConfiguration being built
-         *	@return	bool			true if the editor should be built, false if not
-         */
-        public override bool ShouldNotBuildEditor(UnrealTargetPlatform InPlatform, UnrealTargetConfiguration InConfiguration)
-        {
-            return true;
-        }
-
-        /**
-         *	Whether the platform requires cooked data
-         *	
-         *	@param	InPlatform		The UnrealTargetPlatform being built
-         *	@param	InConfiguration	The UnrealTargetConfiguration being built
-         *	@return	bool			true if the platform requires cooked data, false if not
-         */
-        public override bool BuildRequiresCookedData(UnrealTargetPlatform InPlatform, UnrealTargetConfiguration InConfiguration)
-        {
-            return true;
-        }
-
         /**
          *	Get a list of extra modules the platform requires.
          *	This is to allow undisclosed platforms to add modules they need without exposing information about the platfomr.
@@ -308,20 +277,20 @@ namespace UnrealBuildTool.Linux
          */
         public override void SetUpEnvironment(UEBuildTarget InBuildTarget)
         {
-            InBuildTarget.GlobalCompileEnvironment.Config.Definitions.Add("UNICODE");
-            InBuildTarget.GlobalCompileEnvironment.Config.Definitions.Add("_UNICODE");
-
             InBuildTarget.GlobalCompileEnvironment.Config.Definitions.Add("PLATFORM_LINUX=1");
             InBuildTarget.GlobalCompileEnvironment.Config.Definitions.Add("LINUX=1");
 
             InBuildTarget.GlobalCompileEnvironment.Config.Definitions.Add("WITH_DATABASE_SUPPORT=0");		//@todo linux: valid?
-            InBuildTarget.GlobalCompileEnvironment.Config.Definitions.Add("WITH_EDITOR=0");
 
             // link with Linux libraries.
             InBuildTarget.GlobalLinkEnvironment.Config.AdditionalLibraries.Add("pthread");
             InBuildTarget.GlobalLinkEnvironment.Config.AdditionalLibraries.Add("z");
 
-            UEBuildConfiguration.bCompileSimplygon = false;
+            // Disable Simplygon support if compiling against the NULL RHI.
+            if (InBuildTarget.GlobalCompileEnvironment.Config.Definitions.Contains("USE_NULL_RHI=1"))
+            {
+                UEBuildConfiguration.bCompileSimplygon = false;
+            }
         }
 
         /**
