Add obsolete GUBP functions back in (marked as obsolete to give a warning on use), to prevent samples not being able to generate projects when upgrading to 4.2 release.

diff --git a/Engine/Source/Programs/UnrealBuildTool/System/RulesCompiler.cs b/Engine/Source/Programs/UnrealBuildTool/System/RulesCompiler.cs
--- a/Engine/Source/Programs/UnrealBuildTool/System/RulesCompiler.cs
+++ b/Engine/Source/Programs/UnrealBuildTool/System/RulesCompiler.cs
@@ -667,6 +667,25 @@ namespace UnrealBuildTool
             Result.Add("GameAutomationTest", "BuildCookRun -run -RunAutomationTests -unattended -nullrhi -NoP4" + MacOption);
             return Result;
         }
+        /// <summary>
+        /// Allow the platform to setup emails for the GUBP for folks that care about node failures relating to this platform
+		/// Obsolete. Included to avoid breaking existing projects.
+        /// </summary>
+        /// <param name="Branch">p4 root of the branch we are running</param>
+		[Obsolete]
+        public virtual string GUBP_GetGameFailureEMails_EditorTypeOnly(string Branch)
+        {
+            return "";
+        }
+        /// <summary>
+        /// Allow the Game to set up emails for Promotable and Promotion
+		/// Obsolete. Included to avoid breaking existing projects.
+        /// </summary>
+		[Obsolete]
+        public virtual string GUBP_GetPromotionEMails_EditorTypeOnly(string Branch)
+        {
+            return "";
+        }
 
         /// <summary>
         /// Return a list of "test name", "UAT command" pairs for testing a monolithic
