GitHub pull request #190: Detailed module types for plugins

diff --git a/Engine/Source/Programs/UnrealBuildTool/Configuration/UEBuildModule.cs b/Engine/Source/Programs/UnrealBuildTool/Configuration/UEBuildModule.cs
--- a/Engine/Source/Programs/UnrealBuildTool/Configuration/UEBuildModule.cs
+++ b/Engine/Source/Programs/UnrealBuildTool/Configuration/UEBuildModule.cs
@@ -19,7 +19,6 @@ namespace UnrealBuildTool
 		Developer,
 		Editor,
 		ThirdParty,
-		Plugin,
 		Program,
 		Game,
 	}
@@ -38,11 +37,47 @@ namespace UnrealBuildTool
 		public static readonly string RuntimeFolder = String.Format("{0}Runtime{0}", Path.DirectorySeparatorChar);
 		public static readonly string DeveloperFolder = String.Format("{0}Developer{0}", Path.DirectorySeparatorChar);
 		public static readonly string EditorFolder = String.Format("{0}Editor{0}", Path.DirectorySeparatorChar);
-		public static readonly string PluginFolder = String.Format("Plugins{0}", Path.DirectorySeparatorChar);
-		public static readonly string ThirdPartyFolder = String.Format("{0}ThirdParty{0}", Path.DirectorySeparatorChar);
 		public static readonly string ProgramsFolder = String.Format("{0}Programs{0}", Path.DirectorySeparatorChar);
 
-		public static UEBuildModuleType GetEngineModuleTypeBasedOnLocation(UEBuildModuleType ModuleType, string ModuleFileRelativeToEngineDirectory)
+		public static UEBuildModuleType GetPluginModuleType(string ModuleName)
+		{
+			UEBuildModuleType ModuleType = UEBuildModuleType.Unknown;
+			PluginInfo Info = Plugins.GetPluginInfoForModule(ModuleName);
+			if (null == Info)
+			{
+				return ModuleType;
+			}
+			foreach (PluginInfo.PluginModuleInfo ModuleInfo in Info.Modules)
+			{
+				if (ModuleInfo.Name == ModuleName)
+				{
+					switch (ModuleInfo.Type)
+					{
+						case PluginInfo.PluginModuleType.Developer:
+							ModuleType = UEBuildModuleType.Developer;
+							break;
+						case PluginInfo.PluginModuleType.Editor:
+						case PluginInfo.PluginModuleType.EditorNoCommandlet:
+							ModuleType = UEBuildModuleType.Editor;
+							break;
+						case PluginInfo.PluginModuleType.Program:
+							ModuleType = UEBuildModuleType.Program;
+							break;
+						case PluginInfo.PluginModuleType.Runtime:
+						case PluginInfo.PluginModuleType.RuntimeNoCommandlet:
+							ModuleType = UEBuildModuleType.Runtime;
+							break;
+						default:
+							throw new BuildException("Unhandled plugin module type {0}", ModuleInfo.Type.ToString());
+					}
+					// break out of loop
+					break;
+				}
+			}
+			return ModuleType;
+		}
+
+		public static UEBuildModuleType GetEngineModuleTypeBasedOnLocation(string ModuleName, UEBuildModuleType ModuleType, string ModuleFileRelativeToEngineDirectory)
 		{
 			if (ModuleFileRelativeToEngineDirectory.IndexOf(UEBuildModule.RuntimeFolder, StringComparison.InvariantCultureIgnoreCase) >= 0)
 			{
@@ -56,14 +91,6 @@ namespace UnrealBuildTool
 			{
 				ModuleType = UEBuildModuleType.Editor;
 			}
-			else if (ModuleFileRelativeToEngineDirectory.StartsWith(UEBuildModule.PluginFolder, StringComparison.InvariantCultureIgnoreCase))
-			{
-				ModuleType = UEBuildModuleType.Plugin;
-			}
-			else if (ModuleFileRelativeToEngineDirectory.IndexOf(UEBuildModule.ThirdPartyFolder, StringComparison.InvariantCultureIgnoreCase) >= 0)
-			{
-				ModuleType = UEBuildModuleType.ThirdParty;
-			}
 			else if (ModuleFileRelativeToEngineDirectory.IndexOf(UEBuildModule.ProgramsFolder, StringComparison.InvariantCultureIgnoreCase) >= 0)
 			{
 				ModuleType = UEBuildModuleType.Program;
diff --git a/Engine/Source/Programs/UnrealBuildTool/Configuration/UEBuildTarget.cs b/Engine/Source/Programs/UnrealBuildTool/Configuration/UEBuildTarget.cs
--- a/Engine/Source/Programs/UnrealBuildTool/Configuration/UEBuildTarget.cs
+++ b/Engine/Source/Programs/UnrealBuildTool/Configuration/UEBuildTarget.cs
@@ -2355,27 +2355,42 @@ namespace UnrealBuildTool
 				var ModuleType = UEBuildModuleType.Unknown;
 				var ApplicationOutputPath = "";
 				var ModuleFileRelativeToEngineDirectory = Utils.MakePathRelativeTo(ModuleFileName, ProjectFileGenerator.EngineRelativePath);
-				if (ModuleFileRelativeToEngineDirectory.StartsWith("..") || Path.IsPathRooted(ModuleFileRelativeToEngineDirectory))
+
+				// see if it's external
+				if (RulesObject.Type == ModuleRules.ModuleType.External)
 				{
-					ApplicationOutputPath = ModuleFileName;
-					int SourceIndex = ApplicationOutputPath.LastIndexOf(Path.DirectorySeparatorChar + "Source" + Path.DirectorySeparatorChar);
-					if (SourceIndex != -1)
-					{
-						ApplicationOutputPath = ApplicationOutputPath.Substring(0, SourceIndex + 1);
-						ModuleType = UEBuildModuleType.Game;
-					}
-					else
-					{
-						throw new BuildException("Unable to parse application directory for module '{0}' ({1}). Is it in '../../<APP>/Source'?",
-							ModuleName, ModuleFileName);
-					}
+					ModuleType = UEBuildModuleType.ThirdParty;
 				}
 				else
 				{
-					ModuleType = UEBuildModule.GetEngineModuleTypeBasedOnLocation(ModuleType, ModuleFileRelativeToEngineDirectory);
+					// Check if it's a plugin
+					ModuleType = UEBuildModule.GetPluginModuleType(ModuleName);
 					if (ModuleType == UEBuildModuleType.Unknown)
 					{
-						throw new BuildException("Unable to determine module type for {0}", ModuleFileName);
+						// not a plugin, see if it is a game module 
+						if (ModuleFileRelativeToEngineDirectory.StartsWith("..") || Path.IsPathRooted(ModuleFileRelativeToEngineDirectory))
+						{
+							ApplicationOutputPath = ModuleFileName;
+							int SourceIndex = ApplicationOutputPath.LastIndexOf(Path.DirectorySeparatorChar + "Source" + Path.DirectorySeparatorChar);
+							if (SourceIndex != -1)
+							{
+								ApplicationOutputPath = ApplicationOutputPath.Substring(0, SourceIndex + 1);
+								ModuleType = UEBuildModuleType.Game;
+							}
+							else
+							{
+								throw new BuildException("Unable to parse application directory for module '{0}' ({1}). Is it in '../../<APP>/Source'?",
+									ModuleName, ModuleFileName);
+							}
+						}
+						else
+						{
+							ModuleType = UEBuildModule.GetEngineModuleTypeBasedOnLocation(ModuleName, ModuleType, ModuleFileRelativeToEngineDirectory);
+							if (ModuleType == UEBuildModuleType.Unknown)
+							{
+								throw new BuildException("Unable to determine module type for {0}", ModuleFileName);
+							}
+						}
 					}
 				}
 
diff --git a/Engine/Source/Programs/UnrealHeaderTool/Public/IScriptGeneratorPluginInterface.h b/Engine/Source/Programs/UnrealHeaderTool/Public/IScriptGeneratorPluginInterface.h
--- a/Engine/Source/Programs/UnrealHeaderTool/Public/IScriptGeneratorPluginInterface.h
+++ b/Engine/Source/Programs/UnrealHeaderTool/Public/IScriptGeneratorPluginInterface.h
@@ -13,7 +13,6 @@ struct EBuildModuleType
 		Developer,
 		Editor,
 		ThirdParty,
-		Plugin,
 		Program,
 		Game,
 		// NOTE: If you add a new value, make sure to update the ToString() method below!
@@ -51,10 +50,6 @@ struct EBuildModuleType
 		{
 			return ThirdParty;
 		}
-		else if (FCString::Stricmp(Value, TEXT("Plugin")) == 0)
-		{
-			return Plugin;
-		}
 		else if (FCString::Stricmp(Value, TEXT("Program")) == 0)
 		{
 			return Program;
