Removed Mac MainMenu.xib, the app and window menus are now generated dynamically.

diff --git a/Engine/Source/Editor/MainFrame/Private/Menus/MainMenu.cpp b/Engine/Source/Editor/MainFrame/Private/Menus/MainMenu.cpp
--- a/Engine/Source/Editor/MainFrame/Private/Menus/MainMenu.cpp
+++ b/Engine/Source/Editor/MainFrame/Private/Menus/MainMenu.cpp
@@ -121,12 +121,14 @@ void FMainMenu::FillEditMenu( FMenuBuilder& MenuBuilder, const TSharedRef< FExte
 		}
 		else
 		{
+#if !PLATFORM_MAC // Handled by app's menu in menu bar
 			MenuBuilder.AddMenuEntry(
 				LOCTEXT("EditorPreferencesMenuLabel", "Editor Preferences..."),
 				LOCTEXT("EditorPreferencesMenuToolTip", "Configure the behavior and features of this Editor"),
 				FSlateIcon(),
 				FUIAction(FExecuteAction::CreateStatic(&FSettingsMenu::OpenSettings, FName("Editor"), FName("General"), FName("Appearance")))
 			);
+#endif
 
 			MenuBuilder.AddMenuEntry(
 				LOCTEXT("ProjectSettingsMenuLabel", "Project Settings..."),
@@ -224,6 +226,7 @@ void FMainMenu::FillHelpMenu( FMenuBuilder& MenuBuilder, const TSharedRef< FExte
 	}
 	MenuBuilder.EndSection();
 
+#if !PLATFORM_MAC // Handled by app's menu in menu bar
 	MenuBuilder.BeginSection("HelpApplication", NSLOCTEXT("MainHelpMenu", "Application", "Application"));
 	{
 		const FText AboutWindowTitle = NSLOCTEXT("MainHelpMenu", "AboutUnrealEditor", "About Unreal Editor...");
@@ -231,6 +234,7 @@ void FMainMenu::FillHelpMenu( FMenuBuilder& MenuBuilder, const TSharedRef< FExte
 		MenuBuilder.AddMenuEntry(FMainFrameCommands::Get().AboutUnrealEd, "AboutUnrealEd", AboutWindowTitle);
 	}
 	MenuBuilder.EndSection();
+#endif
 }
 
 
@@ -366,9 +370,10 @@ TSharedRef< SWidget > FMainMenu::MakeMainTabMenu( const TSharedPtr<FTabManager>&
 				}
 			}
 			MenuBuilder.EndSection();
-
+#if !PLATFORM_MAC // Handled by app's menu in menu bar
 			MenuBuilder.AddMenuSeparator();
 			MenuBuilder.AddMenuEntry( FMainFrameCommands::Get().Exit, "Exit" );
+#endif
 		}
 	};
 
diff --git a/Engine/Source/Programs/CrashReporter/CrashReportClient/Private/Mac/CrashReportClientMainMac.cpp b/Engine/Source/Programs/CrashReporter/CrashReportClient/Private/Mac/CrashReportClientMainMac.cpp
--- a/Engine/Source/Programs/CrashReporter/CrashReportClient/Private/Mac/CrashReportClientMainMac.cpp
+++ b/Engine/Source/Programs/CrashReporter/CrashReportClient/Private/Mac/CrashReportClientMainMac.cpp
@@ -40,7 +40,7 @@ static FString GSavedCommandLine;
 //handler for the quit apple event used by the Dock menu
 - (void)handleQuitEvent:(NSAppleEventDescriptor*)Event withReplyEvent:(NSAppleEventDescriptor*)ReplyEvent
 {
-    [self OnQuitRequest:self];
+	[NSApp terminate:self];
 }
 
 - (void)applicationDidFinishLaunching:(NSNotification *)Notification
@@ -63,16 +63,6 @@ static FString GSavedCommandLine;
 	[NSApp terminate: self];
 }
 
-- (IBAction)OnQuitRequest:(id)Sender
-{
-	[NSApp terminate: self];
-}
-
-- (IBAction)OnShowAboutWindow:(id)Sender
-{
-	[NSApp orderFrontStandardAboutPanel: Sender];
-}
-
 @end
 
 int main(int argc, char *argv[])
@@ -98,5 +88,9 @@ int main(int argc, char *argv[])
 		GSavedCommandLine += Argument;
 	}
 	
-	return NSApplicationMain(argc, (const char **)argv);
+	SCOPED_AUTORELEASE_POOL;
+	[NSApplication sharedApplication];
+	[NSApp setDelegate:[UE4AppDelegate new]];
+	[NSApp run];
+	return 0;
 }
diff --git a/Engine/Source/Programs/Mac/UE4EditorServices/Private/Mac/UE4EditorServicesMain.cpp b/Engine/Source/Programs/Mac/UE4EditorServices/Private/Mac/UE4EditorServicesMain.cpp
--- a/Engine/Source/Programs/Mac/UE4EditorServices/Private/Mac/UE4EditorServicesMain.cpp
+++ b/Engine/Source/Programs/Mac/UE4EditorServices/Private/Mac/UE4EditorServicesMain.cpp
@@ -3,6 +3,7 @@
 #include "Core.h"
 #include "RequiredProgramMainCPPInclude.h"
 #include "ModuleManager.h"
+#include "UE4EditorServicesAppDelegate.h"
 
 #import <Cocoa/Cocoa.h>
 
@@ -12,5 +13,9 @@ int main(int argc, char *argv[])
 {
 	FCommandLine::Set(TEXT(""));
 	FPlatformProcess::SetCurrentWorkingDirectoryToBaseDir();
-	return NSApplicationMain(argc, (const char**)argv);
+	SCOPED_AUTORELEASE_POOL;
+	[NSApplication sharedApplication];
+	[NSApp setDelegate:[FUE4EditorServicesAppDelegate new]];
+	[NSApp run];
+	return 0;
 }
diff --git a/Engine/Source/Programs/Mac/UE4EditorServices/Resources/Mac/Info.plist b/Engine/Source/Programs/Mac/UE4EditorServices/Resources/Mac/Info.plist
--- a/Engine/Source/Programs/Mac/UE4EditorServices/Resources/Mac/Info.plist
+++ b/Engine/Source/Programs/Mac/UE4EditorServices/Resources/Mac/Info.plist
@@ -53,8 +53,6 @@
 	<true/>
 	<key>NSHumanReadableCopyright</key>
 	<string>Copyright © 2014 Epic. All rights reserved.</string>
-	<key>NSMainNibFile</key>
-	<string>MainMenu</string>
 	<key>NSPrincipalClass</key>
 	<string>NSApplication</string>
 	<key>NSServices</key>
diff --git a/Engine/Source/Programs/SlateViewer/Private/Mac/SlateViewerMainMac.cpp b/Engine/Source/Programs/SlateViewer/Private/Mac/SlateViewerMainMac.cpp
--- a/Engine/Source/Programs/SlateViewer/Private/Mac/SlateViewerMainMac.cpp
+++ b/Engine/Source/Programs/SlateViewer/Private/Mac/SlateViewerMainMac.cpp
@@ -16,7 +16,7 @@ static FString GSavedCommandLine;
 //handler for the quit apple event used by the Dock menu
 - (void)handleQuitEvent:(NSAppleEventDescriptor*)Event withReplyEvent:(NSAppleEventDescriptor*)ReplyEvent
 {
-    [self OnQuitRequest:self];
+    [NSApp terminate:self];
 }
 
 - (void)applicationDidFinishLaunching:(NSNotification *)Notification
@@ -35,16 +35,6 @@ static FString GSavedCommandLine;
 	[NSApp terminate: self];
 }
 
-- (IBAction)OnQuitRequest:(id)Sender
-{
-	[NSApp terminate: self];
-}
-
-- (IBAction)OnShowAboutWindow:(id)Sender
-{
-	[NSApp orderFrontStandardAboutPanel: Sender];
-}
-
 @end
 
 int main(int argc, char *argv[])
@@ -70,5 +60,9 @@ int main(int argc, char *argv[])
 		GSavedCommandLine += Argument;
 	}
 
-	return NSApplicationMain(argc, (const char **)argv);
+	SCOPED_AUTORELEASE_POOL;
+	[NSApplication sharedApplication];
+	[NSApp setDelegate:[UE4AppDelegate new]];
+	[NSApp run];
+	return 0;
 }
diff --git a/Engine/Source/Programs/SymbolDebugger/Private/Mac/SymbolDebuggerMainMac.cpp b/Engine/Source/Programs/SymbolDebugger/Private/Mac/SymbolDebuggerMainMac.cpp
--- a/Engine/Source/Programs/SymbolDebugger/Private/Mac/SymbolDebuggerMainMac.cpp
+++ b/Engine/Source/Programs/SymbolDebugger/Private/Mac/SymbolDebuggerMainMac.cpp
@@ -19,16 +19,6 @@ static FString GSavedCommandLine;
 	[NSApp terminate: self];
 }
 
-- (IBAction)OnQuitRequest:(id)Sender
-{
-	[NSApp terminate: self];
-}
-
-- (IBAction)OnShowAboutWindow:(id)Sender
-{
-	[NSApp orderFrontStandardAboutPanel: Sender];
-}
-
 @end
 
 int main(int argc, char *argv[])
@@ -44,5 +34,9 @@ int main(int argc, char *argv[])
 		GSavedCommandLine += Argument;
 	}
 
-	return NSApplicationMain(argc, (const char **)argv);
+	SCOPED_AUTORELEASE_POOL;
+	[NSApplication sharedApplication];
+	[NSApp setDelegate:[UE4AppDelegate new]];
+	[NSApp run];
+	return 0;
 }
diff --git a/Engine/Source/Programs/UnrealBuildTool/Mac/MacToolChain.cs b/Engine/Source/Programs/UnrealBuildTool/Mac/MacToolChain.cs
--- a/Engine/Source/Programs/UnrealBuildTool/Mac/MacToolChain.cs
+++ b/Engine/Source/Programs/UnrealBuildTool/Mac/MacToolChain.cs
@@ -873,7 +873,6 @@ namespace UnrealBuildTool
 					string GameName = ExeNameParts[0];
 
 					AppendMacLine(CreateAppBundleScript, "mkdir -p \"{0}.app/Contents/MacOS\"", ExeName);
-					AppendMacLine(CreateAppBundleScript, "mkdir -p \"{0}.app/Contents/Resources/English.lproj\"", ExeName);
 					AppendMacLine(CreateAppBundleScript, "mkdir -p \"{0}.app/Contents/Resources/RadioEffectUnit.component/Contents/MacOS\"", ExeName);
 					AppendMacLine(CreateAppBundleScript, "mkdir -p \"{0}.app/Contents/Resources/RadioEffectUnit.component/Contents/Resources/English.lproj\"", ExeName);
 
@@ -933,23 +932,12 @@ namespace UnrealBuildTool
 					// Generate PkgInfo file
 					AppendMacLine(CreateAppBundleScript, "echo 'echo -n \"APPL????\"' | bash > \"{0}.app/Contents/PkgInfo\"", ExeName);
 
-					// Copy InfoPlist.strings using iconv to convert to UTF-16 in the process
-					AppendMacLine(CreateAppBundleScript, "iconv -f UTF-8 -t UTF-16 \"{0}/Runtime/Launch/Resources/Mac/English.lproj/InfoPlist.strings\" > \"{1}.app/Contents/Resources/English.lproj/InfoPlist.strings\"", EngineSourcePath, ExeName);
-
 					// Copy RadioEffect component
 					AppendMacLine(CreateAppBundleScript, "cp -f \"{0}/ThirdParty/Mac/RadioEffectUnit/RadioEffectUnit.component/Contents/MacOS/RadioEffectUnit\" \"{1}.app/Contents/Resources/RadioEffectUnit.component/Contents/MacOS/RadioEffectUnit\"", EngineSourcePath, ExeName);
 					AppendMacLine(CreateAppBundleScript, "cp -f \"{0}/ThirdParty/Mac/RadioEffectUnit/RadioEffectUnit.component/Contents/Resources/English.lproj/Localizable.strings\" \"{1}.app/Contents/Resources/RadioEffectUnit.component/Contents/Resources/English.lproj/Localizable.strings\"", EngineSourcePath, ExeName);
 					AppendMacLine(CreateAppBundleScript, "cp -f \"{0}/ThirdParty/Mac/RadioEffectUnit/RadioEffectUnit.component/Contents/Resources/RadioEffectUnit.rsrc\" \"{1}.app/Contents/Resources/RadioEffectUnit.component/Contents/Resources/RadioEffectUnit.rsrc\"", EngineSourcePath, ExeName);
 					AppendMacLine(CreateAppBundleScript, "cp -f \"{0}/ThirdParty/Mac/RadioEffectUnit/RadioEffectUnit.component/Contents/Info.plist\" \"{1}.app/Contents/Resources/RadioEffectUnit.component/Contents/Info.plist\"", EngineSourcePath, ExeName);
 
-					// Compile XIB resource
-					string XIBFile = CustomResourcesPath + "/MainMenu.xib";
-					if (!File.Exists(XIBFile))
-					{
-						XIBFile = EngineSourcePath + "/Runtime/Launch/Resources/Mac/English.lproj/MainMenu.xib";
-					}
-					AppendMacLine(CreateAppBundleScript, "\"{0}usr/bin/ibtool\" --errors --warnings --notices --output-format human-readable-text --compile \"{2}.app/Contents/Resources/English.lproj/MainMenu.nib\" \"{1}\" --sdk \"{0}/Platforms/MacOSX.platform/Developer/SDKs/MacOSX{3}.sdk\"", DeveloperDir, XIBFile, ExeName, MacOSSDKVersion);
-
 					// Make sure OS X knows the bundle was updated
 					AppendMacLine(CreateAppBundleScript, "touch -c \"{0}.app\"", ExeName);
 
@@ -960,8 +948,6 @@ namespace UnrealBuildTool
 					QueueFileForBatchUpload(FileItem.GetItemByFullPath(Path.GetFullPath("../../Engine/Source/Runtime/Launch/Resources/Mac/UE4.icns")));
 					QueueFileForBatchUpload(FileItem.GetItemByFullPath(Path.GetFullPath("../../Engine/Source/Runtime/Launch/Resources/Mac/UProject.icns")));
 					QueueFileForBatchUpload(FileItem.GetItemByFullPath(Path.GetFullPath("../../Engine/Source/Runtime/Launch/Resources/Mac/Info.plist")));
-					QueueFileForBatchUpload(FileItem.GetItemByFullPath(Path.GetFullPath("../../Engine/Source/Runtime/Launch/Resources/Mac/English.lproj/InfoPlist.strings")));
-					QueueFileForBatchUpload(FileItem.GetItemByFullPath(Path.GetFullPath("../../Engine/Source/Runtime/Launch/Resources/Mac/English.lproj/MainMenu.xib")));
 					QueueFileForBatchUpload(FileItem.GetItemByFullPath(Path.GetFullPath("../../Engine/Source/ThirdParty/Mac/RadioEffectUnit/RadioEffectUnit.component/Contents/MacOS/RadioEffectUnit")));
 					QueueFileForBatchUpload(FileItem.GetItemByFullPath(Path.GetFullPath("../../Engine/Source/ThirdParty/Mac/RadioEffectUnit/RadioEffectUnit.component/Contents/Resources/English.lproj/Localizable.strings")));
 					QueueFileForBatchUpload(FileItem.GetItemByFullPath(Path.GetFullPath("../../Engine/Source/ThirdParty/Mac/RadioEffectUnit/RadioEffectUnit.component/Contents/Resources/RadioEffectUnit.rsrc")));
@@ -1182,8 +1168,6 @@ namespace UnrealBuildTool
 				Manifest.AddFileName(BundleContentsDirectory + "Info.plist");
 				Manifest.AddFileName(BundleContentsDirectory + "PkgInfo");
 				Manifest.AddFileName(BundleContentsDirectory + "Resources/UE4.icns");
-				Manifest.AddFileName(BundleContentsDirectory + "Resources/English.lproj/InfoPlist.strings");
-				Manifest.AddFileName(BundleContentsDirectory + "Resources/English.lproj/MainMenu.nib");
 				Manifest.AddFileName(BundleContentsDirectory + "Resources/RadioEffectUnit.component/Contents/MacOS/RadioEffectUnit");
 				Manifest.AddFileName(BundleContentsDirectory + "Resources/RadioEffectUnit.component/Contents/Resources/English.lproj/Localizable.strings");
 				Manifest.AddFileName(BundleContentsDirectory + "Resources/RadioEffectUnit.component/Contents/Resources/RadioEffectUnit.rsrc");
diff --git a/Engine/Source/Programs/UnrealFrontend/Private/Mac/MacUnrealFrontendMain.cpp b/Engine/Source/Programs/UnrealFrontend/Private/Mac/MacUnrealFrontendMain.cpp
--- a/Engine/Source/Programs/UnrealFrontend/Private/Mac/MacUnrealFrontendMain.cpp
+++ b/Engine/Source/Programs/UnrealFrontend/Private/Mac/MacUnrealFrontendMain.cpp
@@ -22,7 +22,7 @@ static FString GSavedCommandLine;
 //handler for the quit apple event used by the Dock menu
 - (void)handleQuitEvent:(NSAppleEventDescriptor*)Event withReplyEvent:(NSAppleEventDescriptor*)ReplyEvent
 {
-    [self OnQuitRequest:self];
+    [self requestQuit:self];
 }
 
 - (void)applicationDidFinishLaunching:(NSNotification *)Notification
@@ -63,16 +63,11 @@ static FString GSavedCommandLine;
 	[NSApp terminate: self];
 }
 
-- (IBAction)OnQuitRequest:(id)Sender
+- (IBAction)requestQuit:(id)Sender
 {
 	GIsRequestingExit = true;
 }
 
-- (IBAction)OnShowAboutWindow:(id)Sender
-{
-	[NSApp orderFrontStandardAboutPanel: Sender];
-}
-
 @end
 
 
@@ -99,5 +94,9 @@ int main( int argc, char *argv[] )
 		GSavedCommandLine += Argument;
 	}
 
-	return NSApplicationMain(argc, (const char**)argv);
+	SCOPED_AUTORELEASE_POOL;
+	[NSApplication sharedApplication];
+	[NSApp setDelegate:[UE4AppDelegate new]];
+	[NSApp run];
+	return 0;
 }
diff --git a/Engine/Source/Runtime/Core/Private/Mac/MacPlatformMisc.cpp b/Engine/Source/Runtime/Core/Private/Mac/MacPlatformMisc.cpp
--- a/Engine/Source/Runtime/Core/Private/Mac/MacPlatformMisc.cpp
+++ b/Engine/Source/Runtime/Core/Private/Mac/MacPlatformMisc.cpp
@@ -205,10 +205,86 @@ void FMacPlatformMisc::PlatformInit()
 
 	// Timer resolution.
 	UE_LOG(LogInit, Log, TEXT("High frequency timer resolution =%f MHz"), 0.000001 / FPlatformTime::GetSecondsPerCycle() );
-	
+
 	GMacAppInfo.Init();
 }
 
+void FMacPlatformMisc::PlatformPostInit()
+{
+	// Setup the app menu in menu bar
+	const bool bIsBundledApp = [[[NSBundle mainBundle] bundlePath] hasSuffix:@".app"];
+	if (bIsBundledApp)
+	{
+		NSString* AppName = GIsEditor ? @"Unreal Editor" : FString(GGameName).GetNSString();
+
+		SEL ShowAboutSelector = [[NSApp delegate] respondsToSelector:@selector(showAboutWindow:)] ? @selector(showAboutWindow:) : @selector(orderFrontStandardAboutPanel:);
+		NSMenuItem* AboutItem = [[[NSMenuItem alloc] initWithTitle:[NSString stringWithFormat:@"About %@", AppName] action:ShowAboutSelector keyEquivalent:@""] autorelease];
+
+		NSMenuItem* PreferencesItem = GIsEditor ? [[[NSMenuItem alloc] initWithTitle:@"Preferences..." action:@selector(showPreferencesWindow:) keyEquivalent:@","] autorelease] : nil;
+
+		NSMenuItem* HideItem = [[[NSMenuItem alloc] initWithTitle:[NSString stringWithFormat:@"Hide %@", AppName] action:@selector(hide:) keyEquivalent:@"h"] autorelease];
+		NSMenuItem* HideOthersItem = [[[NSMenuItem alloc] initWithTitle:@"Hide Others" action:@selector(hideOtherApplications:) keyEquivalent:@"h"] autorelease];
+		[HideOthersItem setKeyEquivalentModifierMask:NSCommandKeyMask | NSAlternateKeyMask];
+		NSMenuItem* ShowAllItem = [[[NSMenuItem alloc] initWithTitle:@"Show All" action:@selector(unhideAllApplications:) keyEquivalent:@""] autorelease];
+
+		SEL RequestQuitSelector = [[NSApp delegate] respondsToSelector:@selector(requestQuit:)] ? @selector(requestQuit:) : @selector(terminate:);
+		NSMenuItem* QuitItem = [[[NSMenuItem alloc] initWithTitle:[NSString stringWithFormat:@"Quit %@", AppName] action:RequestQuitSelector keyEquivalent:@"q"] autorelease];
+
+		NSMenuItem* ServicesItem = [[NSMenuItem new] autorelease];
+		NSMenu* ServicesMenu = [[NSMenu new] autorelease];
+		[ServicesItem setTitle:@"Services"];
+		[ServicesItem setSubmenu:ServicesMenu];
+		[NSApp setServicesMenu:ServicesMenu];
+
+		NSMenu* AppMenu = [[NSMenu new] autorelease];
+		[AppMenu addItem:AboutItem];
+		[AppMenu addItem:[NSMenuItem separatorItem]];
+		if (PreferencesItem)
+		{
+			[AppMenu addItem:PreferencesItem];
+			[AppMenu addItem:[NSMenuItem separatorItem]];
+		}
+		[AppMenu addItem:ServicesItem];
+		[AppMenu addItem:[NSMenuItem separatorItem]];
+		[AppMenu addItem:HideItem];
+		[AppMenu addItem:HideOthersItem];
+		[AppMenu addItem:ShowAllItem];
+		[AppMenu addItem:[NSMenuItem separatorItem]];
+		[AppMenu addItem:QuitItem];
+
+		NSMenu* MenuBar = [[NSMenu new] autorelease];
+		NSMenuItem* AppMenuItem = [[NSMenuItem new] autorelease];
+		[MenuBar addItem:AppMenuItem];
+		[NSApp setMainMenu:MenuBar];
+		[AppMenuItem setSubmenu:AppMenu];
+
+		UpdateWindowMenu();
+	}
+}
+
+void FMacPlatformMisc::UpdateWindowMenu()
+{
+	NSMenu* WindowMenu = [NSApp windowsMenu];
+	if (!WindowMenu)
+	{
+		WindowMenu = [[NSMenu new] autorelease];
+		[WindowMenu setTitle:@"Window"];
+		NSMenuItem* WindowMenuItem = [[NSMenuItem new] autorelease];
+		[WindowMenuItem setSubmenu:WindowMenu];
+		[[NSApp mainMenu] addItem:WindowMenuItem];
+		[NSApp setWindowsMenu:WindowMenu];
+	}
+
+	NSMenuItem* MinimizeItem = [[[NSMenuItem alloc] initWithTitle:@"Minimize" action:@selector(miniaturize:) keyEquivalent:@"m"] autorelease];
+	NSMenuItem* ZoomItem = [[[NSMenuItem alloc] initWithTitle:@"Zoom" action:@selector(performZoom:) keyEquivalent:@""] autorelease];
+	NSMenuItem* BringAllToFrontItem = [[[NSMenuItem alloc] initWithTitle:@"Bring All to Front" action:@selector(arrangeInFront:) keyEquivalent:@""] autorelease];
+
+	[WindowMenu addItem:MinimizeItem];
+	[WindowMenu addItem:ZoomItem];
+	[WindowMenu addItem:[NSMenuItem separatorItem]];
+	[WindowMenu addItem:BringAllToFrontItem];
+}
+
 void FMacPlatformMisc::PreventScreenSaver()
 {
 	CGEventRef LocationEvent = CGEventCreate(NULL);
diff --git a/Engine/Source/Runtime/Core/Public/Mac/MacPlatformMisc.h b/Engine/Source/Runtime/Core/Public/Mac/MacPlatformMisc.h
--- a/Engine/Source/Runtime/Core/Public/Mac/MacPlatformMisc.h
+++ b/Engine/Source/Runtime/Core/Public/Mac/MacPlatformMisc.h
@@ -12,6 +12,7 @@
 struct CORE_API FMacPlatformMisc : public FGenericPlatformMisc
 {
 	static void PlatformInit();
+	static void PlatformPostInit();
 	static class GenericApplication* CreateApplication();
 	static void GetEnvironmentVariable(const TCHAR* VariableName, TCHAR* Result, int32 ResultLength);
 	static TArray<uint8> GetMacAddress();
@@ -119,6 +120,8 @@ struct CORE_API FMacPlatformMisc : public FGenericPlatformMisc
 
 	/** @return Get the name of the platform specific file manager (Finder) */
 	static FText GetFileManagerName();
+
+	static void UpdateWindowMenu();
 };
 
 #ifdef __OBJC__
diff --git a/Engine/Source/Runtime/Launch/Launch.Build.cs b/Engine/Source/Runtime/Launch/Launch.Build.cs
--- a/Engine/Source/Runtime/Launch/Launch.Build.cs
+++ b/Engine/Source/Runtime/Launch/Launch.Build.cs
@@ -119,6 +119,7 @@ public class Launch : ModuleRules
 			if (Target.Platform == UnrealTargetPlatform.Mac)
 			{
 				PrivateDependencyModuleNames.Add("MainFrame");
+				PrivateDependencyModuleNames.Add("Settings");
 			}
 			else
 			{
diff --git a/Engine/Source/Runtime/Launch/Private/Mac/LaunchMac.cpp b/Engine/Source/Runtime/Launch/Private/Mac/LaunchMac.cpp
--- a/Engine/Source/Runtime/Launch/Private/Mac/LaunchMac.cpp
+++ b/Engine/Source/Runtime/Launch/Private/Mac/LaunchMac.cpp
@@ -5,6 +5,7 @@
 
 #if WITH_EDITOR
 	#include "MainFrame.h"
+	#include "Settings.h"
 #endif
 
 #include <signal.h>
@@ -89,7 +90,7 @@ void EngineCrashHandler(const FGenericCrashContext & GenericContext)
 }
 #endif
 
-- (IBAction)OnQuitRequest:(id)Sender
+- (IBAction)requestQuit:(id)Sender
 {
 	if (GEngine)
 	{
@@ -111,7 +112,7 @@ void EngineCrashHandler(const FGenericCrashContext & GenericContext)
 	}
 }
 
-- (IBAction)OnShowAboutWindow:(id)Sender
+- (IBAction)showAboutWindow:(id)Sender
 {
 #if WITH_EDITOR
 	if (FModuleManager::Get().IsModuleLoaded(TEXT("MainFrame")))
@@ -119,14 +120,24 @@ void EngineCrashHandler(const FGenericCrashContext & GenericContext)
 		FModuleManager::GetModuleChecked<IMainFrameModule>(TEXT("MainFrame")).ShowAboutWindow();
 	}
 #else
-	[NSApp orderFrontStandardAboutPanel: Sender];
+	[NSApp orderFrontStandardAboutPanel:Sender];
 #endif
 }
 
+#if WITH_EDITOR
+- (IBAction)showPreferencesWindow:(id)Sender
+{
+	if (FModuleManager::Get().IsModuleLoaded(TEXT("Settings")))
+	{
+		FModuleManager::LoadModuleChecked<ISettingsModule>("Settings").ShowViewer(FName("Editor"), FName("General"), FName("Appearance"));
+	}
+}
+#endif
+
 //handler for the quit apple event used by the Dock menu
 - (void)handleQuitEvent:(NSAppleEventDescriptor*)Event withReplyEvent:(NSAppleEventDescriptor*)ReplyEvent
 {
-    [self OnQuitRequest:self];
+    [self requestQuit:self];
 }
 
 - (void)applicationDidFinishLaunching:(NSNotification *)Notification
@@ -231,5 +242,9 @@ INT32_MAIN_INT32_ARGC_TCHAR_ARGV()
 		GSavedCommandLine += Argument;
 	}
 
-	return NSApplicationMain(ArgC, (const char **)ArgV);
+	SCOPED_AUTORELEASE_POOL;
+	[NSApplication sharedApplication];
+	[NSApp setDelegate:[UE4AppDelegate new]];
+	[NSApp run];
+	return 0;
 }
diff --git a/Engine/Source/Runtime/Launch/Resources/Mac/English.lproj/InfoPlist.strings b/Engine/Source/Runtime/Launch/Resources/Mac/English.lproj/InfoPlist.strings
deleted file mode 100644
diff --git a/Engine/Source/Runtime/Launch/Resources/Mac/Info-Editor.plist b/Engine/Source/Runtime/Launch/Resources/Mac/Info-Editor.plist
--- a/Engine/Source/Runtime/Launch/Resources/Mac/Info-Editor.plist
+++ b/Engine/Source/Runtime/Launch/Resources/Mac/Info-Editor.plist
@@ -54,7 +54,9 @@
 	<key>CFBundleInfoDictionaryVersion</key>
 	<string>6.0</string>
 	<key>CFBundleName</key>
-	<string>${EXECUTABLE_NAME}</string>
+	<string>Unreal Editor</string>
+	<key>CFBundleDisplayName</key>
+	<string>Unreal Editor</string>
 	<key>CFBundlePackageType</key>
 	<string>APPL</string>
 	<key>CFBundleShortVersionString</key>
@@ -65,8 +67,6 @@
 	<string>${BUNDLE_VERSION}</string>
 	<key>LSMinimumSystemVersion</key>
 	<string>${MACOSX_DEPLOYMENT_TARGET}</string>
-	<key>NSMainNibFile</key>
-	<string>MainMenu</string>
 	<key>NSPrincipalClass</key>
 	<string>NSApplication</string>
 </dict>
diff --git a/Engine/Source/Runtime/Launch/Resources/Mac/Info.plist b/Engine/Source/Runtime/Launch/Resources/Mac/Info.plist
--- a/Engine/Source/Runtime/Launch/Resources/Mac/Info.plist
+++ b/Engine/Source/Runtime/Launch/Resources/Mac/Info.plist
@@ -24,8 +24,6 @@
 	<string>${BUNDLE_VERSION}</string>
 	<key>LSMinimumSystemVersion</key>
 	<string>${MACOSX_DEPLOYMENT_TARGET}</string>
-	<key>NSMainNibFile</key>
-	<string>MainMenu</string>
 	<key>NSPrincipalClass</key>
 	<string>NSApplication</string>
 </dict>
