Linux toolchain: disable RTTI.

(Shaves off ~1 MB from a stripped Shipping binary).

diff --git a/Engine/Source/Programs/UnrealBuildTool/Linux/LinuxToolChain.cs b/Engine/Source/Programs/UnrealBuildTool/Linux/LinuxToolChain.cs
--- a/Engine/Source/Programs/UnrealBuildTool/Linux/LinuxToolChain.cs
+++ b/Engine/Source/Programs/UnrealBuildTool/Linux/LinuxToolChain.cs
@@ -95,6 +95,7 @@ namespace UnrealBuildTool.Linux
             //Result += " -Wshadow";                      // additional warning not normally included in Wall: warns if there variable/typedef shadows some other variable - not helpful because we have gobs of code that shadows variables
             Result += " -mmmx -msse -msse2";            // allows use of SIMD intrinsics
             Result += " -fno-math-errno";               // do not assume that math ops have side effects
+            Result += " -fno-rtti";                     // no run-time type info
 
             if (String.IsNullOrEmpty(ClangPath))
             {
@@ -144,7 +145,7 @@ namespace UnrealBuildTool.Linux
             }
             else if (CompileEnvironment.Config.TargetConfiguration == CPPTargetConfiguration.Debug)
             {
-		Result += " -fno-inline";                   // disable inlining for better debuggability (e.g. callstacks, "skip file" in gdb)
+                Result += " -fno-inline";                   // disable inlining for better debuggability (e.g. callstacks, "skip file" in gdb)
             }
 
             // debug info
