[INTEGRATE] Merging latest LauncherReleases changes to main, included adding additional logging information, latest version number update, and fixes to offline support

diff --git a/Engine/Source/Programs/UnrealBuildTool/Mac/MacToolChain.cs b/Engine/Source/Programs/UnrealBuildTool/Mac/MacToolChain.cs
--- a/Engine/Source/Programs/UnrealBuildTool/Mac/MacToolChain.cs
+++ b/Engine/Source/Programs/UnrealBuildTool/Mac/MacToolChain.cs
@@ -910,6 +910,10 @@ namespace UnrealBuildTool
 					}
 					AppendMacLine(CreateAppBundleScript, "cp -f \"{0}\" \"{2}.app/Contents/Resources/{1}.icns\"", CustomIcon, IconName, ExeName);
 
+					if (ExeName.StartsWith("UE4Editor"))
+					{
+					}
+
 					if (ExeName.StartsWith("UE4Editor"))
 					{
 						AppendMacLine(CreateAppBundleScript, "cp -f \"{0}/Runtime/Launch/Resources/Mac/UProject.icns\" \"{1}.app/Contents/Resources/UProject.icns\"", EngineSourcePath, ExeName);
