Exclude NotForLicensees modules when generating projects for rocket

#ttp 334121: UBT ProjectFileGenerator should not assume NoRedist is the only way to protect files

diff --git a/Engine/Source/Programs/UnrealBuildTool/System/ProjectFileGenerator.cs b/Engine/Source/Programs/UnrealBuildTool/System/ProjectFileGenerator.cs
--- a/Engine/Source/Programs/UnrealBuildTool/System/ProjectFileGenerator.cs
+++ b/Engine/Source/Programs/UnrealBuildTool/System/ProjectFileGenerator.cs
@@ -847,10 +847,7 @@ namespace UnrealBuildTool
 				// Skip NoRedist files if necessary
 				if( bExcludeNoRedistFiles )
 				{
-					if( CleanBuildFileName.IndexOf( Path.DirectorySeparatorChar + "NoRedist" + Path.DirectorySeparatorChar, StringComparison.InvariantCultureIgnoreCase ) != -1 )
-					{
-						IncludeThisModule = false;
-					}
+					IncludeThisModule = !IsNoRedistModule(CleanBuildFileName);
 				}
 
 				if( IncludeThisModule )
@@ -861,6 +858,32 @@ namespace UnrealBuildTool
 			return AllModuleFiles;
 		}
 
+		/// <summary>
+		/// List of non-redistributable folders
+		/// </summary>
+		private static string[] NoRedistFolders = new string[]
+		{
+			Path.DirectorySeparatorChar + "NoRedist" + Path.DirectorySeparatorChar,
+			Path.DirectorySeparatorChar + "NotForLicensees" + Path.DirectorySeparatorChar
+		};
+
+		/// <summary>
+		/// Checks if a module is in a non-redistributable folder
+		/// </summary>
+		/// <param name="CleanBuildFileName"></param>
+		/// <param name="IncludeThisModule"></param>
+		/// <returns></returns>
+		private static bool IsNoRedistModule(string ModulePath)
+		{
+			foreach (var NoRedistFolderName in NoRedistFolders)
+			{
+				if (ModulePath.IndexOf(NoRedistFolderName, StringComparison.InvariantCultureIgnoreCase) >= 0)
+				{
+					return true;
+				}
+			}
+			return false;
+		}
 
 		/// <summary>
 		/// Finds all target files (filtering by platform)
@@ -893,10 +916,7 @@ namespace UnrealBuildTool
 				// Skip NoRedist files if necessary
 				if( bExcludeNoRedistFiles )
 				{
-					if( CleanTargetFileName.IndexOf( Path.DirectorySeparatorChar + "NoRedist" + Path.DirectorySeparatorChar, StringComparison.InvariantCultureIgnoreCase ) != -1 )
-					{
-						IncludeThisTarget = false;
-					}
+					IncludeThisTarget = !IsNoRedistModule(CleanTargetFileName);
 				}
 
 				if( IncludeThisTarget )
