UBT: support different architectures (target triplets) in Linux toolchain.

diff --git a/Engine/Source/Programs/UnrealBuildTool/Linux/LinuxToolChain.cs b/Engine/Source/Programs/UnrealBuildTool/Linux/LinuxToolChain.cs
--- a/Engine/Source/Programs/UnrealBuildTool/Linux/LinuxToolChain.cs
+++ b/Engine/Source/Programs/UnrealBuildTool/Linux/LinuxToolChain.cs
@@ -16,6 +16,11 @@ namespace UnrealBuildTool.Linux
             return ExternalExecution.GetRuntimePlatform() != UnrealTargetPlatform.Linux;
         }
 
+        protected static bool UsingClang()
+        {
+            return !String.IsNullOrEmpty(ClangPath);
+        }
+
         private string Which(string name)
         {
             Process proc = new Process();
@@ -60,7 +65,8 @@ namespace UnrealBuildTool.Linux
 
                 BaseLinuxPath = BaseLinuxPath.Replace("\"", "");
 
-                // set up the path to our toolchains
+                // set up the path to our toolchains (FIXME: support switching per architecture)
+                GCCPath = "";
                 ClangPath = Path.Combine(BaseLinuxPath, @"bin\Clang++.exe");
                 ArPath = Path.Combine(BaseLinuxPath, @"bin\x86_64-unknown-linux-gnu-ar.exe");
                 RanlibPath = Path.Combine(BaseLinuxPath, @"bin\x86_64-unknown-linux-gnu-ranlib.exe");
@@ -181,10 +187,13 @@ namespace UnrealBuildTool.Linux
 
             // assume we will not perform 32 bit builds on Linux, so define
             // _LINUX64 in any case
-            Result += " -D _LINUX64";
+            Result += " -D_LINUX64";
             if (CrossCompiling())
             {
-                Result += " -target x86_64-unknown-linux-gnu";        // Set target triple
+                if (UsingClang())
+                {
+                    Result += String.Format(" -target {0}", CompileEnvironment.Config.TargetArchitecture);        // Set target triple
+                }
                 Result += String.Format(" --sysroot={0}", BaseLinuxPath);
             }
 
@@ -270,7 +279,10 @@ namespace UnrealBuildTool.Linux
 
             if (CrossCompiling())
             {
-                Result += " -target x86_64-unknown-linux-gnu";        // Set target triple
+                if (UsingClang())
+                {
+                    Result += String.Format(" -target {0}", LinkEnvironment.Config.TargetArchitecture);        // Set target triple
+                }
                 string SysRootPath = BaseLinuxPath.TrimEnd(new char[] { '\\', '/' });
                 Result += String.Format(" \"--sysroot={0}\"", SysRootPath);
             }
@@ -448,7 +460,7 @@ namespace UnrealBuildTool.Linux
                 FileArguments += string.Format(" \"{0}\"", SourceFile.AbsolutePath);
 
                 CompileAction.WorkingDirectory = Path.GetFullPath(".");
-                if (String.IsNullOrEmpty(ClangPath))
+                if (!UsingClang())
                 {
                     CompileAction.CommandPath = GCCPath;
                 }
diff --git a/Engine/Source/Programs/UnrealBuildTool/Linux/UEBuildLinux.cs b/Engine/Source/Programs/UnrealBuildTool/Linux/UEBuildLinux.cs
--- a/Engine/Source/Programs/UnrealBuildTool/Linux/UEBuildLinux.cs
+++ b/Engine/Source/Programs/UnrealBuildTool/Linux/UEBuildLinux.cs
@@ -16,6 +16,15 @@ namespace UnrealBuildTool.Linux
         /** Platform name (embeds architecture for now) */
         static private string TargetPlatformName = "Linux_x64";
 
+        /** Linux architecture (compiler target triplet) */
+        static private string DefaultArchitecture = "x86_64-unknown-linux-gnu";
+
+        /** The current architecture */
+        public override string GetActiveArchitecture()
+        {
+            return DefaultArchitecture;
+        }
+
         /** 
          * Whether platform supports switching SDKs during runtime
          * 
