[GitHub] 195 : Make ProjectFileGenerator.GenerateProjectFiles virtual

diff --git a/Engine/Source/Programs/UnrealBuildTool/System/ProjectFileGenerator.cs b/Engine/Source/Programs/UnrealBuildTool/System/ProjectFileGenerator.cs
--- a/Engine/Source/Programs/UnrealBuildTool/System/ProjectFileGenerator.cs
+++ b/Engine/Source/Programs/UnrealBuildTool/System/ProjectFileGenerator.cs
@@ -250,7 +250,7 @@ namespace UnrealBuildTool
 		/// </summary>
 		/// <param name="Arguments">Command-line arguments</param>
 		/// <param name="bSuccess">True if everything went OK</param>
-		public void GenerateProjectFiles( String[] Arguments, out bool bSuccess )
+		public virtual void GenerateProjectFiles( String[] Arguments, out bool bSuccess )
 		{
 			bSuccess = true;
 
