Fixing shared pch warnings. Making sure Engine.h is not included in any Engine public header nor in any public header of modules used by the engine.

diff --git a/Engine/Plugins/OculusRift/Source/OculusRift/Private/OculusRiftHMD.cpp b/Engine/Plugins/OculusRift/Source/OculusRift/Private/OculusRiftHMD.cpp
--- a/Engine/Plugins/OculusRift/Source/OculusRift/Private/OculusRiftHMD.cpp
+++ b/Engine/Plugins/OculusRift/Source/OculusRift/Private/OculusRiftHMD.cpp
@@ -1,5 +1,6 @@
 // Copyright 1998-2014 Epic Games, Inc. All Rights Reserved.
 
+#include "OculusRiftPrivate.h"
 #include "OculusRiftHMD.h"
 
 #define DEFAULT_PREDICTION_IN_SECONDS 0.035
diff --git a/Engine/Plugins/OculusRift/Source/OculusRift/Private/OculusRiftPrivate.h b/Engine/Plugins/OculusRift/Source/OculusRift/Private/OculusRiftPrivate.h
new file mode 100644
--- /dev/null
+++ b/Engine/Plugins/OculusRift/Source/OculusRift/Private/OculusRiftPrivate.h
@@ -0,0 +1,6 @@
+// Copyright 1998-2012 Epic Games, Inc. All Rights Reserved.
+
+#pragma once
+
+#include "Engine.h"
+#include "IHeadMountedDisplay.h"
\ No newline at end of file
diff --git a/Engine/Plugins/OculusRift/Source/OculusRift/Private/OculusRiftRender.cpp b/Engine/Plugins/OculusRift/Source/OculusRift/Private/OculusRiftRender.cpp
--- a/Engine/Plugins/OculusRift/Source/OculusRift/Private/OculusRiftRender.cpp
+++ b/Engine/Plugins/OculusRift/Source/OculusRift/Private/OculusRiftRender.cpp
@@ -1,7 +1,7 @@
 // Copyright 1998-2014 Epic Games, Inc. All Rights Reserved.
 //
+#include "OculusRiftPrivate.h"
 #include "OculusRiftHMD.h"
-
 #include "RendererPrivate.h"
 #include "ScenePrivate.h"
 #include "PostProcess/PostProcessHMD.h"
diff --git a/Engine/Source/Editor/UnrealEd/Public/Kismet2/EnumEditorUtils.h b/Engine/Source/Editor/UnrealEd/Public/Kismet2/EnumEditorUtils.h
--- a/Engine/Source/Editor/UnrealEd/Public/Kismet2/EnumEditorUtils.h
+++ b/Engine/Source/Editor/UnrealEd/Public/Kismet2/EnumEditorUtils.h
@@ -5,7 +5,6 @@
 
 #pragma once
 
-#include "Engine.h"
 #include "ListenerManager.h"
 
 class UNREALED_API FEnumEditorUtils
diff --git a/Engine/Source/Editor/UnrealEd/Public/Kismet2/KismetEditorUtilities.h b/Engine/Source/Editor/UnrealEd/Public/Kismet2/KismetEditorUtilities.h
--- a/Engine/Source/Editor/UnrealEd/Public/Kismet2/KismetEditorUtilities.h
+++ b/Engine/Source/Editor/UnrealEd/Public/Kismet2/KismetEditorUtilities.h
@@ -6,10 +6,6 @@
 
 #pragma once
 
-
-#include "Engine.h"
-#include "Slate.h"
-
 //////////////////////////////////////////////////////////////////////////
 // FKismetEditorUtilities
 
diff --git a/Engine/Source/Editor/UnrealEd/Public/Kismet2/StructureEditorUtils.h b/Engine/Source/Editor/UnrealEd/Public/Kismet2/StructureEditorUtils.h
--- a/Engine/Source/Editor/UnrealEd/Public/Kismet2/StructureEditorUtils.h
+++ b/Engine/Source/Editor/UnrealEd/Public/Kismet2/StructureEditorUtils.h
@@ -1,7 +1,6 @@
 // Copyright 1998-2014 Epic Games, Inc. All Rights Reserved.
 #pragma once
 
-#include "Engine.h"
 #include "ListenerManager.h"
 #include "Editor/UnrealEd/Classes/UserDefinedStructure/UserDefinedStructEditorData.h"
 
diff --git a/Engine/Source/Programs/UnrealBuildTool/Configuration/UEBuildModule.cs b/Engine/Source/Programs/UnrealBuildTool/Configuration/UEBuildModule.cs
--- a/Engine/Source/Programs/UnrealBuildTool/Configuration/UEBuildModule.cs
+++ b/Engine/Source/Programs/UnrealBuildTool/Configuration/UEBuildModule.cs
@@ -1129,8 +1129,17 @@ namespace UnrealBuildTool
 
 				// Determine what potential precompiled header is used by each source file.
 				double SharedPCHTotalTime = 0.0;
+				if (this.Name == "Engine")
+				{
+					Log.TraceVerbose("ActorConstruction.cpp");
+				}
 				foreach( var CPPFile in SourceFilesToBuild.CPPFiles )
 				{
+					
+					if (CPPFile.AbsolutePath.IndexOf("ActorConstruction.cpp", StringComparison.InvariantCultureIgnoreCase) >= 0)
+					{
+						Log.TraceVerbose("ActorConstruction.cpp");
+					}
 					if (bUseSharedPCHFiles)
 					{
 						var SharedPCHStartTime = DateTime.UtcNow;
diff --git a/Engine/Source/Runtime/Engine/Private/CheatManager.cpp b/Engine/Source/Runtime/Engine/Private/CheatManager.cpp
--- a/Engine/Source/Runtime/Engine/Private/CheatManager.cpp
+++ b/Engine/Source/Runtime/Engine/Private/CheatManager.cpp
@@ -7,6 +7,7 @@
 #include "NavDataGenerator.h"
 #include "OnlineSubsystemUtils.h"
 #include "VisualLog.h"
+#include "GameFramework/Character.h"
 
 #if WITH_EDITOR
 #include "UnrealEd.h"
diff --git a/Engine/Source/Runtime/Engine/Private/GameMode.cpp b/Engine/Source/Runtime/Engine/Private/GameMode.cpp
--- a/Engine/Source/Runtime/Engine/Private/GameMode.cpp
+++ b/Engine/Source/Runtime/Engine/Private/GameMode.cpp
@@ -8,6 +8,9 @@
 #include "GameFramework/GameNetworkManager.h"
 #include "Matinee/MatineeActor.h"
 #include "OnlineSubsystemUtils.h"
+#include "GameFramework/HUD.h"
+#include "GameFramework/DefaultPawn.h"
+#include "GameFramework/SpectatorPawn.h"
 
 DEFINE_LOG_CATEGORY_STATIC(LogGameMode, Log, All);
 
diff --git a/Engine/Source/Runtime/Engine/Private/PlayerController.cpp b/Engine/Source/Runtime/Engine/Private/PlayerController.cpp
--- a/Engine/Source/Runtime/Engine/Private/PlayerController.cpp
+++ b/Engine/Source/Runtime/Engine/Private/PlayerController.cpp
@@ -5,6 +5,7 @@
 #include "Net/UnrealNetwork.h"
 #include "Engine/WorldComposition.h"
 #include "GameFramework/GameNetworkManager.h"
+#include "GameFramework/Character.h"
 #include "ConfigCacheIni.h"
 #include "SoundDefinitions.h"
 #include "OnlineSubsystemUtils.h"
@@ -16,6 +17,8 @@
 #include "Matinee/InterpTrackInstDirector.h"
 #include "Matinee/MatineeActor.h"
 #include "Engine/ActorChannel.h"
+#include "GameFramework/SpectatorPawn.h"
+#include "GameFramework/HUD.h"
 
 DEFINE_LOG_CATEGORY_STATIC(LogPlayerController, Log, All);
 
diff --git a/Engine/Source/Runtime/Engine/Private/UnrealClient.cpp b/Engine/Source/Runtime/Engine/Private/UnrealClient.cpp
--- a/Engine/Source/Runtime/Engine/Private/UnrealClient.cpp
+++ b/Engine/Source/Runtime/Engine/Private/UnrealClient.cpp
@@ -12,6 +12,8 @@
 
 #include "Slate.h"
 #include "SNotificationList.h"
+#include "Engine/PostProcessVolume.h"
+#include "EngineModule.h"
 
 DEFINE_LOG_CATEGORY_STATIC(LogClient, Log, All);
 
diff --git a/Engine/Source/Runtime/Engine/Private/UnrealEngine.cpp b/Engine/Source/Runtime/Engine/Private/UnrealEngine.cpp
--- a/Engine/Source/Runtime/Engine/Private/UnrealEngine.cpp
+++ b/Engine/Source/Runtime/Engine/Private/UnrealEngine.cpp
@@ -42,6 +42,10 @@
 #include "ActiveSound.h"
 #include "DeviceProfiles/DeviceProfileManager.h"
 #include "Animation/SkeletalMeshActor.h"
+#include "GameFramework/HUD.h"
+#include "GameFramework/Character.h"
+#include "Engine/LevelStreamingVolume.h"
+#include "Vehicles/TireType.h"
 
 #include "Particles/Spawn/ParticleModuleSpawn.h"
 #include "Particles/TypeData/ParticleModuleTypeDataMesh.h"
@@ -65,6 +69,7 @@
 #endif
 
 #include "HardwareInfo.h"
+#include "EngineModule.h"
 
 DEFINE_LOG_CATEGORY_STATIC(LogEngine, Log, All);
 
diff --git a/Engine/Source/Runtime/Engine/Private/World.cpp b/Engine/Source/Runtime/Engine/Private/World.cpp
--- a/Engine/Source/Runtime/Engine/Private/World.cpp
+++ b/Engine/Source/Runtime/Engine/Private/World.cpp
@@ -8,6 +8,7 @@
 #include "Engine/LevelScriptBlueprint.h"
 #include "Engine/CullDistanceVolume.h"
 #include "GameFramework/DefaultPhysicsVolume.h"
+#include "GameFramework/CharacterMovementComponent.h"
 #include "Engine/Console.h"
 #include "Engine/WorldComposition.h"
 #include "Matinee/MatineeActor.h"
diff --git a/Engine/Source/Runtime/Engine/Public/ChartCreation.h b/Engine/Source/Runtime/Engine/Public/ChartCreation.h
--- a/Engine/Source/Runtime/Engine/Public/ChartCreation.h
+++ b/Engine/Source/Runtime/Engine/Public/ChartCreation.h
@@ -8,9 +8,6 @@
 #ifndef _CHART_CREATION_INC_
 #define _CHART_CREATION_INC_
 
-
-#include "Engine.h"
-
 #define DO_CHARTING 1
 #if DO_CHARTING
 
diff --git a/Engine/Source/Runtime/Engine/Public/ClearReplacementShaders.h b/Engine/Source/Runtime/Engine/Public/ClearReplacementShaders.h
--- a/Engine/Source/Runtime/Engine/Public/ClearReplacementShaders.h
+++ b/Engine/Source/Runtime/Engine/Public/ClearReplacementShaders.h
@@ -3,7 +3,6 @@
 
 #pragma once
 
-#include "Engine.h"
 #include "GlobalShader.h"
 #include "ShaderParameters.h"
 
diff --git a/Engine/Source/Runtime/Engine/Public/EngineAnalytics.h b/Engine/Source/Runtime/Engine/Public/EngineAnalytics.h
--- a/Engine/Source/Runtime/Engine/Public/EngineAnalytics.h
+++ b/Engine/Source/Runtime/Engine/Public/EngineAnalytics.h
@@ -2,8 +2,6 @@
 
 #pragma once
 
-#include "Engine.h"
-
 class IAnalyticsProvider;
 
 /**
diff --git a/Engine/Source/Runtime/HeadMountedDisplay/Private/HeadMountedDisplayModule.cpp b/Engine/Source/Runtime/HeadMountedDisplay/Private/HeadMountedDisplayModule.cpp
--- a/Engine/Source/Runtime/HeadMountedDisplay/Private/HeadMountedDisplayModule.cpp
+++ b/Engine/Source/Runtime/HeadMountedDisplay/Private/HeadMountedDisplayModule.cpp
@@ -1,6 +1,6 @@
 // Copyright 1998-2014 Epic Games, Inc. All Rights Reserved.
 
-#include "HeadMountedDisplay.h"
+#include "HeadMountedDisplayPrivate.h"
 
 class FHeadMountedDisplayModule : public IHeadMountedDisplayModule
 {
diff --git a/Engine/Source/Runtime/HeadMountedDisplay/Private/HeadMountedDisplayPrivate.h b/Engine/Source/Runtime/HeadMountedDisplay/Private/HeadMountedDisplayPrivate.h
new file mode 100644
--- /dev/null
+++ b/Engine/Source/Runtime/HeadMountedDisplay/Private/HeadMountedDisplayPrivate.h
@@ -0,0 +1,7 @@
+// Copyright 1998-2014 Epic Games, Inc. All Rights Reserved.
+
+#pragma once
+
+#include "Engine.h"
+#include "IHeadMountedDisplayModule.h"
+#include "IHeadMountedDisplay.h"
\ No newline at end of file
diff --git a/Engine/Source/Runtime/HeadMountedDisplay/Public/HeadMountedDisplay.h b/Engine/Source/Runtime/HeadMountedDisplay/Public/HeadMountedDisplay.h
--- a/Engine/Source/Runtime/HeadMountedDisplay/Public/HeadMountedDisplay.h
+++ b/Engine/Source/Runtime/HeadMountedDisplay/Public/HeadMountedDisplay.h
@@ -3,8 +3,6 @@
 #ifndef __HEADMOUNTEDDISPLAY_H__
 #define __HEADMOUNTEDDISPLAY_H__
 
-#include "Core.h"
-#include "Engine.h"
 #include "IHeadMountedDisplayModule.h"
 #include "IHeadMountedDisplay.h"
 
diff --git a/Engine/Source/Runtime/Online/OnlineSubsystemUtils/Public/OnlineSubsystemUtils.h b/Engine/Source/Runtime/Online/OnlineSubsystemUtils/Public/OnlineSubsystemUtils.h
--- a/Engine/Source/Runtime/Online/OnlineSubsystemUtils/Public/OnlineSubsystemUtils.h
+++ b/Engine/Source/Runtime/Online/OnlineSubsystemUtils/Public/OnlineSubsystemUtils.h
@@ -3,7 +3,6 @@
 #pragma once
 
 #include "Core.h"
-#include "Engine.h"
 #include "ModuleManager.h"
 #include "OnlineSubsystemUtilsModule.h"
 #include "Online.h"
diff --git a/Engine/Source/Runtime/Renderer/Private/ScenePrivate.h b/Engine/Source/Runtime/Renderer/Private/ScenePrivate.h
--- a/Engine/Source/Runtime/Renderer/Private/ScenePrivate.h
+++ b/Engine/Source/Runtime/Renderer/Private/ScenePrivate.h
@@ -7,8 +7,6 @@
 #ifndef __SCENEPRIVATE_H__
 #define __SCENEPRIVATE_H__
 
-#include "RendererPrivate.h"
-
 class SceneRenderingBitArrayAllocator
 	: public TInlineAllocator<4,SceneRenderingAllocator>
 {
diff --git a/Engine/Source/Runtime/Renderer/Private/ShaderBaseClasses.h b/Engine/Source/Runtime/Renderer/Private/ShaderBaseClasses.h
--- a/Engine/Source/Runtime/Renderer/Private/ShaderBaseClasses.h
+++ b/Engine/Source/Runtime/Renderer/Private/ShaderBaseClasses.h
@@ -6,7 +6,6 @@
 
 #pragma once
 
-#include "Engine.h"
 #include "MaterialShader.h"
 #include "MeshMaterialShader.h"
 #include "AtmosphereRendering.h"
diff --git a/Engine/Source/Runtime/SkillSystem/Private/GameplayCueView.cpp b/Engine/Source/Runtime/SkillSystem/Private/GameplayCueView.cpp
--- a/Engine/Source/Runtime/SkillSystem/Private/GameplayCueView.cpp
+++ b/Engine/Source/Runtime/SkillSystem/Private/GameplayCueView.cpp
@@ -3,6 +3,7 @@
 #pragma once
 
 #include "SkillSystemModulePrivatePCH.h"
+#include "GameplayEffect.h"
 #include "GameplayCueView.h"
 #include "ParticleDefinitions.h"
 #include "SoundDefinitions.h"
