Removing debug code

diff --git a/Engine/Source/Programs/UnrealBuildTool/Configuration/UEBuildModule.cs b/Engine/Source/Programs/UnrealBuildTool/Configuration/UEBuildModule.cs
--- a/Engine/Source/Programs/UnrealBuildTool/Configuration/UEBuildModule.cs
+++ b/Engine/Source/Programs/UnrealBuildTool/Configuration/UEBuildModule.cs
@@ -1129,10 +1129,6 @@ namespace UnrealBuildTool
 
 				// Determine what potential precompiled header is used by each source file.
 				double SharedPCHTotalTime = 0.0;
-				if (this.Name == "Engine")
-				{
-					Log.TraceVerbose("ActorConstruction.cpp");
-				}
 				foreach( var CPPFile in SourceFilesToBuild.CPPFiles )
 				{
 					
