GitHub pull request #194: Full binary type. Removing hardcoded binary file extensions

diff --git a/Engine/Source/Programs/UnrealBuildTool/Android/AndroidToolChain.cs b/Engine/Source/Programs/UnrealBuildTool/Android/AndroidToolChain.cs
--- a/Engine/Source/Programs/UnrealBuildTool/Android/AndroidToolChain.cs
+++ b/Engine/Source/Programs/UnrealBuildTool/Android/AndroidToolChain.cs
@@ -398,9 +398,10 @@ namespace UnrealBuildTool
 			}
 			if (CompileEnvironment.Config.PrecompiledHeaderAction == PrecompiledHeaderAction.Include)
 			{
+				var PCHExtension = UEBuildPlatform.BuildPlatformDictionary[UnrealTargetPlatform.Android].GetBinaryExtension(UEBuildBinaryType.PrecompiledHeader);
 				// Add the precompiled header file's path to the include path so Clang can find it.
 				// This needs to be before the other include paths to ensure Clang uses it instead of the source header file.
-				PCHArguments += string.Format(" -include \"{0}\"", CompileEnvironment.PrecompiledHeaderFile.AbsolutePath.Replace(".gch", ""));
+				PCHArguments += string.Format(" -include \"{0}\"", CompileEnvironment.PrecompiledHeaderFile.AbsolutePath.Replace(PCHExtension, ""));
 			}
 
 			// Add include paths to the argument list.
@@ -469,11 +470,13 @@ namespace UnrealBuildTool
 
 				if (CompileEnvironment.Config.PrecompiledHeaderAction == PrecompiledHeaderAction.Create)
 				{
+					var PCHExtension = UEBuildPlatform.BuildPlatformDictionary[UnrealTargetPlatform.Android].GetBinaryExtension(UEBuildBinaryType.PrecompiledHeader);
+
 					// Add the precompiled header file to the produced item list.
 					FileItem PrecompiledHeaderFile = FileItem.GetItemByPath(
 						Path.Combine(
 							CompileEnvironment.Config.OutputDirectory,
-							Path.GetFileName(SourceFile.AbsolutePath) + ".gch"
+							Path.GetFileName(SourceFile.AbsolutePath) + PCHExtension
 							)
 						);
 
@@ -491,11 +494,13 @@ namespace UnrealBuildTool
 						CompileAction.PrerequisiteItems.Add(CompileEnvironment.PrecompiledHeaderFile);
 					}
 
+					var ObjectFileExtension = UEBuildPlatform.BuildPlatformDictionary[UnrealTargetPlatform.Android].GetBinaryExtension(UEBuildBinaryType.Object);
+
 					// Add the object file to the produced item list.
 					FileItem ObjectFile = FileItem.GetItemByPath(
 						Path.Combine(
 							CompileEnvironment.Config.OutputDirectory,
-							Path.GetFileName(SourceFile.AbsolutePath) + ".o"
+							Path.GetFileName(SourceFile.AbsolutePath) + ObjectFileExtension
 							)
 						);
 					CompileAction.ProducedItems.Add(ObjectFile);
diff --git a/Engine/Source/Programs/UnrealBuildTool/Android/UEBuildAndroid.cs b/Engine/Source/Programs/UnrealBuildTool/Android/UEBuildAndroid.cs
--- a/Engine/Source/Programs/UnrealBuildTool/Android/UEBuildAndroid.cs
+++ b/Engine/Source/Programs/UnrealBuildTool/Android/UEBuildAndroid.cs
@@ -95,11 +95,15 @@ namespace UnrealBuildTool
 					return ".so";
 				case UEBuildBinaryType.StaticLibrary:
 					return ".a";
+				case UEBuildBinaryType.Object:
+					return ".o";
+				case UEBuildBinaryType.PrecompiledHeader:
+					return ".gch";
 			}
 			return base.GetBinaryExtension(InBinaryType);
 		}
 
-		public override string GetDebugInfoExtension( UEBuildBinaryType InBinaryType )
+		public override string GetDebugInfoExtension(UEBuildBinaryType InBinaryType)
 		{
 			return "";
 		}
diff --git a/Engine/Source/Programs/UnrealBuildTool/Configuration/UEBuildBinary.cs b/Engine/Source/Programs/UnrealBuildTool/Configuration/UEBuildBinary.cs
--- a/Engine/Source/Programs/UnrealBuildTool/Configuration/UEBuildBinary.cs
+++ b/Engine/Source/Programs/UnrealBuildTool/Configuration/UEBuildBinary.cs
@@ -11,14 +11,16 @@ using System.Xml;
 namespace UnrealBuildTool
 {	
 	/// <summary>
-	/// The type of binary to build
+	/// All binary types generated by UBT
 	/// </summary>
 	public enum UEBuildBinaryType
 	{
 		Executable,
 		DynamicLinkLibrary,
-		StaticLibrary
-	};
+		StaticLibrary,
+		Object,
+		PrecompiledHeader
+	}
 
 	/// <summary>
 	/// UEBuildBinary configuration
diff --git a/Engine/Source/Programs/UnrealBuildTool/HTML5/HTML5ToolChain.cs b/Engine/Source/Programs/UnrealBuildTool/HTML5/HTML5ToolChain.cs
--- a/Engine/Source/Programs/UnrealBuildTool/HTML5/HTML5ToolChain.cs
+++ b/Engine/Source/Programs/UnrealBuildTool/HTML5/HTML5ToolChain.cs
@@ -350,12 +350,12 @@ namespace UnrealBuildTool
 
 				// Add the source file path to the command-line.
                 string FileArguments = string.Format(" \"{0}\"", SourceFile.AbsolutePath);
-
+				var ObjectFileExtension = UEBuildPlatform.BuildPlatformDictionary[UnrealTargetPlatform.HTML5].GetBinaryExtension(UEBuildBinaryType.Object);
             	// Add the object file to the produced item list.
 				FileItem ObjectFile = FileItem.GetItemByPath(
 					Path.Combine(
 						CompileEnvironment.Config.OutputDirectory,
-						Path.GetFileName(SourceFile.AbsolutePath) + (".o")
+					Path.GetFileName(SourceFile.AbsolutePath) + ObjectFileExtension
 						)
 					);
 				CompileAction.ProducedItems.Add(ObjectFile);
diff --git a/Engine/Source/Programs/UnrealBuildTool/HTML5/UEBuildHTML5.cs b/Engine/Source/Programs/UnrealBuildTool/HTML5/UEBuildHTML5.cs
--- a/Engine/Source/Programs/UnrealBuildTool/HTML5/UEBuildHTML5.cs
+++ b/Engine/Source/Programs/UnrealBuildTool/HTML5/UEBuildHTML5.cs
@@ -118,6 +118,10 @@ namespace UnrealBuildTool
 						return ".exe";
 					case UEBuildBinaryType.StaticLibrary:
 						return ".lib";
+					case UEBuildBinaryType.Object:
+						return ".o";
+					case UEBuildBinaryType.PrecompiledHeader:
+						return ".gch";
 				}
 				return base.GetBinaryExtension(InBinaryType);
 			}
diff --git a/Engine/Source/Programs/UnrealBuildTool/IOS/IOSToolChain.cs b/Engine/Source/Programs/UnrealBuildTool/IOS/IOSToolChain.cs
--- a/Engine/Source/Programs/UnrealBuildTool/IOS/IOSToolChain.cs
+++ b/Engine/Source/Programs/UnrealBuildTool/IOS/IOSToolChain.cs
@@ -427,7 +427,8 @@ namespace UnrealBuildTool
 			{
 				// Add the precompiled header file's path to the include path so GCC can find it.
 				// This needs to be before the other include paths to ensure GCC uses it instead of the source header file.
-				PCHArguments += string.Format(" -include \"{0}\"", ConvertPath(CompileEnvironment.PrecompiledHeaderFile.AbsolutePath.Replace(".gch", "")));
+				var PrecompiledFileExtension = UEBuildPlatform.BuildPlatformDictionary[UnrealTargetPlatform.IOS].GetBinaryExtension(UEBuildBinaryType.PrecompiledHeader);
+				PCHArguments += string.Format(" -include \"{0}\"", ConvertPath(CompileEnvironment.PrecompiledHeaderFile.AbsolutePath.Replace(PrecompiledFileExtension, "")));
 			}
 
 			// Add include paths to the argument list.
@@ -526,11 +527,12 @@ namespace UnrealBuildTool
 
 				if (CompileEnvironment.Config.PrecompiledHeaderAction == PrecompiledHeaderAction.Create)
 				{
+					var PrecompiledFileExtension = UEBuildPlatform.BuildPlatformDictionary[UnrealTargetPlatform.IOS].GetBinaryExtension(UEBuildBinaryType.PrecompiledHeader);
 					// Add the precompiled header file to the produced item list.
 					FileItem PrecompiledHeaderFile = FileItem.GetItemByPath(
 						Path.Combine(
 							CompileEnvironment.Config.OutputDirectory,
-							Path.GetFileName(SourceFile.AbsolutePath) + ".gch"
+						Path.GetFileName(SourceFile.AbsolutePath) + PrecompiledFileExtension
 							)
 						);
 
@@ -548,12 +550,12 @@ namespace UnrealBuildTool
 						CompileAction.bIsUsingPCH = true;
 						CompileAction.PrerequisiteItems.Add(CompileEnvironment.PrecompiledHeaderFile);
 					}
-
+					var ObjectFileExtension = UEBuildPlatform.BuildPlatformDictionary[UnrealTargetPlatform.IOS].GetBinaryExtension(UEBuildBinaryType.Object);
 					// Add the object file to the produced item list.
 					FileItem ObjectFile = FileItem.GetItemByPath(
 						Path.Combine(
 							CompileEnvironment.Config.OutputDirectory,
-							Path.GetFileName(SourceFile.AbsolutePath) + ".o"
+							Path.GetFileName(SourceFile.AbsolutePath) + ObjectFileExtension
 							)
 						);
 
diff --git a/Engine/Source/Programs/UnrealBuildTool/IOS/UEBuildIOS.cs b/Engine/Source/Programs/UnrealBuildTool/IOS/UEBuildIOS.cs
--- a/Engine/Source/Programs/UnrealBuildTool/IOS/UEBuildIOS.cs
+++ b/Engine/Source/Programs/UnrealBuildTool/IOS/UEBuildIOS.cs
@@ -84,6 +84,10 @@ namespace UnrealBuildTool
 					return "";
 				case UEBuildBinaryType.StaticLibrary:
 					return ".a";
+				case UEBuildBinaryType.Object:
+					return ".o";
+				case UEBuildBinaryType.PrecompiledHeader:
+					return ".gch";
 			}
 			return base.GetBinaryExtension(InBinaryType);
 		}
diff --git a/Engine/Source/Programs/UnrealBuildTool/Linux/LinuxToolChain.cs b/Engine/Source/Programs/UnrealBuildTool/Linux/LinuxToolChain.cs
--- a/Engine/Source/Programs/UnrealBuildTool/Linux/LinuxToolChain.cs
+++ b/Engine/Source/Programs/UnrealBuildTool/Linux/LinuxToolChain.cs
@@ -352,7 +352,8 @@ namespace UnrealBuildTool.Linux
             {
                 // Add the precompiled header file's path to the include path so Clang can find it.
                 // This needs to be before the other include paths to ensure Clang uses it instead of the source header file.
-                PCHArguments += string.Format(" -include \"{0}\"", CompileEnvironment.PrecompiledHeaderFile.AbsolutePath.Replace(".gch", ""));
+				var PrecompiledFileExtension = UEBuildPlatform.BuildPlatformDictionary[UnrealTargetPlatform.Linux].GetBinaryExtension(UEBuildBinaryType.PrecompiledHeader);
+                PCHArguments += string.Format(" -include \"{0}\"", CompileEnvironment.PrecompiledHeaderFile.AbsolutePath.Replace(PrecompiledFileExtension, ""));
             }
 
             // Add include paths to the argument list.
@@ -421,11 +422,12 @@ namespace UnrealBuildTool.Linux
 
                 if (CompileEnvironment.Config.PrecompiledHeaderAction == PrecompiledHeaderAction.Create)
                 {
+					var PrecompiledFileExtension = UEBuildPlatform.BuildPlatformDictionary[UnrealTargetPlatform.Linux].GetBinaryExtension(UEBuildBinaryType.PrecompiledHeader);
                     // Add the precompiled header file to the produced item list.
                     FileItem PrecompiledHeaderFile = FileItem.GetItemByPath(
                         Path.Combine(
                             CompileEnvironment.Config.OutputDirectory,
-                            Path.GetFileName(SourceFile.AbsolutePath) + ".gch"
+                            Path.GetFileName(SourceFile.AbsolutePath) + PrecompiledFileExtension
                             )
                         );
 
@@ -443,11 +445,12 @@ namespace UnrealBuildTool.Linux
                         CompileAction.PrerequisiteItems.Add(CompileEnvironment.PrecompiledHeaderFile);
                     }
 
+					var ObjectFileExtension = UEBuildPlatform.BuildPlatformDictionary[UnrealTargetPlatform.Linux].GetBinaryExtension(UEBuildBinaryType.Object);
                     // Add the object file to the produced item list.
                     FileItem ObjectFile = FileItem.GetItemByPath(
                         Path.Combine(
                             CompileEnvironment.Config.OutputDirectory,
-                            Path.GetFileName(SourceFile.AbsolutePath) + ".o"
+                            Path.GetFileName(SourceFile.AbsolutePath) + ObjectFileExtension
                             )
                         );
                     CompileAction.ProducedItems.Add(ObjectFile);
diff --git a/Engine/Source/Programs/UnrealBuildTool/Linux/UEBuildLinux.cs b/Engine/Source/Programs/UnrealBuildTool/Linux/UEBuildLinux.cs
--- a/Engine/Source/Programs/UnrealBuildTool/Linux/UEBuildLinux.cs
+++ b/Engine/Source/Programs/UnrealBuildTool/Linux/UEBuildLinux.cs
@@ -168,6 +168,10 @@ namespace UnrealBuildTool.Linux
                     return "";
                 case UEBuildBinaryType.StaticLibrary:
                     return ".a";
+				case UEBuildBinaryType.Object:
+					return ".o";
+				case UEBuildBinaryType.PrecompiledHeader:
+					return ".gch";
             }
             return base.GetBinaryExtension(InBinaryType);
         }
diff --git a/Engine/Source/Programs/UnrealBuildTool/Mac/MacToolChain.cs b/Engine/Source/Programs/UnrealBuildTool/Mac/MacToolChain.cs
--- a/Engine/Source/Programs/UnrealBuildTool/Mac/MacToolChain.cs
+++ b/Engine/Source/Programs/UnrealBuildTool/Mac/MacToolChain.cs
@@ -234,7 +234,8 @@ namespace UnrealBuildTool
 			{
 				// Add the precompiled header file's path to the include path so GCC can find it.
 				// This needs to be before the other include paths to ensure GCC uses it instead of the source header file.
-				PCHArguments += string.Format(" -include \"{0}\"", CompileEnvironment.PrecompiledHeaderFile.AbsolutePath.Replace(".gch", ""));
+				var PrecompiledFileExtension = UEBuildPlatform.BuildPlatformDictionary[UnrealTargetPlatform.Mac].GetBinaryExtension(UEBuildBinaryType.PrecompiledHeader);
+				PCHArguments += string.Format(" -include \"{0}\"", CompileEnvironment.PrecompiledHeaderFile.AbsolutePath.Replace(PrecompiledFileExtension, ""));
 			}
 
 			// Add include paths to the argument list.
@@ -333,11 +334,12 @@ namespace UnrealBuildTool
 
 				if (CompileEnvironment.Config.PrecompiledHeaderAction == PrecompiledHeaderAction.Create)
 				{
+					var PrecompiledHeaderExtension = UEBuildPlatform.BuildPlatformDictionary[UnrealTargetPlatform.Mac].GetBinaryExtension(UEBuildBinaryType.PrecompiledHeader);
 					// Add the precompiled header file to the produced item list.
 					FileItem PrecompiledHeaderFile = FileItem.GetItemByPath(
 						Path.Combine(
 							CompileEnvironment.Config.OutputDirectory,
-							Path.GetFileName(SourceFile.AbsolutePath) + ".gch"
+							Path.GetFileName(SourceFile.AbsolutePath) + PrecompiledHeaderExtension
 							)
 						);
 
@@ -355,12 +357,12 @@ namespace UnrealBuildTool
 						CompileAction.bIsUsingPCH = true;
 						CompileAction.PrerequisiteItems.Add(CompileEnvironment.PrecompiledHeaderFile);
 					}
-
+					var ObjectFileExtension = UEBuildPlatform.BuildPlatformDictionary[UnrealTargetPlatform.Mac].GetBinaryExtension(UEBuildBinaryType.Object);
 					// Add the object file to the produced item list.
 					FileItem ObjectFile = FileItem.GetItemByPath(
 						Path.Combine(
 							CompileEnvironment.Config.OutputDirectory,
-							Path.GetFileName(SourceFile.AbsolutePath) + ".o"
+							Path.GetFileName(SourceFile.AbsolutePath) + ObjectFileExtension
 							)
 						);
 
diff --git a/Engine/Source/Programs/UnrealBuildTool/Mac/UEBuildMac.cs b/Engine/Source/Programs/UnrealBuildTool/Mac/UEBuildMac.cs
--- a/Engine/Source/Programs/UnrealBuildTool/Mac/UEBuildMac.cs
+++ b/Engine/Source/Programs/UnrealBuildTool/Mac/UEBuildMac.cs
@@ -67,6 +67,10 @@ namespace UnrealBuildTool
 					return "";
 				case UEBuildBinaryType.StaticLibrary:
 					return ".a";
+				case UEBuildBinaryType.Object:
+					return ".o";
+				case UEBuildBinaryType.PrecompiledHeader:
+					return ".gch";
 			}
 			return base.GetBinaryExtension(InBinaryType);
 		}
@@ -78,7 +82,7 @@ namespace UnrealBuildTool
 		 *	
 		 *	@return	string				The debug info extension (i.e. 'pdb')
 		 */
-		public override string GetDebugInfoExtension( UEBuildBinaryType InBinaryType )
+		public override string GetDebugInfoExtension(UEBuildBinaryType InBinaryType)
 		{
 			return BuildConfiguration.bGeneratedSYMFile ? ".dsym" : "";
 		}
diff --git a/Engine/Source/Programs/UnrealBuildTool/WinRT/UEBuildWinRT.cs b/Engine/Source/Programs/UnrealBuildTool/WinRT/UEBuildWinRT.cs
--- a/Engine/Source/Programs/UnrealBuildTool/WinRT/UEBuildWinRT.cs
+++ b/Engine/Source/Programs/UnrealBuildTool/WinRT/UEBuildWinRT.cs
@@ -126,6 +126,10 @@ namespace UnrealBuildTool
 					return ".exe";
 				case UEBuildBinaryType.StaticLibrary:
 					return ".lib";
+				case UEBuildBinaryType.Object:
+					return ".obj";
+				case UEBuildBinaryType.PrecompiledHeader:
+					return ".pch";
 			}
 			return base.GetBinaryExtension(InBinaryType);
 		}
@@ -137,7 +141,7 @@ namespace UnrealBuildTool
 		 *	
 		 *	@return	string				The debug info extension (i.e. 'pdb')
 		 */
-		public override string GetDebugInfoExtension( UEBuildBinaryType InBinaryType )
+		public override string GetDebugInfoExtension(UEBuildBinaryType InBinaryType)
 		{
 			return ".pdb";
 		}
diff --git a/Engine/Source/Programs/UnrealBuildTool/WinRT/WinRTToolChain.cs b/Engine/Source/Programs/UnrealBuildTool/WinRT/WinRTToolChain.cs
--- a/Engine/Source/Programs/UnrealBuildTool/WinRT/WinRTToolChain.cs
+++ b/Engine/Source/Programs/UnrealBuildTool/WinRT/WinRTToolChain.cs
@@ -523,11 +523,13 @@ namespace UnrealBuildTool
 					string OriginalPCHHeaderDirectory = Path.GetDirectoryName(SourceFile.AbsolutePath);
 					FileArguments += string.Format(" /I \"{0}\"", OriginalPCHHeaderDirectory);
 
+					var PCHExtension = UEBuildPlatform.BuildPlatformDictionary[UnrealTargetPlatform.WinRT].GetBinaryExtension(UEBuildBinaryType.PrecompiledHeader);
+
 					// Add the precompiled header file to the produced items list.
 					FileItem PrecompiledHeaderFile = FileItem.GetItemByPath(
 						Path.Combine(
 							CompileEnvironment.Config.OutputDirectory,
-							Path.GetFileName(SourceFile.AbsolutePath) + ".pch"
+							Path.GetFileName(SourceFile.AbsolutePath) + PCHExtension
 							)
 						);
 					CompileAction.ProducedItems.Add(PrecompiledHeaderFile);
@@ -578,11 +580,13 @@ namespace UnrealBuildTool
 					CompileAction.StatusDescription = Path.GetFileName(SourceFile.AbsolutePath);
 				}
 
+				var ObjectFileExtension = UEBuildPlatform.BuildPlatformDictionary[UnrealTargetPlatform.WinRT].GetBinaryExtension(UEBuildBinaryType.Object);
+
 				// Add the object file to the produced item list.
 				FileItem ObjectFile = FileItem.GetItemByPath(
 					Path.Combine(
 						CompileEnvironment.Config.OutputDirectory,
-						Path.GetFileName(SourceFile.AbsolutePath) + ".obj"
+						Path.GetFileName(SourceFile.AbsolutePath) + ObjectFileExtension
 						)
 					);
 				CompileAction.ProducedItems.Add(ObjectFile);
diff --git a/Engine/Source/Programs/UnrealBuildTool/Windows/UEBuildWindows.cs b/Engine/Source/Programs/UnrealBuildTool/Windows/UEBuildWindows.cs
--- a/Engine/Source/Programs/UnrealBuildTool/Windows/UEBuildWindows.cs
+++ b/Engine/Source/Programs/UnrealBuildTool/Windows/UEBuildWindows.cs
@@ -173,6 +173,10 @@ namespace UnrealBuildTool
 					return ".exe";
 				case UEBuildBinaryType.StaticLibrary:
 					return ".lib";
+				case UEBuildBinaryType.Object:
+					return ".obj";
+				case UEBuildBinaryType.PrecompiledHeader:
+					return ".pch";
 			}
 			return base.GetBinaryExtension(InBinaryType);
 		}
@@ -210,8 +214,7 @@ namespace UnrealBuildTool
             switch (InBinaryType)
             {
                 case UEBuildBinaryType.DynamicLinkLibrary:
-                    return ".pdb";
-                case UEBuildBinaryType.Executable:
+				case UEBuildBinaryType.Executable:
                     return ".pdb";
             }
             return "";
diff --git a/Engine/Source/Programs/UnrealBuildTool/Windows/VCToolChain.cs b/Engine/Source/Programs/UnrealBuildTool/Windows/VCToolChain.cs
--- a/Engine/Source/Programs/UnrealBuildTool/Windows/VCToolChain.cs
+++ b/Engine/Source/Programs/UnrealBuildTool/Windows/VCToolChain.cs
@@ -733,11 +733,12 @@ namespace UnrealBuildTool
 					string OriginalPCHHeaderDirectory = Path.GetDirectoryName( SourceFile.AbsolutePath );
 					FileArguments += string.Format(" /I \"{0}\"", OriginalPCHHeaderDirectory);
 
+					var PrecompiledFileExtension = UEBuildPlatform.BuildPlatformDictionary[UnrealTargetPlatform.Win64].GetBinaryExtension(UEBuildBinaryType.PrecompiledHeader);
 					// Add the precompiled header file to the produced items list.
 					FileItem PrecompiledHeaderFile = FileItem.GetItemByPath(
 						Path.Combine(
 							CompileEnvironment.Config.OutputDirectory,
-							Path.GetFileName(SourceFile.AbsolutePath) + ".pch"
+						Path.GetFileName(SourceFile.AbsolutePath) + PrecompiledFileExtension
 							)
 						);
 					CompileAction.ProducedItems.Add(PrecompiledHeaderFile);
@@ -815,11 +816,12 @@ namespace UnrealBuildTool
 
 				if( bEmitsObjectFile )
 				{
+					var ObjectFileExtension = UEBuildPlatform.BuildPlatformDictionary[UnrealTargetPlatform.Win64].GetBinaryExtension(UEBuildBinaryType.Object);
 					// Add the object file to the produced item list.
 					FileItem ObjectFile = FileItem.GetItemByPath(
 						Path.Combine(
 							CompileEnvironment.Config.OutputDirectory,
-							Path.GetFileName(SourceFile.AbsolutePath) + ".obj"
+							Path.GetFileName(SourceFile.AbsolutePath) + ObjectFileExtension
 							)
 						);
 					CompileAction.ProducedItems.Add(ObjectFile);
