UBT XML configuration improvements:
- XSD driven IntelliSense
- defaults read from code (default XML file will be regenerated on UBT startup if is different than it should be)
- configurable classes/fields have to be now annotated with XmlConfig and XmlConfigField attributes

diff --git a/Engine/Programs/UnrealBuildTool/BuildConfiguration.xml b/Engine/Programs/UnrealBuildTool/BuildConfiguration.xml
deleted file mode 100644
diff --git a/Engine/Source/Programs/UnrealBuildTool/Android/UEBuildAndroid.cs b/Engine/Source/Programs/UnrealBuildTool/Android/UEBuildAndroid.cs
--- a/Engine/Source/Programs/UnrealBuildTool/Android/UEBuildAndroid.cs
+++ b/Engine/Source/Programs/UnrealBuildTool/Android/UEBuildAndroid.cs
@@ -9,13 +9,17 @@ using System.Xml;
 
 namespace UnrealBuildTool
 {
+	[XmlConfig]
 	class AndroidPlatform : UEBuildPlatform
 	{
 		/// <summary>
 		/// Android settings.
 		/// </summary>
+		[XmlConfigField]
 		public static string AndroidArchitecture = "-armv7";
+		[XmlConfigField]
 		public static string AndroidNdkApiTarget = "latest";
+		[XmlConfigField]
 		public static string AndroidSdkApiTarget = "latest";
 
 		// The current architecture - affects everything about how UBT operates on Android
diff --git a/Engine/Source/Programs/UnrealBuildTool/Configuration/BuildConfiguration.cs b/Engine/Source/Programs/UnrealBuildTool/Configuration/BuildConfiguration.cs
--- a/Engine/Source/Programs/UnrealBuildTool/Configuration/BuildConfiguration.cs
+++ b/Engine/Source/Programs/UnrealBuildTool/Configuration/BuildConfiguration.cs
@@ -8,6 +8,7 @@ using System.Reflection;
 
 namespace UnrealBuildTool
 {
+	[XmlConfigAttribute(true)]
 	public class BuildConfiguration
 	{
 		/// <summary>
@@ -295,11 +296,11 @@ namespace UnrealBuildTool
 		/// <summary>
 		/// Sets the configuration back to defaults.
 		/// </summary>
-		public static void Reset()
+		public static void LoadDefaults()
 		{
 			bAllowLTCG = false;
 			bAllowRemotelyCompiledPCHs = false;
-            bAllowXGE = true;
+			bAllowXGE = true;
 
 			// Don't bother to check external (stable) headers for modification.  It slows down UBT's dependency checking.
 			bCheckExternalHeadersForModification = false;
@@ -343,7 +344,7 @@ namespace UnrealBuildTool
 			// these private PCHs do yield fastest incremental compiles for single files, it causes full rebuilds to
 			// take an inordinate amount of time, and intermediates will use up many gigabytes of disk space.
 			bUseSharedPCHs = true;
-	
+
 			WriteTargetInfoPath = null;
 
 			// Using unity build to coalesce source files allows for much faster full rebuild times.  For fastest iteration times on single-file
@@ -351,7 +352,7 @@ namespace UnrealBuildTool
 			// the synthesized file names for unity code files may change between builds.
 			bUseUnityBuild = true;
 
-            bForceUnityBuild = false;
+			bForceUnityBuild = false;
 
 			// For programmers working on gameplay code, it's convenient to have fast iteration times when changing only single source
 			// files.  We've set this default such that smaller gameplay modules will not use the unity build feature, allowing them
@@ -379,7 +380,7 @@ namespace UnrealBuildTool
 			bTestIncludeDependencyResolveCache = false;
 			// if we are testing the resolve cache, we require UBT to use it.
 			bUseIncludeDependencyResolveCache = true;
-            
+
 			//IMPORTANT THIS IS THE MAIN SWITCH FOR MONO FAST CALLS
 			//  if this is set to true, then fast calls will be on by default on Dingo, and if false it will be off by default on Dingo.
 			//  This can be overridden by -fastmonocalls  or -nofastmonocalls in the NMAKE params.
@@ -394,9 +395,6 @@ namespace UnrealBuildTool
 
 			// set up some paths
 			BaseIntermediateFolder = "Intermediate/Build/";
-
-			// Load eventual XML configuration files if they exist to override default values.
-			XmlConfigLoader.Load(typeof(BuildConfiguration));
 		}
 
 		/// <summary>
diff --git a/Engine/Source/Programs/UnrealBuildTool/Configuration/UEBuildConfiguration.cs b/Engine/Source/Programs/UnrealBuildTool/Configuration/UEBuildConfiguration.cs
--- a/Engine/Source/Programs/UnrealBuildTool/Configuration/UEBuildConfiguration.cs
+++ b/Engine/Source/Programs/UnrealBuildTool/Configuration/UEBuildConfiguration.cs
@@ -5,6 +5,7 @@ using System.IO;
 
 namespace UnrealBuildTool
 {
+	[XmlConfig(true)]
 	public class UEBuildConfiguration
 	{
 		/** Whether to include PhysX support */
@@ -121,8 +122,10 @@ namespace UnrealBuildTool
 		/** True if we want to favor optimizing size over speed */
 		public static bool bCompileForSize;
 
-		/** Sets the configuration back to defaults. */
-		public static void Reset()
+		/// <summary>
+		/// Sets the configuration back to defaults.
+		/// </summary>
+		public static void LoadDefaults()
 		{
 			//@todo. Allow disabling PhysX/APEX via these values...
 			// Currently, WITH_PHYSX is forced to true in Engine.h (as it isn't defined anywhere by the builder)
@@ -147,16 +150,20 @@ namespace UnrealBuildTool
 			bCompileSpeedTree = true;
 			bCompileWithStatsWithoutEngine = false;
 			bCompileWithPluginSupport = false;
-            bUseLoggingInShipping = false;
+			bUseLoggingInShipping = false;
 			bCompileSteamOSS = true;
 			bCompileMcpOSS = true;
-            bOBBinAPK = false;
+			bOBBinAPK = false;
 			bCompilePhysXVehicle = true;
 			bCompileFreeType = true;
 			bCompileForSize = false;
+		}
 
-			XmlConfigLoader.Load(typeof(UEBuildConfiguration));
-
+		/// <summary>
+		/// Function to call to after reset default data.
+		/// </summary>
+		public static void PostReset()
+		{
 			// Configuration overrides.
 			string SteamVersion = "Steamv129a";
 			bCompileSteamOSS = bCompileSteamOSS
@@ -168,13 +175,13 @@ namespace UnrealBuildTool
 
 			bCompileSimplygon = bCompileSimplygon
 				&& Directory.Exists(UEBuildConfiguration.UEThirdPartyDirectory + "NotForLicensees") == true
-                && Directory.Exists(UEBuildConfiguration.UEThirdPartyDirectory + "NotForLicensees/Simplygon") == true
+				&& Directory.Exists(UEBuildConfiguration.UEThirdPartyDirectory + "NotForLicensees/Simplygon") == true
 				&& Directory.Exists("Developer/SimplygonMeshReduction") == true
 				&& !(ProjectFileGenerator.bGenerateProjectFiles && ProjectFileGenerator.bGeneratingRocketProjectFiles);
 
 			bCompileSpeedTree = bCompileSpeedTree
-                && Directory.Exists(UEBuildConfiguration.UEThirdPartyDirectory + "NotForLicensees") == true
-                && Directory.Exists(UEBuildConfiguration.UEThirdPartyDirectory + "NotForLicensees/SpeedTree") == true
+				&& Directory.Exists(UEBuildConfiguration.UEThirdPartyDirectory + "NotForLicensees") == true
+				&& Directory.Exists(UEBuildConfiguration.UEThirdPartyDirectory + "NotForLicensees/SpeedTree") == true
 				&& !(ProjectFileGenerator.bGenerateProjectFiles && ProjectFileGenerator.bGeneratingRocketProjectFiles);
 		}
 
diff --git a/Engine/Source/Programs/UnrealBuildTool/HTML5/UEBuildHTML5.cs b/Engine/Source/Programs/UnrealBuildTool/HTML5/UEBuildHTML5.cs
--- a/Engine/Source/Programs/UnrealBuildTool/HTML5/UEBuildHTML5.cs
+++ b/Engine/Source/Programs/UnrealBuildTool/HTML5/UEBuildHTML5.cs
@@ -8,8 +8,10 @@ using System.IO;
 
 namespace UnrealBuildTool
 {
+	[XmlConfig]
 	class HTML5Platform : UEBuildPlatform
 	{
+		[XmlConfigField]
 		public static string HTML5Architecture = "";
 
 		// The current architecture - affects everything about how UBT operates on HTML5
diff --git a/Engine/Source/Programs/UnrealBuildTool/IOS/IOSToolChain.cs b/Engine/Source/Programs/UnrealBuildTool/IOS/IOSToolChain.cs
--- a/Engine/Source/Programs/UnrealBuildTool/IOS/IOSToolChain.cs
+++ b/Engine/Source/Programs/UnrealBuildTool/IOS/IOSToolChain.cs
@@ -15,6 +15,7 @@ using Ionic.Zlib;
 
 namespace UnrealBuildTool
 {
+	[XmlConfig(true)]
 	class IOSToolChain : RemoteToolChain
 	{
 		public override void RegisterToolChain()
@@ -70,12 +71,10 @@ namespace UnrealBuildTool
 		public List<UEBuildFramework> RememberedAdditionalFrameworks = new List<UEBuildFramework>();
 
 		/// <summary>
-		/// Function to call to reset default data.
+		/// Function to call to after reset default data.
 		/// </summary>
-		public static void Reset()
+		public static void PostReset()
 		{
-			XmlConfigLoader.Load(typeof(IOSToolChain));
-
 			/** Location of the SDKs */
 			BaseSDKDir = XcodeDeveloperDir + "Platforms/iPhoneOS.platform/Developer/SDKs";
 			BaseSDKDirSim = XcodeDeveloperDir + "Platforms/iPhoneSimulator.platform/Developer/SDKs";
diff --git a/Engine/Source/Programs/UnrealBuildTool/IOS/UEBuildIOS.cs b/Engine/Source/Programs/UnrealBuildTool/IOS/UEBuildIOS.cs
--- a/Engine/Source/Programs/UnrealBuildTool/IOS/UEBuildIOS.cs
+++ b/Engine/Source/Programs/UnrealBuildTool/IOS/UEBuildIOS.cs
@@ -9,9 +9,11 @@ using System.Xml;
 
 namespace UnrealBuildTool
 {
+	[XmlConfig]
 	class IOSPlatform : UEBuildPlatform
 	{
 		// by default, use an empty architecture (which is really just a modifer to the platform for some paths/names)
+		[XmlConfigField]
 		public static string IOSArchitecture = "";
 
 		// The current architecture - affects everything about how UBT operates on IOS
diff --git a/Engine/Source/Programs/UnrealBuildTool/System/ProjectFileGenerator.cs b/Engine/Source/Programs/UnrealBuildTool/System/ProjectFileGenerator.cs
--- a/Engine/Source/Programs/UnrealBuildTool/System/ProjectFileGenerator.cs
+++ b/Engine/Source/Programs/UnrealBuildTool/System/ProjectFileGenerator.cs
@@ -575,6 +575,11 @@ namespace UnrealBuildTool
 		/// <param name="EngineProject">Engine project to add files to.</param>
 		private void AddUBTConfigFilesToEngineProject(ProjectFile EngineProject)
 		{
+			EngineProject.AddAliasedFileToProject(new AliasedFile(
+					XmlConfigLoader.GetXSDPath(),
+					Path.Combine("Config", "UnrealBuildTool", "Schema")
+				));
+
 			foreach(var BuildConfigurationPath in XmlConfigLoader.ConfigLocationHierarchy)
 			{
 				if(!BuildConfigurationPath.bExists)
diff --git a/Engine/Source/Programs/UnrealBuildTool/System/UnrealBuildTool.cs b/Engine/Source/Programs/UnrealBuildTool/System/UnrealBuildTool.cs
--- a/Engine/Source/Programs/UnrealBuildTool/System/UnrealBuildTool.cs
+++ b/Engine/Source/Programs/UnrealBuildTool/System/UnrealBuildTool.cs
@@ -601,8 +601,8 @@ namespace UnrealBuildTool
 			ECompilationResult Result = ECompilationResult.Succeeded;
 
 			// Reset early so we can access BuildConfiguration even before RunUBT() is called
-			BuildConfiguration.Reset();
-			UEBuildConfiguration.Reset();
+			XmlConfigLoader.Reset<BuildConfiguration>();
+			XmlConfigLoader.Reset<UEBuildConfiguration>();
 
 			Log.TraceVerbose("UnrealBuildTool (DEBUG OUTPUT MODE)");
 			Log.TraceVerbose("Command-line: {0}", String.Join(" ", Arguments));
@@ -1018,9 +1018,9 @@ namespace UnrealBuildTool
 
 			// Reset global configurations
 			ResetAllActions();
-			CPPEnvironment.Reset();
-			BuildConfiguration.Reset();
-			UEBuildConfiguration.Reset();
+			XmlConfigLoader.Reset<CPPEnvironment>();
+			XmlConfigLoader.Reset<BuildConfiguration>();
+			XmlConfigLoader.Reset<UEBuildConfiguration>();
 
 			ParseBuildConfigurationFlags(Arguments);
 
diff --git a/Engine/Source/Programs/UnrealBuildTool/ToolChain/RemoteToolChain.cs b/Engine/Source/Programs/UnrealBuildTool/ToolChain/RemoteToolChain.cs
--- a/Engine/Source/Programs/UnrealBuildTool/ToolChain/RemoteToolChain.cs
+++ b/Engine/Source/Programs/UnrealBuildTool/ToolChain/RemoteToolChain.cs
@@ -12,6 +12,7 @@ using System.Text.RegularExpressions;
 
 namespace UnrealBuildTool
 {
+	[XmlConfig]
 	abstract class RemoteToolChain : UEToolChain
     {
 		protected void RegisterRemoteToolChain(UnrealTargetPlatform InPlatform, CPPTargetPlatform CPPPlatform)
@@ -24,7 +25,9 @@ namespace UnrealBuildTool
 		}
 
 		/** These two variables will be loaded from XML config file in XmlConfigLoader.Init() */
+		[XmlConfigField]
 		public static string RemoteServerName;
+		[XmlConfigField]
 		public static string[] PotentialServerNames = new string[] { };
 
 		/** Keep a list of remote files that are potentially copied from local to remote */
diff --git a/Engine/Source/Programs/UnrealBuildTool/Utilities/XmlConfigLoader.cs b/Engine/Source/Programs/UnrealBuildTool/Utilities/XmlConfigLoader.cs
--- a/Engine/Source/Programs/UnrealBuildTool/Utilities/XmlConfigLoader.cs
+++ b/Engine/Source/Programs/UnrealBuildTool/Utilities/XmlConfigLoader.cs
@@ -1,20 +1,159 @@
 ﻿using System;
+using System.Collections;
 using System.Collections.Generic;
 using System.Globalization;
 using System.IO;
 using System.Linq;
 using System.Reflection;
 using System.Text;
+using System.Text.RegularExpressions;
 using System.Threading;
 using System.Xml;
+using System.Xml.Linq;
 
 namespace UnrealBuildTool
 {
+	/// <summary>
+	/// Attribute to annotate classes that can be set using xml configuration system.
+	/// </summary>
+	[AttributeUsage(AttributeTargets.Class | AttributeTargets.Struct)]
+	public class XmlConfigAttribute : Attribute
+	{
+		public readonly bool bAllPublicStaticFields;
+
+		public XmlConfigAttribute(bool bAllPublicStaticFields = false)
+		{
+			this.bAllPublicStaticFields = bAllPublicStaticFields;
+		}
+	}
+
+	/// <summary>
+	/// Attribute to annotate fields in type that can be set using xml configuration system.
+	/// </summary>
+	[AttributeUsage(AttributeTargets.Field)]
+	public class XmlConfigFieldAttribute : Attribute
+	{
+
+	}
+
 	/// <summary>
 	/// Loads data from disk XML files and stores it in memory.
 	/// </summary>
 	public static class XmlConfigLoader
 	{
+		/// <summary>
+		/// Resets given config class.
+		/// </summary>
+		/// <typeparam name="ConfigClass">The class to reset.</typeparam>
+		public static void Reset<ConfigClass>()
+		{
+			var Type = typeof(ConfigClass);
+
+			InvokeIfExists(Type, "LoadDefaults");
+
+			// Load eventual XML configuration files if they exist to override default values.
+			Load(Type);
+
+			InvokeIfExists(Type, "PostReset");
+		}
+
+		/// <summary>
+		/// Creates XML element for given config class.
+		/// </summary>
+		/// <param name="ConfigType">A type of the config class.</param>
+		/// <returns>XML element representing config class.</returns>
+		private static XElement CreateConfigTypeXMLElement(Type ConfigType)
+		{
+			var NS = XNamespace.Get("https://www.unrealengine.com/BuildConfiguration");
+
+			return new XElement(NS + ConfigType.Name,
+				from Field in GetConfigurableFields(ConfigType)
+				where !(Field.FieldType == typeof(string) && Field.GetValue(null) == null)
+				select CreateFieldXMLElement(Field));
+		}
+
+		/// <summary>
+		/// Gets XML representing current config.
+		/// </summary>
+		/// <returns>XML representing current config.</returns>
+		private static XDocument GetConfigXML()
+		{
+			var NS = XNamespace.Get("https://www.unrealengine.com/BuildConfiguration");
+			return new XDocument(
+					new XElement(NS + "Configuration",
+						from ConfigType in GetAllConfigurationTypes()
+						select CreateConfigTypeXMLElement(ConfigType)
+					)
+				);
+		}
+
+		/// <summary>
+		/// Gets default config XML.
+		/// </summary>
+		/// <returns>Default config XML.</returns>
+		public static string GetDefaultXML()
+		{
+			foreach(var ConfigType in GetAllConfigurationTypes())
+			{
+				InvokeIfExists(ConfigType, "LoadDefaults");
+			}
+
+			var Comment = @"
+	#########################################################################
+	#																		#
+	#	This is an XML with default UBT configuration. If you want to		#
+	#	override it, create the same file in the locations c. or d.			#
+	#	(see below). DONT'T CHANGE CONTENTS OF THIS FILE!					#	
+	#																		#
+	#########################################################################
+	
+	The syntax of this file is:
+	<Configuration>
+		<ClassA>
+			<Field1>Value</Field2>
+			<ArrayField2>
+				<Item>First item</Item>
+				<Item>Second item</Item>
+				...
+				<Item>Last item</Item>
+			</ArrayField2>
+			
+			...
+			
+		</ClassA>
+		<ClassB>
+			...
+		</ClassB>
+
+		...
+
+		<ClassN>
+			...
+		</ClassN>
+	</Configuration>
+	
+	There are four possible location for this file:
+		a. UE4/Engine/Programs/UnrealBuildTool
+		b. UE4/Engine/Programs/NotForLicensees/UnrealBuildTool
+		c. UE4/Engine/Saved/UnrealBuildTool
+		d. My Documents/Unreal Engine/UnrealBuildTool
+	
+	The UBT is looking for it in all four places in the given order and
+	overrides already read data with the loaded ones, hence d. has the
+	priority. Not defined classes and fields are left alone.
+";
+
+			var DefaultXml = GetConfigXML();
+
+			DefaultXml.AddFirst(new XComment(Comment));
+
+			return WriteToBuffer(DefaultXml);
+		}
+
+		/// <summary>
+		/// Loads data for a given config class.
+		/// </summary>
+		/// <param name="Class">A type of a config class.</param>
 		public static void Load(Type Class)
 		{
 			XmlConfigLoaderClassData ClassData;
@@ -25,11 +164,54 @@ namespace UnrealBuildTool
 			}
 		}
 
+		/// <summary>
+		/// Builds xsd for BuildConfiguration.xml files.
+		/// </summary>
+		/// <returns>Content of the xsd file.</returns>
+		public static string BuildXSD()
+		{
+			// all elements use this namespace
+			var NS = XNamespace.Get("http://www.w3.org/2001/XMLSchema");
+
+			return WriteToBuffer(new XDocument(
+				// define the root element that declares the schema and target namespace it is validating.
+				new XElement(NS + "schema",
+					new XAttribute("targetNamespace", "https://www.unrealengine.com/BuildConfiguration"),
+					new XAttribute("elementFormDefault", "qualified"),
+					new XElement(NS + "element", new XAttribute("name", "Configuration"),
+						new XElement(NS + "complexType",
+							new XElement(NS + "all",
+								// loop over all public types in UBT assembly
+								from ConfigType in GetAllConfigurationTypes()
+								// find all public static fields of intrinsic types (and a few extra)
+								let PublicStaticFields = GetConfigurableFields(ConfigType).ToList()
+								where PublicStaticFields.Count > 0
+								select
+								new XElement(NS + "element",
+									new XAttribute("name", ConfigType.Name),
+									new XAttribute("minOccurs", "0"),
+									new XAttribute("maxOccurs", "1"),
+									new XElement(NS + "complexType",
+										new XElement(NS + "all",
+											from Field in PublicStaticFields
+											select CreateXSDElementForField(Field)
+										)
+									)
+								)
+							)
+						)
+					)
+				)
+			));
+		}
+
 		/// <summary>
 		/// Initializing data and resets all found classes.
 		/// </summary>
 		public static void Init()
 		{
+			OverwriteIfDifferent(GetXSDPath(), BuildXSD());
+
 			LoadData();
 
 			foreach(var ClassData in Data)
@@ -38,6 +220,52 @@ namespace UnrealBuildTool
 			}
 		}
 
+		/// <summary>
+		/// Overwrites file at FilePath with the Content if the content was
+		/// different. If the file doesn't exist it creates file with the
+		/// Content.
+		/// </summary>
+		/// <param name="FilePath">File to check, overwrite or create.</param>
+		/// <param name="Content">Content to fill the file with.</param>
+		/// <param name="bReadOnlyFile">If file can and should be read-only?</param>
+		private static void OverwriteIfDifferent(string FilePath, string Content, bool bReadOnlyFile = false)
+		{
+			if (FileDifferentThan(FilePath, Content))
+			{
+				if(bReadOnlyFile && File.Exists(FilePath))
+				{
+					var Attributes = File.GetAttributes(FilePath);
+					if(Attributes.HasFlag(FileAttributes.ReadOnly))
+					{
+						File.SetAttributes(FilePath, Attributes & ~FileAttributes.ReadOnly);
+					}
+				}
+
+				File.WriteAllText(FilePath, Content);
+
+				if(bReadOnlyFile)
+				{
+					File.SetAttributes(FilePath, File.GetAttributes(FilePath) | FileAttributes.ReadOnly);
+				}
+			}
+		}
+
+		/// <summary>
+		/// Tells if file at given path has different content than given.
+		/// </summary>
+		/// <param name="FilePath">Path of the file to check.</param>
+		/// <param name="Content">Content to check.</param>
+		/// <returns>True if file at FilePath has different content than Content. False otherwise.</returns>
+		private static bool FileDifferentThan(string FilePath, string Content)
+		{
+			if (!File.Exists(FilePath))
+			{
+				return true;
+			}
+
+			return !File.ReadAllText(FilePath).Equals(Content, StringComparison.InvariantCulture);
+		}
+
 		/// <summary>
 		/// Cache entry class to store loaded info for given class.
 		/// </summary>
@@ -168,7 +396,7 @@ namespace UnrealBuildTool
 			public bool bCreateIfDoesNotExist { get; private set; }
 
 			// Tells if config file exists in this location.
-			public bool bExists { get; private set; }
+			public bool bExists { get; protected set; }
 
 			public XmlConfigLocation(string[] FSLocations, string IDEFolderName, bool bCreateIfDoesNotExist = false)
 			{
@@ -188,7 +416,7 @@ namespace UnrealBuildTool
 			/// <summary>
 			/// Creates template file in the FS location.
 			/// </summary>
-			public void CreateUserXmlConfigTemplate()
+			public virtual void CreateUserXmlConfigTemplate()
 			{
 				try
 				{
@@ -197,7 +425,8 @@ namespace UnrealBuildTool
 					var FilePath = Path.Combine(FSLocation);
 
 					const string TemplateContent =
-						"<Configuration>\n" +
+						"<?xml version=\"1.0\" encoding=\"utf-8\" ?>" +
+						"<Configuration xmlns=\"https://www.unrealengine.com/BuildConfiguration\">\n" +
 						"	<BuildConfiguration>\n" +
 						"	</BuildConfiguration>\n" +
 						"</Configuration>\n";
@@ -211,6 +440,55 @@ namespace UnrealBuildTool
 				}
 			}
 
+			/// <summary>
+			/// Tells if procedure should try to create file for this location.
+			/// </summary>
+			/// <returns>True if procedure should try to create file for this location. False otherwise.</returns>
+			public virtual bool IfShouldCreateFile()
+			{
+				return bCreateIfDoesNotExist && !bExists;
+			}
+		}
+
+		/// <summary>
+		/// Class that stores information about possible default BuildConfiguration.xml.
+		/// </summary>
+		public class XmlDefaultConfigLocation : XmlConfigLocation
+		{
+			public XmlDefaultConfigLocation(string FSLocation)
+				: base(FSLocation, "Default", true)
+			{
+
+			}
+			
+			/// <summary>
+			/// Creates template file in the FS location.
+			/// </summary>
+			public override void CreateUserXmlConfigTemplate()
+			{
+				try
+				{
+					Directory.CreateDirectory(Path.GetDirectoryName(FSLocation));
+
+					var FilePath = Path.Combine(FSLocation);
+
+					OverwriteIfDifferent(FilePath, GetDefaultXML(), true);
+					bExists = true;
+				}
+				catch (Exception)
+				{
+					// Ignore quietly.
+				}
+			}
+
+			/// <summary>
+			/// Tells if procedure should try to create file for this location.
+			/// </summary>
+			/// <returns>True if procedure should try to create file for this location. False otherwise.</returns>
+			public override bool IfShouldCreateFile()
+			{
+				return true;
+			}
 		}
 
 		public static readonly XmlConfigLocation[] ConfigLocationHierarchy;
@@ -234,7 +512,7 @@ namespace UnrealBuildTool
 
 			ConfigLocationHierarchy = new XmlConfigLocation[]
 			{
-				new XmlConfigLocation(Path.Combine(UE4EnginePath, "Programs", "UnrealBuildTool"), "Default"),
+				new XmlDefaultConfigLocation(Path.Combine(UE4EnginePath, "Programs", "UnrealBuildTool")),
 				new XmlConfigLocation(Path.Combine(UE4EnginePath, "Programs", "NotForLicensees", "UnrealBuildTool"), "NotForLicensees"),
 				new XmlConfigLocation(Path.Combine(UE4EnginePath, "Saved", "UnrealBuildTool"), "User", true),
 				new XmlConfigLocation(new string[] {
@@ -251,22 +529,22 @@ namespace UnrealBuildTool
 		{
 			foreach (var PossibleConfigLocation in ConfigLocationHierarchy)
 			{
-				if(!PossibleConfigLocation.bExists)
+				if (!PossibleConfigLocation.IfShouldCreateFile())
 				{
 					continue;
 				}
 
-				Load(PossibleConfigLocation.FSLocation);
+				PossibleConfigLocation.CreateUserXmlConfigTemplate();
 			}
 
 			foreach (var PossibleConfigLocation in ConfigLocationHierarchy)
 			{
-				if(!PossibleConfigLocation.bCreateIfDoesNotExist || PossibleConfigLocation.bExists)
+				if(!PossibleConfigLocation.bExists)
 				{
 					continue;
 				}
 
-				PossibleConfigLocation.CreateUserXmlConfigTemplate();
+				Load(PossibleConfigLocation.FSLocation);
 			}
 		}
 
@@ -277,9 +555,34 @@ namespace UnrealBuildTool
 		private static void Load(string ConfigurationXmlPath)
 		{
 			var ConfigDocument = new XmlDocument();
+			var NS = new XmlNamespaceManager(ConfigDocument.NameTable);
+			NS.AddNamespace("ns", "https://www.unrealengine.com/BuildConfiguration");
 			ConfigDocument.Load(ConfigurationXmlPath);
 
-			var XmlClasses = ConfigDocument.DocumentElement.SelectNodes("/Configuration/*");
+			var XmlClasses = ConfigDocument.DocumentElement.SelectNodes("/ns:Configuration/*", NS);
+
+			if(XmlClasses.Count == 0)
+			{
+				if(ConfigDocument.DocumentElement.Name == "Configuration")
+				{
+					ConfigDocument.DocumentElement.SetAttribute("xmlns", "https://www.unrealengine.com/BuildConfiguration");
+
+					var NSDoc = new XmlDocument();
+
+					NSDoc.LoadXml(ConfigDocument.OuterXml);
+
+					try
+					{
+						File.WriteAllText(ConfigurationXmlPath, WriteToBuffer(NSDoc));
+					}
+					catch(Exception)
+					{
+						// Ignore gently.
+					}
+
+					XmlClasses = NSDoc.DocumentElement.SelectNodes("/ns:Configuration/*", NS);
+				}
+			}
 
 			foreach (XmlNode XmlClass in XmlClasses)
 			{
@@ -292,6 +595,16 @@ namespace UnrealBuildTool
 				}
 
 				XmlConfigLoaderClassData ClassData;
+				bool bAllPublicStaticFields = false;
+
+				var Attributes = ClassType.GetCustomAttributes(typeof(XmlConfigAttribute), false);
+				if (Attributes.Length == 0)
+				{
+					Log.TraceVerbose("XmlConfig Loading: class '{0}' is not allowed to be configured using XML system.", XmlClass.Name);
+					continue;
+				}
+				bAllPublicStaticFields = ((XmlConfigAttribute)Attributes.First()).bAllPublicStaticFields;
+
 				if (!Data.TryGetValue(ClassType, out ClassData))
 				{
 					ClassData = new XmlConfigLoaderClassData(ClassType);
@@ -304,15 +617,15 @@ namespace UnrealBuildTool
 				{
 					FieldInfo Field = ClassType.GetField(XmlField.Name);
 
-					if (Field == null || !Field.IsPublic || !Field.IsStatic)
+					if (Field == null || !Field.IsPublic || !Field.IsStatic || (!bAllPublicStaticFields && Field.GetCustomAttributes(typeof(XmlConfigFieldAttribute), false).Length == 0))
 					{
-						throw new BuildException("BuildConfiguration Loading: field '{0}' doesn't exist or is either non-public or non-static.", XmlField.Name);
+						throw new BuildException("BuildConfiguration Loading: field '{0}' doesn't exist or is either non-public, non-static or not-xml-configurable.", XmlField.Name);
 					}
 
 					if(Field.FieldType.IsArray)
 					{
 						// If the type is an array type get items for it.
-						var XmlItems = XmlField.SelectNodes("Item");
+						var XmlItems = XmlField.SelectNodes("ns:Item", NS);
 
 						// Get the C# type of the array.
 						var ItemType = Field.FieldType.GetElementType();
@@ -426,6 +739,208 @@ namespace UnrealBuildTool
 			return true;
 		}
 
+		/// <summary>
+		/// Gets all types that can be configured using this xml system.
+		/// </summary>
+		/// <returns>Enumerable of types that can be configured using this xml system.</returns>
+		private static IEnumerable<Type> GetAllConfigurationTypes()
+		{
+			foreach(var Type in Assembly.GetExecutingAssembly().GetTypes())
+			{
+				if(Type.GetCustomAttributes(typeof(XmlConfigAttribute), false).Length > 0)
+				{
+					yield return Type;
+				}
+			}
+		}
+
+		/// <summary>
+		/// Invokes a public static parameterless method from type if it exists.
+		/// </summary>
+		/// <param name="ClassType">Class type to look the method in.</param>
+		/// <param name="MethodName">Name of the method.</param>
+		private static void InvokeIfExists(Type ClassType, string MethodName)
+		{
+			var Method = ClassType.GetMethod(MethodName, new Type[]{}, new ParameterModifier[]{});
+
+			if(Method != null && Method.IsPublic && Method.IsStatic)
+			{
+				Method.Invoke(null, new object[]{});
+			}
+		}
+
+		/// <summary>
+		/// Enumerates all fields that are configurable using this system.
+		/// </summary>
+		/// <param name="ClassType">Type of the class to enumrate field for.</param>
+		/// <returns>Enumerable fields that are configurable using this system.</returns>
+		private static IEnumerable<FieldInfo> GetConfigurableFields(Type ClassType)
+		{
+			var bAllPublicStaticFields = ((XmlConfigAttribute) ClassType.GetCustomAttributes(typeof(XmlConfigAttribute), false).First()).bAllPublicStaticFields;
+
+			foreach(var Field in ClassType.GetFields())
+			{
+				if(Field.IsPublic && Field.IsStatic && (bAllPublicStaticFields || Field.GetCustomAttributes(typeof(XmlConfigFieldAttribute), false).Length > 0))
+				{
+					yield return Field;
+				}
+			}
+		}
+
+		/// <summary>
+		/// Gets xsd type string for the field type.
+		/// </summary>
+		/// <param name="FieldType">Field type to get xsd string for.</param>
+		/// <returns>String representation of xsd type corresponding given field type.</returns>
+		private static string GetFieldXSDType(Type FieldType)
+		{
+			var TypeMap = new Dictionary<Type, string>
+			{
+				{ typeof(string), "string" },
+				{ typeof(bool), "boolean" },
+				{ typeof(int), "int" },
+				{ typeof(long), "long" },
+				{ typeof(DateTime), "dateTime" }
+			};
+
+			return TypeMap.ContainsKey(FieldType)
+				? TypeMap[FieldType]
+				: FieldType.Name.ToLowerInvariant();
+		}
+
+		/// <summary>
+		/// Gets xml representation of an array field value.
+		/// </summary>
+		/// <param name="Field">Array field to represent.</param>
+		/// <returns>Array field value as XML elements.</returns>
+		private static XElement[] GetArrayFieldXMLValue(FieldInfo Field)
+		{
+			var NS = XNamespace.Get("http://www.w3.org/2001/XMLSchema");
+
+			var Array = Field.GetValue(null) as IEnumerable;
+			var ElementList = new List<XElement>();
+
+			foreach(var Element in Array)
+			{
+				ElementList.Add(new XElement(NS + "Item", GetObjectXMLValue(Field.FieldType.GetElementType(), Element)));
+			}
+
+			return ElementList.ToArray();
+		}
+
+		/// <summary>
+		/// XML representation of a regular field value.
+		/// </summary>
+		/// <param name="Field">Field to represent value.</param>
+		/// <returns>Field value as XML.</returns>
+		private static string GetFieldXMLValue(FieldInfo Field)
+		{
+			return GetObjectXMLValue(Field.FieldType, Field.GetValue(null));
+		}
+
+		/// <summary>
+		/// Gets custom C# object XML string representation.
+		/// </summary>
+		/// <param name="FieldType">The type of the field.</param>
+		/// <param name="Obj">The value.</param>
+		/// <returns>String representation.</returns>
+		private static string GetObjectXMLValue(Type FieldType, object Obj)
+		{
+			if (Obj == null && FieldType == typeof(string))
+			{
+				return "";
+			}
+
+			if (FieldType == typeof(bool))
+			{
+				return (bool) Obj ? "true" : "false";
+			}
+
+			return Obj.ToString();
+		}
+
+		/// <summary>
+		/// Gets standard BuildConfiguration.xml schema path.
+		/// </summary>
+		/// <returns>Standard BuildConfiguration.xml schema path.</returns>
+		public static string GetXSDPath()
+		{
+			return new FileInfo(Path.Combine(Utils.GetExecutingAssemblyDirectory(), "..", "..", "Saved", "UnrealBuildTool", "BuildConfiguration.Schema.xsd")).FullName;
+		}
+
+		private static XElement CreateXSDElementForField(FieldInfo Field)
+		{
+			var NS = XNamespace.Get("http://www.w3.org/2001/XMLSchema");
+
+			if (Field.FieldType.IsArray)
+			{
+				return new XElement(NS + "element",
+					new XAttribute("name", Field.Name),
+					new XAttribute("minOccurs", "0"),
+					new XAttribute("maxOccurs", "1"),
+					new XElement(NS + "complexType",
+						new XElement(NS + "all",
+							new XElement(NS + "element",
+								new XAttribute("name", "Item"),
+								new XAttribute("type", GetFieldXSDType(Field.FieldType.GetElementType())),
+								new XAttribute("minOccurs", "0")
+							)
+						)
+					)
+				);
+			}
+
+			return new XElement(NS + "element",
+				new XAttribute("name", Field.Name),
+				new XAttribute("type", GetFieldXSDType(Field.FieldType)),
+				new XAttribute("minOccurs", "0"),
+				new XAttribute("maxOccurs", "1"));
+		}
+
+		private static XElement CreateFieldXMLElement(FieldInfo Field)
+		{
+			var NS = XNamespace.Get("https://www.unrealengine.com/BuildConfiguration");
+
+			if (Field.FieldType.IsArray)
+			{
+				return new XElement(NS + Field.Name, GetArrayFieldXMLValue(Field));
+			}
+			else
+			{
+				return new XElement(NS + Field.Name, GetFieldXMLValue(Field));
+			}
+		}
+
+		private static string WriteToBuffer(XmlDocument Document)
+		{
+			using (var Reader = new XmlNodeReader(Document))
+			{
+				Reader.MoveToContent();
+				return WriteToBuffer(XDocument.Load(Reader));
+			}
+		}
+
+		private static string WriteToBuffer(XDocument Document)
+		{
+			var Settings = new XmlWriterSettings();
+			Settings.Indent = true;
+			Settings.IndentChars = "\t";
+			Settings.CloseOutput = true;
+
+			using (var MS = new MemoryStream())
+			{
+				using (var SW = new StreamWriter(MS, Encoding.UTF8))
+				{
+					using (var Writer = XmlWriter.Create(SW, Settings))
+					{
+						Document.WriteTo(Writer);
+					}
+				}
+
+				return Encoding.UTF8.GetString(MS.ToArray());
+			}
+		}
+
 		// Map to store class data in.
 		private static readonly Dictionary<Type, XmlConfigLoaderClassData> Data = new Dictionary<Type, XmlConfigLoaderClassData>();
 	}
diff --git a/Engine/Source/Programs/UnrealBuildTool/WinRT/UEBuildWinRT.cs b/Engine/Source/Programs/UnrealBuildTool/WinRT/UEBuildWinRT.cs
--- a/Engine/Source/Programs/UnrealBuildTool/WinRT/UEBuildWinRT.cs
+++ b/Engine/Source/Programs/UnrealBuildTool/WinRT/UEBuildWinRT.cs
@@ -8,11 +8,13 @@ using System.IO;
 
 namespace UnrealBuildTool
 {
+	[XmlConfig]
 	class WinRTPlatform : UEBuildPlatform
 	{
 		/// <summary>
 		/// Should the app be compiled as WinRT
 		/// </summary>
+		[XmlConfigField]
 		public static bool bCompileWinRT = false;
 
 		public static bool IsVisualStudioInstalled()
