Properly conditionally include telemetry provider ET.

diff --git a/Engine/Source/Programs/UnrealBuildTool/UnrealBuildTool.csproj b/Engine/Source/Programs/UnrealBuildTool/UnrealBuildTool.csproj
--- a/Engine/Source/Programs/UnrealBuildTool/UnrealBuildTool.csproj
+++ b/Engine/Source/Programs/UnrealBuildTool/UnrealBuildTool.csproj
@@ -174,7 +174,7 @@
     </When>
   </Choose>
   <Choose>
-    <When Condition="Exists('System\NoRedist\TelemetryProviderET.cs')">
+    <When Condition="Exists('System\NotForLicensees\TelemetryProviderET.cs')">
       <ItemGroup>
         <Compile Include="System\NotForLicensees\TelemetryProviderET.cs" />
       </ItemGroup>
@@ -289,9 +289,6 @@
   <ItemGroup>
     <None Include="app.config" />
   </ItemGroup>
-  <ItemGroup>
-    <Folder Include="System\NoRedist\" />
-  </ItemGroup>
   <Import Project="$(MSBuildBinPath)\Microsoft.CSharp.targets" />
   <!-- To modify your build process, add your task inside one of the targets below and uncomment it. 
        Other similar extension points exist, see Microsoft.Common.targets.
