Some missing folders for XMLs and XSD wasn't created on build machines. Now it creates them if they don't exist. CIS fix.

diff --git a/Engine/Source/Programs/UnrealBuildTool/Utilities/XmlConfigLoader.cs b/Engine/Source/Programs/UnrealBuildTool/Utilities/XmlConfigLoader.cs
--- a/Engine/Source/Programs/UnrealBuildTool/Utilities/XmlConfigLoader.cs
+++ b/Engine/Source/Programs/UnrealBuildTool/Utilities/XmlConfigLoader.cs
@@ -241,6 +241,7 @@ namespace UnrealBuildTool
 					}
 				}
 
+				Directory.CreateDirectory(Path.GetDirectoryName(FilePath));
 				File.WriteAllText(FilePath, Content);
 
 				if(bReadOnlyFile)
