Update to AUPlugIn API rather than old, deprecated ComponentManager now that RadioEffectUnit has been updated.

diff --git a/Engine/Source/Programs/UnrealBuildTool/Mac/MacToolChain.cs b/Engine/Source/Programs/UnrealBuildTool/Mac/MacToolChain.cs
--- a/Engine/Source/Programs/UnrealBuildTool/Mac/MacToolChain.cs
+++ b/Engine/Source/Programs/UnrealBuildTool/Mac/MacToolChain.cs
@@ -941,7 +941,6 @@ namespace UnrealBuildTool
 					// Copy RadioEffect component
 					AppendMacLine(CreateAppBundleScript, "cp -f \"{0}/ThirdParty/Mac/RadioEffectUnit/RadioEffectUnit.component/Contents/MacOS/RadioEffectUnit\" \"{1}.app/Contents/Resources/RadioEffectUnit.component/Contents/MacOS/RadioEffectUnit\"", EngineSourcePath, ExeName);
 					AppendMacLine(CreateAppBundleScript, "cp -f \"{0}/ThirdParty/Mac/RadioEffectUnit/RadioEffectUnit.component/Contents/Resources/English.lproj/Localizable.strings\" \"{1}.app/Contents/Resources/RadioEffectUnit.component/Contents/Resources/English.lproj/Localizable.strings\"", EngineSourcePath, ExeName);
-					AppendMacLine(CreateAppBundleScript, "cp -f \"{0}/ThirdParty/Mac/RadioEffectUnit/RadioEffectUnit.component/Contents/Resources/RadioEffectUnit.rsrc\" \"{1}.app/Contents/Resources/RadioEffectUnit.component/Contents/Resources/RadioEffectUnit.rsrc\"", EngineSourcePath, ExeName);
 					AppendMacLine(CreateAppBundleScript, "cp -f \"{0}/ThirdParty/Mac/RadioEffectUnit/RadioEffectUnit.component/Contents/Info.plist\" \"{1}.app/Contents/Resources/RadioEffectUnit.component/Contents/Info.plist\"", EngineSourcePath, ExeName);
 
 					// Make sure OS X knows the bundle was updated
@@ -956,7 +955,6 @@ namespace UnrealBuildTool
 					QueueFileForBatchUpload(FileItem.GetItemByFullPath(Path.GetFullPath("../../Engine/Source/Runtime/Launch/Resources/Mac/Info.plist")));
 					QueueFileForBatchUpload(FileItem.GetItemByFullPath(Path.GetFullPath("../../Engine/Source/ThirdParty/Mac/RadioEffectUnit/RadioEffectUnit.component/Contents/MacOS/RadioEffectUnit")));
 					QueueFileForBatchUpload(FileItem.GetItemByFullPath(Path.GetFullPath("../../Engine/Source/ThirdParty/Mac/RadioEffectUnit/RadioEffectUnit.component/Contents/Resources/English.lproj/Localizable.strings")));
-					QueueFileForBatchUpload(FileItem.GetItemByFullPath(Path.GetFullPath("../../Engine/Source/ThirdParty/Mac/RadioEffectUnit/RadioEffectUnit.component/Contents/Resources/RadioEffectUnit.rsrc")));
 					QueueFileForBatchUpload(FileItem.GetItemByFullPath(Path.GetFullPath("../../Engine/Source/ThirdParty/Mac/RadioEffectUnit/RadioEffectUnit.component/Contents/Info.plist")));
 					QueueFileForBatchUpload(FileItem.GetItemByFullPath(Path.Combine(LinkEnvironment.Config.IntermediateDirectory, "DylibCopy.sh")));
 				}
@@ -1176,7 +1174,6 @@ namespace UnrealBuildTool
 				Manifest.AddFileName(BundleContentsDirectory + "Resources/UE4.icns");
 				Manifest.AddFileName(BundleContentsDirectory + "Resources/RadioEffectUnit.component/Contents/MacOS/RadioEffectUnit");
 				Manifest.AddFileName(BundleContentsDirectory + "Resources/RadioEffectUnit.component/Contents/Resources/English.lproj/Localizable.strings");
-				Manifest.AddFileName(BundleContentsDirectory + "Resources/RadioEffectUnit.component/Contents/Resources/RadioEffectUnit.rsrc");
 				Manifest.AddFileName(BundleContentsDirectory + "Resources/RadioEffectUnit.component/Contents/Info.plist");
 
 				if (Binary.Config.TargetName.StartsWith("UE4Editor"))
diff --git a/Engine/Source/Runtime/Mac/CoreAudio/Private/CoreAudioEffects.cpp b/Engine/Source/Runtime/Mac/CoreAudio/Private/CoreAudioEffects.cpp
--- a/Engine/Source/Runtime/Mac/CoreAudio/Private/CoreAudioEffects.cpp
+++ b/Engine/Source/Runtime/Mac/CoreAudio/Private/CoreAudioEffects.cpp
@@ -15,28 +15,58 @@
 
 extern FCoreAudioSoundSource *GAudioChannels[MAX_AUDIOCHANNELS + 1];
 
-static bool LoadRadioEffectComponent()
+static CFBundleRef LoadRadioEffectComponent()
 {
 	bool bLoaded = false;
+	CFBundleRef Bundle = NULL;
 
 	CFBundleRef MainBundleRef = CFBundleGetMainBundle();
 	if( MainBundleRef )
 	{
 		CFURLRef ComponentURLRef = CFBundleCopyResourceURL( MainBundleRef, CFSTR("RadioEffectUnit"), CFSTR("component"), 0 );
-		FSRef ComponentFSRef;
 		if( ComponentURLRef )
 		{
-			if( CFURLGetFSRef( ComponentURLRef, &ComponentFSRef ) )
+			Bundle = CFBundleCreate(kCFAllocatorDefault, ComponentURLRef);
+			if(Bundle)
 			{
-				OSStatus Status = RegisterComponentFileRef( &ComponentFSRef, false );
-				bLoaded = Status == noErr;
+				CFArrayRef Components = (CFArrayRef)CFBundleGetValueForInfoDictionaryKey(Bundle, CFSTR("AudioComponents"));
+				if ( Components && CFArrayGetCount(Components) )
+				{
+					CFDictionaryRef ComponentInfo = (CFDictionaryRef)CFArrayGetValueAtIndex(Components, 0);
+					CFNumberRef ComponentVersion = (CFNumberRef)CFDictionaryGetValue(ComponentInfo, CFSTR("version"));
+					CFStringRef ComponentFactoryFunction = (CFStringRef)CFDictionaryGetValue(ComponentInfo, CFSTR("factoryFunction"));
+					if( ComponentVersion && ComponentFactoryFunction )
+					{
+						int32 Version = 0;
+						CFNumberGetValue(ComponentVersion, kCFNumberSInt32Type, &Version);
+						
+						AudioComponentFactoryFunction Factory = (AudioComponentFactoryFunction)CFBundleGetFunctionPointerForName(Bundle, ComponentFactoryFunction);
+						if(Factory)
+						{
+							//	fill out the AudioComponentDescription
+							AudioComponentDescription theDescription;
+							theDescription.componentType = kAudioUnitType_Effect;
+							theDescription.componentSubType = 'Rdio';
+							theDescription.componentManufacturer = 'Epic';
+							theDescription.componentFlagsMask = 0;
+							AudioComponent RadioComponent = AudioComponentRegister(&theDescription, CFSTR("Epic Games: RadioEffectUnit"), Version, Factory);
+							check(RadioComponent);
+							bLoaded = (RadioComponent != NULL);
+						}
+					}
+				}
+				if(!bLoaded)
+				{
+					CFRelease(Bundle);
+					Bundle = NULL;
+				}
 			}
 
 			CFRelease( ComponentURLRef );
 		}
 	}
 
-	return bLoaded;
+	return Bundle;
 }
 
 /*------------------------------------------------------------------------------------
@@ -55,12 +85,20 @@ FCoreAudioEffectsManager::FCoreAudioEffectsManager( FAudioDevice* InDevice )
 	: FAudioEffectsManager( InDevice )
 {
 #if RADIO_ENABLED
-	bRadioAvailable = LoadRadioEffectComponent();
+	RadioBundle = LoadRadioEffectComponent();
+	bRadioAvailable = (RadioBundle != NULL);
 #endif
 }
 
 FCoreAudioEffectsManager::~FCoreAudioEffectsManager()
 {
+#if RADIO_ENABLED
+	if(RadioBundle)
+	{
+		CFRelease(RadioBundle);
+		RadioBundle = NULL;
+	}
+#endif
 }
 
 /** 
diff --git a/Engine/Source/Runtime/Mac/CoreAudio/Public/CoreAudioEffects.h b/Engine/Source/Runtime/Mac/CoreAudio/Public/CoreAudioEffects.h
--- a/Engine/Source/Runtime/Mac/CoreAudio/Public/CoreAudioEffects.h
+++ b/Engine/Source/Runtime/Mac/CoreAudio/Public/CoreAudioEffects.h
@@ -40,6 +40,7 @@ public:
 private:
 
 	bool bRadioAvailable;
+	CFBundleRef RadioBundle;
 
 	/** Console variables to tweak the radio effect output */
 	static TConsoleVariableData<float>*	Radio_ChebyshevPowerMultiplier;
