CIS fix. The XML default options were overriding command line parameters, which caused problems on build machines.

diff --git a/Engine/Source/Programs/UnrealBuildTool/System/UnrealBuildTool.cs b/Engine/Source/Programs/UnrealBuildTool/System/UnrealBuildTool.cs
--- a/Engine/Source/Programs/UnrealBuildTool/System/UnrealBuildTool.cs
+++ b/Engine/Source/Programs/UnrealBuildTool/System/UnrealBuildTool.cs
@@ -1018,9 +1018,6 @@ namespace UnrealBuildTool
 
 			// Reset global configurations
 			ResetAllActions();
-			XmlConfigLoader.Reset<CPPEnvironment>();
-			XmlConfigLoader.Reset<BuildConfiguration>();
-			XmlConfigLoader.Reset<UEBuildConfiguration>();
 
 			ParseBuildConfigurationFlags(Arguments);
 
