UBT: Renamed XmlConfigField to XmlConfig attribute and left only field annotation mechanism. If class has any fields annotated then it's going to be included in the configuration mechanism.

diff --git a/Engine/Source/Programs/UnrealBuildTool/Android/UEBuildAndroid.cs b/Engine/Source/Programs/UnrealBuildTool/Android/UEBuildAndroid.cs
--- a/Engine/Source/Programs/UnrealBuildTool/Android/UEBuildAndroid.cs
+++ b/Engine/Source/Programs/UnrealBuildTool/Android/UEBuildAndroid.cs
@@ -9,17 +9,16 @@ using System.Xml;
 
 namespace UnrealBuildTool
 {
-	[XmlConfig]
 	class AndroidPlatform : UEBuildPlatform
 	{
 		/// <summary>
 		/// Android settings.
 		/// </summary>
-		[XmlConfigField]
+		[XmlConfig]
 		public static string AndroidArchitecture = "-armv7";
-		[XmlConfigField]
+		[XmlConfig]
 		public static string AndroidNdkApiTarget = "latest";
-		[XmlConfigField]
+		[XmlConfig]
 		public static string AndroidSdkApiTarget = "latest";
 
 		// The current architecture - affects everything about how UBT operates on Android
diff --git a/Engine/Source/Programs/UnrealBuildTool/Configuration/BuildConfiguration.cs b/Engine/Source/Programs/UnrealBuildTool/Configuration/BuildConfiguration.cs
--- a/Engine/Source/Programs/UnrealBuildTool/Configuration/BuildConfiguration.cs
+++ b/Engine/Source/Programs/UnrealBuildTool/Configuration/BuildConfiguration.cs
@@ -8,16 +8,17 @@ using System.Reflection;
 
 namespace UnrealBuildTool
 {
-	[XmlConfigAttribute(true)]
 	public class BuildConfiguration
 	{
 		/// <summary>
 		/// Whether to unify C++ code into larger files for faster compilation.
 		/// </summary>
+		[XmlConfig]
 		public static bool bUseUnityBuild;
         /// <summary>
         /// Whether to _force_ unify C++ code into larger files for faster compilation.
         /// </summary>
+		[XmlConfig]
         public static bool bForceUnityBuild;
 
 		/// <summary>
@@ -25,66 +26,79 @@ namespace UnrealBuildTool
 		/// allows small game modules to have faster iterative compile times for single files, at the expense of slower full
 		/// rebuild times.  This setting can be overridden by the bFasterWithoutUnity option in a module's Build.cs file.
 		/// </summary>
+		[XmlConfig]
 		public static int MinGameModuleSourceFilesForUnityBuild;
 
 		/// <summary>
 		/// New Monolithic Graphics drivers have optional "fast calls" replacing various D3d functions
 		/// </summary>
-        public static bool bUseFastMonoCalls;
+		[XmlConfig]
+		public static bool bUseFastMonoCalls;
 
         /// <summary>
 		/// An approximate number of bytes of C++ code to target for inclusion in a single unified C++ file.
 		/// </summary>
+		[XmlConfig]
 		public static int NumIncludedBytesPerUnityCPP;
 
 		/// <summary>
 		/// Whether to stress test the C++ unity build robustness by including all C++ files files in a project from a single unified file.
 		/// </summary>
+		[XmlConfig]
 		public static bool bStressTestUnity;
 
 		/// <summary>
 		/// Whether headers in system paths should be checked for modification when determining outdated actions.
 		/// </summary>
+		[XmlConfig]
 		public static bool bCheckSystemHeadersForModification;
 
 		/// <summary>
 		/// Whether headers in the Development\External folder should be checked for modification when determining outdated actions.
 		/// </summary>
+		[XmlConfig]
 		public static bool bCheckExternalHeadersForModification;
 
 		/// <summary>
 		/// Whether to ignore import library files that are out of date when building targets.  Set this to true to improve iteration time.
 		/// </summary>
+		[XmlConfig]
 		public static bool bIgnoreOutdatedImportLibraries;
 
 		/// <summary>
 		/// Whether to globally disable debug info generation; see DebugInfoHeuristics.cs for per-config and per-platform options.
 		/// </summary>
+		[XmlConfig]
 		public static bool bDisableDebugInfo;
 
 		/// <summary>
 		/// Whether to disable debug info on PC in development builds (for faster developer iteration, as link times are extremely fast with debug info disabled.)
 		/// </summary>
+		[XmlConfig]
 		public static bool bOmitPCDebugInfoInDevelopment;
 
 		/// <summary>
 		/// Whether PDB files should be used for Visual C++ builds.
 		/// </summary>
+		[XmlConfig]
 		public static bool bUsePDBFiles;
 
  		/// <summary>
  		/// Whether PCH files should be used.
 		/// </summary>
+		[XmlConfig]
 		public static bool bUsePCHFiles;
 
 		/// <summary>
 		/// Whether to generate command line dependencies for compile actions when requested
 		/// </summary>
+		[XmlConfig]
 		public static bool bUseActionHistory;
 
 		/// <summary>
 		/// The minimum number of files that must use a pre-compiled header before it will be created and used.
 		/// </summary>
+		[XmlConfig]
 		public static int MinFilesUsingPrecompiledHeader;
 
 		/// <summary>
@@ -93,91 +107,109 @@ namespace UnrealBuildTool
 		/// full rebuild times for small game projects.  This can be overridden by setting MinFilesUsingPrecompiledHeaderOverride in
 		/// a module's Build.cs file
 		/// </summary>
+		[XmlConfig]
 		public static bool bForcePrecompiledHeaderForGameModules;
 
 		/// <summary>
 		/// Whether debug info should be written to the console.
 		/// </summary>
+		[XmlConfig]
 		public static bool bPrintDebugInfo;
 
 		/// <summary>
 		/// Prints performance diagnostics about include dependencies and other bits
 		/// </summary>
+		[XmlConfig]
 		public static bool bPrintPerformanceInfo;
 
 		/// <summary>
 		/// Whether to log detailed action stats. This forces local execution.
 		/// </summary>
+		[XmlConfig]
 		public static bool bLogDetailedActionStats;
 
 		/// <summary>
 		/// Whether to deploy the executable after compilation on platforms that require deployment.
 		/// </summary>
+		[XmlConfig]
 		public static bool bDeployAfterCompile;
 
 		/// <summary>
 		/// Whether to clean Builds directory on a remote Mac before building
 		/// </summary>
+		[XmlConfig]
 		public static bool bFlushBuildDirOnRemoteMac;
 
 		/// <summary>
 		/// Whether XGE may be used.
 		/// </summary>
+		[XmlConfig]
 		public static bool bAllowXGE;
 
 		/// <summary>
 		/// Whether we should just export the XGE XML and pretend it succeeded
 		/// </summary>
+		[XmlConfig]
 		public static bool bXGEExport;
 
 		/// <summary>
 		/// Whether to display the XGE build monitor.
 		/// </summary>
+		[XmlConfig]
 		public static bool bShowXGEMonitor;
 
 		/// <summary>
 		/// When enabled, XGE will stop compiling targets after a compile error occurs.  Recommended, as it saves computing resources for others.
 		/// </summary>
+		[XmlConfig]
 		public static bool bStopXGECompilationAfterErrors;
 
 		/// <summary>
 		/// When enabled, allows XGE to compile pre-compiled header files on remote machines.  Otherwise, PCHs are always generated locally.
 		/// </summary>
+		[XmlConfig]
 		public static bool bAllowRemotelyCompiledPCHs;
 
 		/// <summary>
 		/// Whether or not to delete outdated produced items.
 		/// </summary>
+		[XmlConfig]
 		public static bool bShouldDeleteAllOutdatedProducedItems;
 
 		/// <summary>
 		/// Whether to use incremental linking or not.
 		/// </summary>
+		[XmlConfig]
 		public static bool bUseIncrementalLinking;
 
 		/// <summary>
 		/// Whether to allow the use of LTCG (link time code generation) 
 		/// </summary>
+		[XmlConfig]
 		public static bool bAllowLTCG;
 
 		/// <summary>
 		/// Whether to support edit and continue.  Only works on Microsoft compilers in 32-bit compiles.
 		/// </summary>
+		[XmlConfig]
 		public static bool bSupportEditAndContinue;
 
 		/// <summary>
 		/// Whether to omit frame pointers or not. Disabling is useful for e.g. memory profiling on the PC
 		/// </summary>
+		[XmlConfig]
 		public static bool bOmitFramePointers;
         
 		/// <summary>
 		/// Processor count multiplier for local execution. Can be below 1 to reserve CPU for other tasks.
 		/// </summary>
+		[XmlConfig]
 		public static double ProcessorCountMultiplier;
 
 		/// <summary>
 		/// The intermediate folder - i.e. Intermediate/Build.
 		/// </summary>
+		[XmlConfig]
 		public static string BaseIntermediateFolder;
 
 		/// <summary>
@@ -233,26 +265,31 @@ namespace UnrealBuildTool
 		/// <summary>
 		/// Whether to generate a dSYM file or not.
 		/// </summary>
+		[XmlConfig]
 		public static bool bGeneratedSYMFile;
 
 		/// <summary>
 		/// Whether to strip iOS symbols or not (implied by bGeneratedSYMFile).
 		/// </summary>
+		[XmlConfig]
 		public static bool bStripSymbolsOnIOS;
 
 		/// <summary>
 		/// If true, then a stub IPA will be generated when compiling is done (minimal files needed for a valid IPA)
 		/// </summary>
+		[XmlConfig]
 		public static bool bCreateStubIPA;
 
 		/// <summary>
 		/// If true, then a stub IPA will be generated when compiling is done (minimal files needed for a valid IPA)
 		/// </summary>
+		[XmlConfig]
 		public static bool bCopyAppBundleBackToDevice;
 
 		/// <summary>
 		/// If true, then enable memory profiling in the build (defines USE_MALLOC_PROFILER=1 and forces bOmitFramePointers=false)
 		/// </summary>
+		[XmlConfig]
 		public static bool bUseMallocProfiler;
 
 		/// <summary>
@@ -261,17 +298,20 @@ namespace UnrealBuildTool
 		/// /analyze option.  The compiler included with the Windows SDK supports this, as well as Team Edition
 		/// versions of Visual Studio.
 		/// </summary>
+		[XmlConfig]
 		public static bool bEnableCodeAnalysis;
 
 		/// <summary>
 		/// Enables "Shared PCHs", an experimental feature which may significantly speed up compile times by attempting to
 		/// share certain PCH files between modules that UBT detects is including those PCH's header files
 		/// </summary>
+		[XmlConfig]
 		public static bool bUseSharedPCHs;
 
 		/// <summary>
 		/// Writes an XML summary of the compile environment for the target to the given path, but does not build it.
 		/// </summary>
+		[XmlConfig]
 		public static string WriteTargetInfoPath;
 
 		/// <summary>
@@ -279,11 +319,13 @@ namespace UnrealBuildTool
 		/// This is technically not fully correct because the dependency cache is global and each module could have a different set of include paths that could cause headers
 		/// to resolve files differently. In practice this is not the case, and significantly speeds up UBT when nothing is to be done.
 		/// </summary>
+		[XmlConfig]
 		public static bool bUseIncludeDependencyResolveCache;
 
 		/// <summary>
 		/// Used to test the dependency resolve cache. This will verify the resolve cache has no conflicts by resolving every time and checking against any previous resolve attempts.
 		/// </summary>
+		[XmlConfig]
 		public static bool bTestIncludeDependencyResolveCache;
 
 		/// <summary>
@@ -291,6 +333,7 @@ namespace UnrealBuildTool
 		/// compile against the Release CRT but leave optimizations turned off, so that it is still easy to debug.  If you enable this
 		/// setting, actual Debug versions of third party static libraries will be needed when compiling in Debug.
 		/// </summary>
+		[XmlConfig]
 		public static bool bDebugBuildsActuallyUseDebugCRT;
 
 		/// <summary>
diff --git a/Engine/Source/Programs/UnrealBuildTool/Configuration/UEBuildConfiguration.cs b/Engine/Source/Programs/UnrealBuildTool/Configuration/UEBuildConfiguration.cs
--- a/Engine/Source/Programs/UnrealBuildTool/Configuration/UEBuildConfiguration.cs
+++ b/Engine/Source/Programs/UnrealBuildTool/Configuration/UEBuildConfiguration.cs
@@ -5,121 +5,158 @@ using System.IO;
 
 namespace UnrealBuildTool
 {
-	[XmlConfig(true)]
 	public class UEBuildConfiguration
 	{
 		/** Whether to include PhysX support */
+		[XmlConfig]
 		public static bool bCompilePhysX;
 
 		/** Whether to include PhysX APEX support */
+		[XmlConfig]
 		public static bool bCompileAPEX;
 
 		/** Whether to include Box2D support */
+		[XmlConfig]
 		public static bool bCompileBox2D;
 
-        /** Whether to include ICU unicode/i18n support in core */
-        public static bool bCompileICU;
+		/** Whether to include ICU unicode/i18n support in core */
+		[XmlConfig]
+		public static bool bCompileICU;
 
 		/** Whether to build a stripped down version of the game specifically for dedicated server. */
+		[XmlConfig]
 		public static bool bBuildDedicatedServer;
 
 		/** Whether to compile the editor or not. Only desktop platforms (Windows or Mac) will use this, other platforms force this to false */
+		[XmlConfig]
 		public static bool bBuildEditor;
 
 		/** Whether to compile code related to building assets. Consoles generally cannot build assets. Desktop platforms generally can. */
+		[XmlConfig]
 		public static bool bBuildRequiresCookedData;
 
 		/** Whether to compile WITH_EDITORONLY_DATA disabled. Only Windows will use this, other platforms force this to false */
+		[XmlConfig]
 		public static bool bBuildWithEditorOnlyData;
 
 		/** Whether to compile the developer tools. */
+		[XmlConfig]
 		public static bool bBuildDeveloperTools;
 
 		/** Whether to force compiling the target platform modules, even if they wouldn't normally be built */
+		[XmlConfig]
 		public static bool bForceBuildTargetPlatforms;
 
 		/** Whether to force compiling shader format modules, even if they wouldn't normally be built. */
+		[XmlConfig]
 		public static bool bForceBuildShaderFormats;
 
 		/** Whether we should compile in support for Simplygon or not. */
+		[XmlConfig]
 		public static bool bCompileSimplygon;
 
-        /** Whether we should compile in support for Steam OnlineSubsystem or not. [RCL] FIXME 2014-Apr-17: bCompileSteamOSS means "bHasSteamworksInstalled" for some code, these meanings need to be untangled */
+		/** Whether we should compile in support for Steam OnlineSubsystem or not. [RCL] FIXME 2014-Apr-17: bCompileSteamOSS means "bHasSteamworksInstalled" for some code, these meanings need to be untangled */
+		[XmlConfig]
 		public static bool bCompileSteamOSS;
 
 		/** Whether we should compile in support for Mcp OnlineSubsystem or not. */
+		[XmlConfig]
 		public static bool bCompileMcpOSS;
 
 		/** Whether to compile lean and mean version of UE. */
+		[XmlConfig]
 		public static bool bCompileLeanAndMeanUE;
 
 		/** Whether to generate a manifest file that contains the files to add to Perforce */
+		[XmlConfig]
 		public static bool bGenerateManifest;
 
 		/** Whether to add to the existing manifest (if it exists), or start afresh */
+		[XmlConfig]
 		public static bool bMergeManifests;
 
 		/** Whether to 'clean' the given project */
+		[XmlConfig]
 		public static bool bCleanProject;
 
 		/** Whether we are just running the PrepTargetForDeployment step */
+		[XmlConfig]
 		public static bool bPrepForDeployment;
 
 		/** Enabled for all builds that include the engine project.  Disabled only when building standalone apps that only link with Core. */
+		[XmlConfig]
 		public static bool bCompileAgainstEngine;
 
 		/** Enabled for all builds that include the CoreUObject project.  Disabled only when building standalone apps that only link with Core. */
+		[XmlConfig]
 		public static bool bCompileAgainstCoreUObject;
 
 		/** If true, include ADO database support in core */
+		[XmlConfig]
 		public static bool bIncludeADO;
 
 		/** Directory for the third part files/libs */
+		[XmlConfig]
 		public static string UEThirdPartyDirectory;
 
 		/** If true, force header regeneration. Intended for the build machine */
+		[XmlConfig]
 		public static bool bForceHeaderGeneration;
 
 		/** If true, do not build UHT, assume it is already built */
+		[XmlConfig]
 		public static bool bDoNotBuildUHT;
 
 		/** If true, fail if any of the generated header files is out of date. */
+		[XmlConfig]
 		public static bool bFailIfGeneratedCodeChanges;
 
 		/** Whether to compile Recast navmesh generation */
+		[XmlConfig]
 		public static bool bCompileRecast;
 
 		/** Whether to compile SpeedTree support. */
+		[XmlConfig]
 		public static bool bCompileSpeedTree;
 
 		/** Enable exceptions for all modules */
+		[XmlConfig]
 		public static bool bForceEnableExceptions;
 
 		/** Compile server-only code. */
+		[XmlConfig]
 		public static bool bWithServerCode;
 
 		/** Whether to include stats support even without the engine */
+		[XmlConfig]
 		public static bool bCompileWithStatsWithoutEngine;
 
 		/** Whether to include plugin support */
+		[XmlConfig]
 		public static bool bCompileWithPluginSupport;
 
 		/** Whether to turn on logging for test/shipping builds */
+		[XmlConfig]
 		public static bool bUseLoggingInShipping;
 
 		/** True if plugins should be excluded when building target. */
+		[XmlConfig]
 		public static bool bExcludePlugins;
 
-        /** True if we need to package up Android with the OBB in the APK file */
-        public static bool bOBBinAPK;
+		/** True if we need to package up Android with the OBB in the APK file */
+		[XmlConfig]
+		public static bool bOBBinAPK;
 
 		/** True if we need PhysX vehicle support */
+		[XmlConfig]
 		public static bool bCompilePhysXVehicle;
 
 		/** True if we need FreeType support */
+		[XmlConfig]
 		public static bool bCompileFreeType;
 
 		/** True if we want to favor optimizing size over speed */
+		[XmlConfig]
 		public static bool bCompileForSize;
 
 		/// <summary>
diff --git a/Engine/Source/Programs/UnrealBuildTool/HTML5/UEBuildHTML5.cs b/Engine/Source/Programs/UnrealBuildTool/HTML5/UEBuildHTML5.cs
--- a/Engine/Source/Programs/UnrealBuildTool/HTML5/UEBuildHTML5.cs
+++ b/Engine/Source/Programs/UnrealBuildTool/HTML5/UEBuildHTML5.cs
@@ -8,10 +8,9 @@ using System.IO;
 
 namespace UnrealBuildTool
 {
-	[XmlConfig]
 	class HTML5Platform : UEBuildPlatform
 	{
-		[XmlConfigField]
+		[XmlConfig]
 		public static string HTML5Architecture = "";
 
 		// The current architecture - affects everything about how UBT operates on HTML5
diff --git a/Engine/Source/Programs/UnrealBuildTool/IOS/IOSToolChain.cs b/Engine/Source/Programs/UnrealBuildTool/IOS/IOSToolChain.cs
--- a/Engine/Source/Programs/UnrealBuildTool/IOS/IOSToolChain.cs
+++ b/Engine/Source/Programs/UnrealBuildTool/IOS/IOSToolChain.cs
@@ -15,7 +15,6 @@ using Ionic.Zlib;
 
 namespace UnrealBuildTool
 {
-	[XmlConfig(true)]
 	class IOSToolChain : RemoteToolChain
 	{
 		public override void RegisterToolChain()
@@ -34,22 +33,28 @@ namespace UnrealBuildTool
 		// If you are looking for where to change the remote compile server name, look in RemoteToolChain.cs
 
 		/** If this is set, then we don't do any post-compile steps except moving the executable into the proper spot on the Mac */
+		[XmlConfig]
 		public static bool bUseDangerouslyFastMode = false;
 
 		/** Which version of the iOS SDK to target at build time */
+		[XmlConfig]
 		public static string IOSSDKVersion = "latest";
 
 		/** The architecture(s) to compile */
+		[XmlConfig]
 		public static string NonShippingArchitectures = "armv7";
+		[XmlConfig]
 		public static string ShippingArchitectures = "armv7";
 
 		// In case the SDK checking fails for some reason, use this version
 		private static string BackupVersion = "6.0";
 
 		/** Which version of the iOS to allow at run time */
+		[XmlConfig]
 		public static string IOSVersion = "6.0";
 
 		/** Which developer directory to root from */
+		[XmlConfig]
 		public static string XcodeDeveloperDir = "/Applications/Xcode.app/Contents/Developer/";
 
 		/** Location of the SDKs */
diff --git a/Engine/Source/Programs/UnrealBuildTool/IOS/UEBuildIOS.cs b/Engine/Source/Programs/UnrealBuildTool/IOS/UEBuildIOS.cs
--- a/Engine/Source/Programs/UnrealBuildTool/IOS/UEBuildIOS.cs
+++ b/Engine/Source/Programs/UnrealBuildTool/IOS/UEBuildIOS.cs
@@ -9,11 +9,10 @@ using System.Xml;
 
 namespace UnrealBuildTool
 {
-	[XmlConfig]
 	class IOSPlatform : UEBuildPlatform
 	{
 		// by default, use an empty architecture (which is really just a modifer to the platform for some paths/names)
-		[XmlConfigField]
+		[XmlConfig]
 		public static string IOSArchitecture = "";
 
 		// The current architecture - affects everything about how UBT operates on IOS
diff --git a/Engine/Source/Programs/UnrealBuildTool/System/ProjectFileGenerator.cs b/Engine/Source/Programs/UnrealBuildTool/System/ProjectFileGenerator.cs
--- a/Engine/Source/Programs/UnrealBuildTool/System/ProjectFileGenerator.cs
+++ b/Engine/Source/Programs/UnrealBuildTool/System/ProjectFileGenerator.cs
@@ -577,7 +577,7 @@ namespace UnrealBuildTool
 		{
 			EngineProject.AddAliasedFileToProject(new AliasedFile(
 					XmlConfigLoader.GetXSDPath(),
-					Path.Combine("Config", "UnrealBuildTool", "Schema")
+					Path.Combine("Programs", "UnrealBuildTool")
 				));
 
 			foreach(var BuildConfigurationPath in XmlConfigLoader.ConfigLocationHierarchy)
diff --git a/Engine/Source/Programs/UnrealBuildTool/ToolChain/RemoteToolChain.cs b/Engine/Source/Programs/UnrealBuildTool/ToolChain/RemoteToolChain.cs
--- a/Engine/Source/Programs/UnrealBuildTool/ToolChain/RemoteToolChain.cs
+++ b/Engine/Source/Programs/UnrealBuildTool/ToolChain/RemoteToolChain.cs
@@ -12,7 +12,6 @@ using System.Text.RegularExpressions;
 
 namespace UnrealBuildTool
 {
-	[XmlConfig]
 	abstract class RemoteToolChain : UEToolChain
     {
 		protected void RegisterRemoteToolChain(UnrealTargetPlatform InPlatform, CPPTargetPlatform CPPPlatform)
@@ -25,9 +24,9 @@ namespace UnrealBuildTool
 		}
 
 		/** These two variables will be loaded from XML config file in XmlConfigLoader.Init() */
-		[XmlConfigField]
+		[XmlConfig]
 		public static string RemoteServerName;
-		[XmlConfigField]
+		[XmlConfig]
 		public static string[] PotentialServerNames = new string[] { };
 
 		/** Keep a list of remote files that are potentially copied from local to remote */
diff --git a/Engine/Source/Programs/UnrealBuildTool/Utilities/XmlConfigLoader.cs b/Engine/Source/Programs/UnrealBuildTool/Utilities/XmlConfigLoader.cs
--- a/Engine/Source/Programs/UnrealBuildTool/Utilities/XmlConfigLoader.cs
+++ b/Engine/Source/Programs/UnrealBuildTool/Utilities/XmlConfigLoader.cs
@@ -14,24 +14,10 @@ using System.Xml.Linq;
 namespace UnrealBuildTool
 {
 	/// <summary>
-	/// Attribute to annotate classes that can be set using xml configuration system.
-	/// </summary>
-	[AttributeUsage(AttributeTargets.Class | AttributeTargets.Struct)]
-	public class XmlConfigAttribute : Attribute
-	{
-		public readonly bool bAllPublicStaticFields;
-
-		public XmlConfigAttribute(bool bAllPublicStaticFields = false)
-		{
-			this.bAllPublicStaticFields = bAllPublicStaticFields;
-		}
-	}
-
-	/// <summary>
-	/// Attribute to annotate fields in type that can be set using xml configuration system.
+	/// Attribute to annotate fields in type that can be set using XML configuration system.
 	/// </summary>
 	[AttributeUsage(AttributeTargets.Field)]
-	public class XmlConfigFieldAttribute : Attribute
+	public class XmlConfigAttribute : Attribute
 	{
 
 	}
@@ -595,16 +581,13 @@ namespace UnrealBuildTool
 					continue;
 				}
 
-				XmlConfigLoaderClassData ClassData;
-				bool bAllPublicStaticFields = false;
-
-				var Attributes = ClassType.GetCustomAttributes(typeof(XmlConfigAttribute), false);
-				if (Attributes.Length == 0)
+				if (IsConfigurableClass(ClassType))
 				{
 					Log.TraceVerbose("XmlConfig Loading: class '{0}' is not allowed to be configured using XML system.", XmlClass.Name);
 					continue;
 				}
-				bAllPublicStaticFields = ((XmlConfigAttribute)Attributes.First()).bAllPublicStaticFields;
+
+				XmlConfigLoaderClassData ClassData;
 
 				if (!Data.TryGetValue(ClassType, out ClassData))
 				{
@@ -618,7 +601,7 @@ namespace UnrealBuildTool
 				{
 					FieldInfo Field = ClassType.GetField(XmlField.Name);
 
-					if (Field == null || !Field.IsPublic || !Field.IsStatic || (!bAllPublicStaticFields && Field.GetCustomAttributes(typeof(XmlConfigFieldAttribute), false).Length == 0))
+					if (Field == null || !IsConfigurableField(Field))
 					{
 						throw new BuildException("BuildConfiguration Loading: field '{0}' doesn't exist or is either non-public, non-static or not-xml-configurable.", XmlField.Name);
 					}
@@ -741,18 +724,32 @@ namespace UnrealBuildTool
 		}
 
 		/// <summary>
-		/// Gets all types that can be configured using this xml system.
+		/// Tells if given class is XML configurable.
+		/// </summary>
+		/// <param name="Class">Class to check.</param>
+		/// <returns>True if the class is configurable using XML system. Otherwise false.</returns>
+		private static bool IsConfigurableClass(Type Class)
+		{
+			return Class.GetFields().Any((Field) => IsConfigurableField(Field));
+		}
+
+		/// <summary>
+		/// Tells if given field is XML configurable.
+		/// </summary>
+		/// <param name="Field">field to check.</param>
+		/// <returns>True if the field is configurable using XML system. Otherwise false.</returns>
+		private static bool IsConfigurableField(FieldInfo Field)
+		{
+			return Field.IsStatic && Field.IsPublic && Field.GetCustomAttributes(typeof(XmlConfigAttribute), false).Length > 0;
+		}
+
+		/// <summary>
+		/// Gets all types that can be configured using this XML system.
 		/// </summary>
 		/// <returns>Enumerable of types that can be configured using this xml system.</returns>
 		private static IEnumerable<Type> GetAllConfigurationTypes()
 		{
-			foreach(var Type in Assembly.GetExecutingAssembly().GetTypes())
-			{
-				if(Type.GetCustomAttributes(typeof(XmlConfigAttribute), false).Length > 0)
-				{
-					yield return Type;
-				}
-			}
+			return Assembly.GetExecutingAssembly().GetTypes().Where((Class) => IsConfigurableClass(Class));
 		}
 
 		/// <summary>
@@ -777,15 +774,7 @@ namespace UnrealBuildTool
 		/// <returns>Enumerable fields that are configurable using this system.</returns>
 		private static IEnumerable<FieldInfo> GetConfigurableFields(Type ClassType)
 		{
-			var bAllPublicStaticFields = ((XmlConfigAttribute) ClassType.GetCustomAttributes(typeof(XmlConfigAttribute), false).First()).bAllPublicStaticFields;
-
-			foreach(var Field in ClassType.GetFields())
-			{
-				if(Field.IsPublic && Field.IsStatic && (bAllPublicStaticFields || Field.GetCustomAttributes(typeof(XmlConfigFieldAttribute), false).Length > 0))
-				{
-					yield return Field;
-				}
-			}
+			return ClassType.GetFields().Where((Field) => IsConfigurableField(Field));
 		}
 
 		/// <summary>
@@ -938,7 +927,8 @@ namespace UnrealBuildTool
 					}
 				}
 
-				return Encoding.UTF8.GetString(MS.ToArray());
+				// Skipping first 3 bytes, that informs about encoding.
+				return Encoding.UTF8.GetString(MS.ToArray().Skip(3).ToArray());
 			}
 		}
 
diff --git a/Engine/Source/Programs/UnrealBuildTool/WinRT/UEBuildWinRT.cs b/Engine/Source/Programs/UnrealBuildTool/WinRT/UEBuildWinRT.cs
--- a/Engine/Source/Programs/UnrealBuildTool/WinRT/UEBuildWinRT.cs
+++ b/Engine/Source/Programs/UnrealBuildTool/WinRT/UEBuildWinRT.cs
@@ -8,13 +8,12 @@ using System.IO;
 
 namespace UnrealBuildTool
 {
-	[XmlConfig]
 	class WinRTPlatform : UEBuildPlatform
 	{
 		/// <summary>
 		/// Should the app be compiled as WinRT
 		/// </summary>
-		[XmlConfigField]
+		[XmlConfig]
 		public static bool bCompileWinRT = false;
 
 		public static bool IsVisualStudioInstalled()
