Minor fix for UBT.

diff --git a/Engine/Source/Programs/UnrealBuildTool/Utilities/XmlConfigLoader.cs b/Engine/Source/Programs/UnrealBuildTool/Utilities/XmlConfigLoader.cs
--- a/Engine/Source/Programs/UnrealBuildTool/Utilities/XmlConfigLoader.cs
+++ b/Engine/Source/Programs/UnrealBuildTool/Utilities/XmlConfigLoader.cs
@@ -581,7 +581,7 @@ namespace UnrealBuildTool
 					continue;
 				}
 
-				if (IsConfigurableClass(ClassType))
+				if (!IsConfigurableClass(ClassType))
 				{
 					Log.TraceVerbose("XmlConfig Loading: class '{0}' is not allowed to be configured using XML system.", XmlClass.Name);
 					continue;
