UBT: Not calling PostReset function after initial loading.

diff --git a/Engine/Source/Programs/UnrealBuildTool/Utilities/XmlConfigLoader.cs b/Engine/Source/Programs/UnrealBuildTool/Utilities/XmlConfigLoader.cs
--- a/Engine/Source/Programs/UnrealBuildTool/Utilities/XmlConfigLoader.cs
+++ b/Engine/Source/Programs/UnrealBuildTool/Utilities/XmlConfigLoader.cs
@@ -203,6 +203,8 @@ namespace UnrealBuildTool
 			foreach(var ClassData in Data)
 			{
 				ClassData.Value.ResetData();
+
+				InvokeIfExists(ClassData.Key, "PostReset");
 			}
 		}
 
