- More information for "latest" failure on a builder

diff --git a/Engine/Source/Programs/UnrealBuildTool/IOS/IOSToolChain.cs b/Engine/Source/Programs/UnrealBuildTool/IOS/IOSToolChain.cs
--- a/Engine/Source/Programs/UnrealBuildTool/IOS/IOSToolChain.cs
+++ b/Engine/Source/Programs/UnrealBuildTool/IOS/IOSToolChain.cs
@@ -147,9 +147,11 @@ namespace UnrealBuildTool
 						IOSSDKVersion = MaxSDKVersionString;
 					}
 				}
-				catch (Exception)
+				catch (Exception Ex)
 				{
 					// on any exception, just use the backup version
+					Log.TraceInformation("Triggered an exception while looking for SDK directory in Xcode.app. Using Backup version...");
+					Log.TraceInformation("{0}", Ex.ToString());
 					IOSSDKVersion = BackupVersion;
 				}
 
