[GitHub] 132 : Introduce NativeBuildEnvironmentConfiguration
#github 132
#ttp 334942

diff --git a/Engine/Source/Programs/UnrealBuildTool/Android/AndroidToolChain.cs b/Engine/Source/Programs/UnrealBuildTool/Android/AndroidToolChain.cs
--- a/Engine/Source/Programs/UnrealBuildTool/Android/AndroidToolChain.cs
+++ b/Engine/Source/Programs/UnrealBuildTool/Android/AndroidToolChain.cs
@@ -156,7 +156,7 @@ namespace UnrealBuildTool
 			Result += " -Wno-invalid-offsetof";			// needed to suppress warnings about using offsetof on non-POD types.
 
 			// shipping builds will cause this warning with "ensure", so disable only in those case
-			if (CompileEnvironment.Config.TargetConfiguration == CPPTargetConfiguration.Shipping)
+			if (CompileEnvironment.Config.Target.Configuration == CPPTargetConfiguration.Shipping)
 			{
 				Result += " -Wno-unused-value";
 			}
@@ -168,7 +168,7 @@ namespace UnrealBuildTool
 			}
 
 			// optimization level
-			if (CompileEnvironment.Config.TargetConfiguration == CPPTargetConfiguration.Debug)
+			if (CompileEnvironment.Config.Target.Configuration == CPPTargetConfiguration.Debug)
 			{
 				Result += " -O0";
 			}
@@ -178,7 +178,7 @@ namespace UnrealBuildTool
 			}
 
 			//@todo android: these are copied verbatim from UE3 and probably need adjustment
-			if (CompileEnvironment.Config.TargetArchitecture == "-armv7")
+			if (CompileEnvironment.Config.Target.Architecture == "-armv7")
 			{
 		//		Result += " -mthumb-interwork";			// Generates code which supports calling between ARM and Thumb instructions, w/o it you can't reliability use both together 
 				Result += " -funwind-tables";			// Just generates any needed static data, affects no code 
@@ -197,14 +197,14 @@ namespace UnrealBuildTool
 				Result += " -mfpu=vfpv3-d16";			//@todo android: UE3 was just vfp. arm7a should all support v3 with 16 registers
 
 				// Some switches interfere with on-device debugging
-				if (CompileEnvironment.Config.TargetConfiguration != CPPTargetConfiguration.Debug)
+				if (CompileEnvironment.Config.Target.Configuration != CPPTargetConfiguration.Debug)
 				{
 					Result += " -ffunction-sections";   // Places each function in its own section of the output file, linker may be able to perform opts to improve locality of reference
 				}
 
 				Result += " -fsigned-char";				// Treat chars as signed //@todo android: any concerns about ABI compatibility with libs here?
 			}
-			else if (CompileEnvironment.Config.TargetArchitecture == "-x86")
+			else if (CompileEnvironment.Config.Target.Architecture == "-x86")
 			{
 				Result += " -fstrict-aliasing";
 				Result += " -funswitch-loops";
@@ -283,18 +283,18 @@ namespace UnrealBuildTool
 		{
 			string Result = "";
 			
-			Result += (LinkEnvironment.Config.TargetArchitecture == "-armv7") ? ToolchainParamsArm : ToolchainParamsx86;
+			Result += (LinkEnvironment.Config.Target.Architecture == "-armv7") ? ToolchainParamsArm : ToolchainParamsx86;
 
 			Result += " -nostdlib";
 			Result += " -Wl,-shared,-Bsymbolic";
 			Result += " -Wl,--no-undefined";
 
-			if (LinkEnvironment.Config.TargetArchitecture == "-armv7")
+			if (LinkEnvironment.Config.Target.Architecture == "-armv7")
 			{
 				Result += " -march=armv7-a";
 				Result += " -Wl,--fix-cortex-a8";		// required to route around a CPU bug in some Cortex-A8 implementations
 			}
-			else if (LinkEnvironment.Config.TargetArchitecture == "-x86")
+			else if (LinkEnvironment.Config.Target.Architecture == "-x86")
 			{
 				Result += " -march=atom";
 			}
@@ -437,12 +437,12 @@ namespace UnrealBuildTool
 				// should we disable optimizations on this file?
 				// @todo android - We wouldn't need this if we could disable optimizations per function (via pragma)
 				bool bDisableOptimizations = false;// SourceFile.AbsolutePath.ToUpperInvariant().IndexOf("\\SLATE\\") != -1;
-				if (bDisableOptimizations && CompileEnvironment.Config.TargetConfiguration != CPPTargetConfiguration.Debug)
+				if (bDisableOptimizations && CompileEnvironment.Config.Target.Configuration != CPPTargetConfiguration.Debug)
 				{
 					Log.TraceWarning("Disabling optimizations on {0}", SourceFile.AbsolutePath);
 				}
 
-				bDisableOptimizations = bDisableOptimizations || CompileEnvironment.Config.TargetConfiguration == CPPTargetConfiguration.Debug;
+				bDisableOptimizations = bDisableOptimizations || CompileEnvironment.Config.Target.Configuration == CPPTargetConfiguration.Debug;
 
 				// Add C or C++ specific compiler arguments.
 				if (CompileEnvironment.Config.PrecompiledHeaderAction == PrecompiledHeaderAction.Create)
@@ -549,7 +549,7 @@ namespace UnrealBuildTool
 
 			if (LinkEnvironment.Config.bIsBuildingLibrary)
 			{
-				LinkAction.CommandPath = (LinkEnvironment.Config.TargetArchitecture == "-armv7") ? ArPathArm : ArPathx86;
+				LinkAction.CommandPath = (LinkEnvironment.Config.Target.Architecture == "-armv7") ? ArPathArm : ArPathx86;
 			}
 			else
 			{
diff --git a/Engine/Source/Programs/UnrealBuildTool/Android/UEBuildAndroid.cs b/Engine/Source/Programs/UnrealBuildTool/Android/UEBuildAndroid.cs
--- a/Engine/Source/Programs/UnrealBuildTool/Android/UEBuildAndroid.cs
+++ b/Engine/Source/Programs/UnrealBuildTool/Android/UEBuildAndroid.cs
@@ -251,7 +251,7 @@ namespace UnrealBuildTool
 			InBuildTarget.GlobalCompileEnvironment.Config.Definitions.Add("USE_NULL_RHI=0");
 			InBuildTarget.GlobalCompileEnvironment.Config.Definitions.Add("REQUIRES_ALIGNED_INT_ACCESS");
 
-			if (InBuildTarget.GlobalCompileEnvironment.Config.TargetArchitecture == "-armv7")
+			if (InBuildTarget.GlobalCompileEnvironment.Config.Target.Architecture == "-armv7")
 			{
 				InBuildTarget.GlobalCompileEnvironment.Config.SystemIncludePaths.Add("$(NDKROOT)/sources/cxx-stl/gnu-libstdc++/" + GccVersion + "/include");
 				InBuildTarget.GlobalCompileEnvironment.Config.SystemIncludePaths.Add("$(NDKROOT)/sources/cxx-stl/gnu-libstdc++/" + GccVersion + "/libs/armeabi-v7a/include");
diff --git a/Engine/Source/Programs/UnrealBuildTool/Configuration/UEBuildBinary.cs b/Engine/Source/Programs/UnrealBuildTool/Configuration/UEBuildBinary.cs
--- a/Engine/Source/Programs/UnrealBuildTool/Configuration/UEBuildBinary.cs
+++ b/Engine/Source/Programs/UnrealBuildTool/Configuration/UEBuildBinary.cs
@@ -580,7 +580,7 @@ namespace UnrealBuildTool
 					// Mark the link environment as cross-referenced.
 					BinaryLinkEnvironment.Config.bIsCrossReferenced = true;
 
-					if (BinaryLinkEnvironment.Config.TargetPlatform != CPPTargetPlatform.Mac)
+					if (BinaryLinkEnvironment.Config.Target.Platform != CPPTargetPlatform.Mac)
 					{
 						// Create the import library.
 						OutputFiles.Add(BinaryLinkEnvironment.LinkExecutable(true));
@@ -608,7 +608,7 @@ namespace UnrealBuildTool
 
 				// if building for Mac on a Mac, use actions to finalize the builds (otherwise, we use Deploy)
 				if (ExternalExecution.GetRuntimePlatform() == UnrealTargetPlatform.Mac &&
-					BinaryLinkEnvironment.Config.TargetPlatform == CPPTargetPlatform.Mac)
+					BinaryLinkEnvironment.Config.Target.Platform == CPPTargetPlatform.Mac)
 				{
 					if (!BinaryLinkEnvironment.Config.bIsBuildingDLL && !BinaryLinkEnvironment.Config.bIsBuildingLibrary)
 					{
@@ -657,8 +657,8 @@ namespace UnrealBuildTool
 		{
 			string LibraryFileName;
 			if (Config.Type == UEBuildBinaryType.StaticLibrary
-                || DependentLinkEnvironment.Config.TargetPlatform == CPPTargetPlatform.Mac
-                || DependentLinkEnvironment.Config.TargetPlatform == CPPTargetPlatform.Linux)
+                || DependentLinkEnvironment.Config.Target.Platform == CPPTargetPlatform.Mac
+                || DependentLinkEnvironment.Config.Target.Platform == CPPTargetPlatform.Linux)
 			{
 				LibraryFileName = Config.OutputFilePath;
 			}
@@ -724,7 +724,7 @@ namespace UnrealBuildTool
 		public override IEnumerable<FileItem> Build(CPPEnvironment CompileEnvironment, LinkEnvironment LinkEnvironment)
 		{
 			var ProjectCSharpEnviroment = new CSharpEnvironment();
-			if (LinkEnvironment.Config.TargetConfiguration == CPPTargetConfiguration.Debug)
+			if (LinkEnvironment.Config.Target.Configuration == CPPTargetConfiguration.Debug)
 			{ 
 				ProjectCSharpEnviroment.TargetConfiguration = CSharpTargetConfiguration.Debug;
 			}
@@ -732,7 +732,7 @@ namespace UnrealBuildTool
 			{
 				ProjectCSharpEnviroment.TargetConfiguration = CSharpTargetConfiguration.Development;
 			}
-			ProjectCSharpEnviroment.EnvironmentTargetPlatform = LinkEnvironment.Config.TargetPlatform;
+			ProjectCSharpEnviroment.EnvironmentTargetPlatform = LinkEnvironment.Config.Target.Platform;
 
 			// Currently only supported by windows...
 			UEToolChain.GetPlatformToolChain(CPPTargetPlatform.Win64).CompileCSharpProject(
diff --git a/Engine/Source/Programs/UnrealBuildTool/Configuration/UEBuildModule.cs b/Engine/Source/Programs/UnrealBuildTool/Configuration/UEBuildModule.cs
--- a/Engine/Source/Programs/UnrealBuildTool/Configuration/UEBuildModule.cs
+++ b/Engine/Source/Programs/UnrealBuildTool/Configuration/UEBuildModule.cs
@@ -1012,7 +1012,7 @@ namespace UnrealBuildTool
 		// UEBuildModule interface.
 		public override List<FileItem> Compile(CPPEnvironment GlobalCompileEnvironment, CPPEnvironment CompileEnvironment, bool bCompileMonolithic)
 		{
-			var BuildPlatform = UEBuildPlatform.GetBuildPlatformForCPPTargetPlatform(CompileEnvironment.Config.TargetPlatform);
+			var BuildPlatform = UEBuildPlatform.GetBuildPlatformForCPPTargetPlatform(CompileEnvironment.Config.Target.Platform);
 
 			var LinkInputFiles = new List<FileItem>();
 			if( ProjectFileGenerator.bGenerateProjectFiles && IntelliSenseGatherer == null)
@@ -1103,7 +1103,7 @@ namespace UnrealBuildTool
 
 			// Precompiled header support.
 			bool bWasModuleCodeCompiled = false;
-			if (BuildPlatform.ShouldUsePCHFiles(CompileEnvironment.Config.TargetPlatform, CompileEnvironment.Config.TargetConfiguration))
+			if (BuildPlatform.ShouldUsePCHFiles(CompileEnvironment.Config.Target.Platform, CompileEnvironment.Config.Target.Configuration))
 			{
 				var PCHTimerStart = DateTime.UtcNow;
 
@@ -1432,9 +1432,7 @@ namespace UnrealBuildTool
 				// Create a link environment for it
 				LinkEnvironment RedistLinkEnvironment = new LinkEnvironment();
 				RedistLinkEnvironment.InputFiles.AddRange(LinkInputFiles);
-				RedistLinkEnvironment.Config.TargetArchitecture    = CompileEnvironment.Config.TargetArchitecture;
-				RedistLinkEnvironment.Config.TargetConfiguration   = CompileEnvironment.Config.TargetConfiguration;
-				RedistLinkEnvironment.Config.TargetPlatform        = CompileEnvironment.Config.TargetPlatform;
+				RedistLinkEnvironment.Config.Target                = CompileEnvironment.Config.Target;
 				RedistLinkEnvironment.Config.bIsBuildingDLL        = false;
 				RedistLinkEnvironment.Config.bIsBuildingLibrary    = true;
 				RedistLinkEnvironment.Config.IntermediateDirectory = Binary.Config.IntermediateDirectory;
@@ -1471,7 +1469,7 @@ namespace UnrealBuildTool
 			var SharedPCHCodeOptimization = SharedPCHEnvironment.OptimizeCode;
 			var ModuleCodeOptimization    = ModuleCompileEnvironment.Config.OptimizeCode;
 
-			if (CompileEnvironment.Config.TargetConfiguration != CPPTargetConfiguration.Debug)
+			if (CompileEnvironment.Config.Target.Configuration != CPPTargetConfiguration.Debug)
 			{
 				if (SharedPCHCodeOptimization == ModuleRules.CodeOptimization.InNonDebugBuilds)
 				{
@@ -1512,7 +1510,7 @@ namespace UnrealBuildTool
 
 		private FileItem ProcessDependencies(FileItem CPPFile, CPPEnvironment ModuleCompileEnvironment)
 		{
-			List<DependencyInclude> DirectIncludeFilenames = CPPEnvironment.GetDirectIncludeDependencies(CPPFile, ModuleCompileEnvironment.Config.TargetPlatform);
+			List<DependencyInclude> DirectIncludeFilenames = CPPEnvironment.GetDirectIncludeDependencies(CPPFile, ModuleCompileEnvironment.Config.Target.Platform);
 			if (BuildConfiguration.bPrintDebugInfo)
 			{
 				Log.TraceVerbose("Found direct includes for {0}: {1}", Path.GetFileName(CPPFile.AbsolutePath), string.Join(", ", DirectIncludeFilenames.Select(F => F.IncludeName)));
@@ -1575,7 +1573,7 @@ namespace UnrealBuildTool
 			// Switch the optimization flag if we're building a game module
 			if (Target.Configuration == UnrealTargetConfiguration.DebugGame && Type == UEBuildModuleType.Game)
 			{
-				Result.Config.TargetConfiguration = CPPTargetConfiguration.Debug;
+				Result.Config.Target.Configuration = CPPTargetConfiguration.Debug;
 			}
 
 			// Add the module's private definitions.
diff --git a/Engine/Source/Programs/UnrealBuildTool/Configuration/UEBuildPlatform.cs b/Engine/Source/Programs/UnrealBuildTool/Configuration/UEBuildPlatform.cs
--- a/Engine/Source/Programs/UnrealBuildTool/Configuration/UEBuildPlatform.cs
+++ b/Engine/Source/Programs/UnrealBuildTool/Configuration/UEBuildPlatform.cs
@@ -1135,8 +1135,8 @@ namespace UnrealBuildTool
 			}
 
 			// Set up the global C++ compilation and link environment.
-			InBuildTarget.GlobalCompileEnvironment.Config.TargetConfiguration = CompileConfiguration;
-			InBuildTarget.GlobalLinkEnvironment.Config.TargetConfiguration = CompileConfiguration;
+			InBuildTarget.GlobalCompileEnvironment.Config.Target.Configuration = CompileConfiguration;
+			InBuildTarget.GlobalLinkEnvironment.Config.Target.Configuration = CompileConfiguration;
 
 			// Create debug info based on the heuristics specified by the user.
 			InBuildTarget.GlobalCompileEnvironment.Config.bCreateDebugInfo =
diff --git a/Engine/Source/Programs/UnrealBuildTool/Configuration/UEBuildTarget.cs b/Engine/Source/Programs/UnrealBuildTool/Configuration/UEBuildTarget.cs
--- a/Engine/Source/Programs/UnrealBuildTool/Configuration/UEBuildTarget.cs
+++ b/Engine/Source/Programs/UnrealBuildTool/Configuration/UEBuildTarget.cs
@@ -1152,7 +1152,7 @@ namespace UnrealBuildTool
 			AppBinaries.AddRange( NewBinaries );
 
 			// On Mac AppBinaries paths for non-console targets need to be adjusted to be inside the app bundle
-			if (GlobalLinkEnvironment.Config.TargetPlatform == CPPTargetPlatform.Mac && !GlobalLinkEnvironment.Config.bIsBuildingConsoleApplication)
+			if (GlobalLinkEnvironment.Config.Target.Platform == CPPTargetPlatform.Mac && !GlobalLinkEnvironment.Config.bIsBuildingConsoleApplication)
 			{
 				MacToolChain.FixBundleBinariesPaths( this, AppBinaries );
 			}
@@ -1347,7 +1347,7 @@ namespace UnrealBuildTool
 			}
 
 			// On Mac we have actions that should be executed after all the binaries are created
-			if (GlobalLinkEnvironment.Config.TargetPlatform == CPPTargetPlatform.Mac)
+			if (GlobalLinkEnvironment.Config.Target.Platform == CPPTargetPlatform.Mac)
 			{
 				MacToolChain.SetupBundleDependencies(AppBinaries, GameName);
 			}
@@ -1355,7 +1355,7 @@ namespace UnrealBuildTool
 			// If we're only generating the manifest, return now
 			if (UEBuildConfiguration.bGenerateManifest || UEBuildConfiguration.bCleanProject)
 			{
-                GenerateManifest(AppBinaries, GlobalLinkEnvironment.Config.TargetPlatform, SpecialRocketLibFilesThatAreBuildProducts);
+                GenerateManifest(AppBinaries, GlobalLinkEnvironment.Config.Target.Platform, SpecialRocketLibFilesThatAreBuildProducts);
                 if (!BuildConfiguration.bXGEExport)
                 {
                     return ECompilationResult.Succeeded;
@@ -2277,8 +2277,8 @@ namespace UnrealBuildTool
 
 			// Validates current settings and updates if required.
 			BuildConfiguration.ValidateConfiguration(
-				GlobalCompileEnvironment.Config.TargetConfiguration,
-				GlobalCompileEnvironment.Config.TargetPlatform,
+				GlobalCompileEnvironment.Config.Target.Configuration,
+				GlobalCompileEnvironment.Config.Target.Platform,
 				GlobalCompileEnvironment.Config.bCreateDebugInfo);
 		}
 
@@ -2288,12 +2288,12 @@ namespace UnrealBuildTool
 
 			CPPTargetPlatform MainCompilePlatform = BuildPlatform.GetCPPTargetPlatform(Platform);
 
-			GlobalLinkEnvironment.Config.TargetPlatform = MainCompilePlatform;
-			GlobalCompileEnvironment.Config.TargetPlatform = MainCompilePlatform;
+			GlobalLinkEnvironment.Config.Target.Platform = MainCompilePlatform;
+			GlobalCompileEnvironment.Config.Target.Platform = MainCompilePlatform;
 
 			string ActiveArchitecture = BuildPlatform.GetActiveArchitecture();
-			GlobalCompileEnvironment.Config.TargetArchitecture = ActiveArchitecture;
-			GlobalLinkEnvironment.Config.TargetArchitecture = ActiveArchitecture;
+			GlobalCompileEnvironment.Config.Target.Architecture = ActiveArchitecture;
+			GlobalLinkEnvironment.Config.Target.Architecture = ActiveArchitecture;
 
 
 			// Set up the platform-specific environment.
diff --git a/Engine/Source/Programs/UnrealBuildTool/HTML5/HTML5ToolChain.cs b/Engine/Source/Programs/UnrealBuildTool/HTML5/HTML5ToolChain.cs
--- a/Engine/Source/Programs/UnrealBuildTool/HTML5/HTML5ToolChain.cs
+++ b/Engine/Source/Programs/UnrealBuildTool/HTML5/HTML5ToolChain.cs
@@ -99,9 +99,9 @@ namespace UnrealBuildTool
 		
 		static string GetCLArguments_Global(CPPEnvironment CompileEnvironment)
 		{
-			string Result = GetSharedArguments_Global(CompileEnvironment.Config.TargetConfiguration, CompileEnvironment.Config.TargetArchitecture);
+			string Result = GetSharedArguments_Global(CompileEnvironment.Config.Target.Configuration, CompileEnvironment.Config.Target.Architecture);
 
-			if (CompileEnvironment.Config.TargetArchitecture != "-win32")
+			if (CompileEnvironment.Config.Target.Architecture != "-win32")
 			{
 				// do we want debug info?
 				/*			if (CompileEnvironment.Config.bCreateDebugInfo)
@@ -111,15 +111,15 @@ namespace UnrealBuildTool
 
 				Result += " -Wno-warn-absolute-paths ";
 
-				if (CompileEnvironment.Config.TargetConfiguration == CPPTargetConfiguration.Debug)
+				if (CompileEnvironment.Config.Target.Configuration == CPPTargetConfiguration.Debug)
 				{
 					Result += " -O0";
 				}
-				if (CompileEnvironment.Config.TargetConfiguration == CPPTargetConfiguration.Debug || CompileEnvironment.Config.TargetConfiguration == CPPTargetConfiguration.Development)
+				if (CompileEnvironment.Config.Target.Configuration == CPPTargetConfiguration.Debug || CompileEnvironment.Config.Target.Configuration == CPPTargetConfiguration.Development)
 				{
 					Result += " -s GL_ASSERTIONS=1 ";
 				}
-				if (CompileEnvironment.Config.TargetConfiguration == CPPTargetConfiguration.Development)
+				if (CompileEnvironment.Config.Target.Configuration == CPPTargetConfiguration.Development)
 				{
 					if (UEBuildConfiguration.bCompileForSize)
 					{
@@ -130,7 +130,7 @@ namespace UnrealBuildTool
 						Result += " -O2 -s ASM_JS=1 -s OUTLINING_LIMIT=110000";
 					}
 				}
-				if (CompileEnvironment.Config.TargetConfiguration == CPPTargetConfiguration.Shipping)
+				if (CompileEnvironment.Config.Target.Configuration == CPPTargetConfiguration.Shipping)
 				{
 					if (UEBuildConfiguration.bCompileForSize)
 					{
@@ -151,7 +151,7 @@ namespace UnrealBuildTool
 		{
 			string Result = "";
 
-			if (CompileEnvironment.Config.TargetArchitecture != "-win32")
+			if (CompileEnvironment.Config.Target.Architecture != "-win32")
 			{
 				Result = " -std=c++11";
 			}
@@ -167,21 +167,21 @@ namespace UnrealBuildTool
 
 		static string GetLinkArguments(LinkEnvironment LinkEnvironment)
 		{
-			string Result = GetSharedArguments_Global(LinkEnvironment.Config.TargetConfiguration, LinkEnvironment.Config.TargetArchitecture);
+			string Result = GetSharedArguments_Global(LinkEnvironment.Config.Target.Configuration, LinkEnvironment.Config.Target.Architecture);
 
-			if (LinkEnvironment.Config.TargetArchitecture != "-win32")
+			if (LinkEnvironment.Config.Target.Architecture != "-win32")
 			{
 
 				// enable verbose mode
 				Result += " -v";
 
-				if (LinkEnvironment.Config.TargetConfiguration != CPPTargetConfiguration.Shipping)
+				if (LinkEnvironment.Config.Target.Configuration != CPPTargetConfiguration.Shipping)
 				{
 					// enable caching of intermediate date
 					Result += " --jcache -g";
 				}
 
-				if (LinkEnvironment.Config.TargetConfiguration == CPPTargetConfiguration.Debug)
+				if (LinkEnvironment.Config.Target.Configuration == CPPTargetConfiguration.Debug)
 				{
 					// check for alignment/etc checking
 					//Result += " -s SAFE_HEAP=1";
@@ -192,19 +192,19 @@ namespace UnrealBuildTool
 					Result += " -s ASSERTIONS=1";
 				}
 
-				if (LinkEnvironment.Config.TargetConfiguration == CPPTargetConfiguration.Debug)
+				if (LinkEnvironment.Config.Target.Configuration == CPPTargetConfiguration.Debug)
 				{
 					Result += " -O0";
 				}
-				if (LinkEnvironment.Config.TargetConfiguration == CPPTargetConfiguration.Debug || LinkEnvironment.Config.TargetConfiguration == CPPTargetConfiguration.Development)
+				if (LinkEnvironment.Config.Target.Configuration == CPPTargetConfiguration.Debug || LinkEnvironment.Config.Target.Configuration == CPPTargetConfiguration.Development)
 				{
 					Result += " -s GL_ASSERTIONS=1 ";
 				}
-				if (LinkEnvironment.Config.TargetConfiguration == CPPTargetConfiguration.Development)
+				if (LinkEnvironment.Config.Target.Configuration == CPPTargetConfiguration.Development)
 				{
 					Result += " -O2 -s ASM_JS=1 -s OUTLINING_LIMIT=110000";
 				}
-				if (LinkEnvironment.Config.TargetConfiguration == CPPTargetConfiguration.Shipping)
+				if (LinkEnvironment.Config.Target.Configuration == CPPTargetConfiguration.Shipping)
 				{
 					Result += " -O3 -s ASM_JS=1 -s OUTLINING_LIMIT=40000";
 				}
@@ -224,7 +224,7 @@ namespace UnrealBuildTool
 		{
 			string Result = "";
 
-			if (LinkEnvironment.Config.TargetArchitecture == "-win32")
+			if (LinkEnvironment.Config.Target.Architecture == "-win32")
 			{
 				// Prevents the linker from displaying its logo for each invocation.
 				Result += " /NOLOGO";
@@ -241,7 +241,7 @@ namespace UnrealBuildTool
 				//
 				//	Shipping & LTCG
 				//
-				if (LinkEnvironment.Config.TargetConfiguration == CPPTargetConfiguration.Shipping)
+				if (LinkEnvironment.Config.Target.Configuration == CPPTargetConfiguration.Shipping)
 				{
 					// Use link-time code generation.
 					Result += " /ltcg";
@@ -303,7 +303,7 @@ namespace UnrealBuildTool
 
 		public override CPPOutput CompileCPPFiles(CPPEnvironment CompileEnvironment, List<FileItem> SourceFiles, string ModuleName)
 		{
-			if (CompileEnvironment.Config.TargetArchitecture == "-win32")
+			if (CompileEnvironment.Config.Target.Architecture == "-win32")
 			{
 				return base.CompileCPPFiles(CompileEnvironment, SourceFiles, ModuleName);
 			}
@@ -364,7 +364,7 @@ namespace UnrealBuildTool
 				// Add C or C++ specific compiler arguments.
 				if (bIsPlainCFile)
 				{
-					FileArguments += GetCLArguments_C(CompileEnvironment.Config.TargetArchitecture);
+					FileArguments += GetCLArguments_C(CompileEnvironment.Config.Target.Architecture);
 				}
 				else
 				{
@@ -403,7 +403,7 @@ namespace UnrealBuildTool
 		{
 			CPPOutput Result = new CPPOutput();
 
-			if (Environment.Config.TargetArchitecture == "-win32")
+			if (Environment.Config.Target.Architecture == "-win32")
 			{
 				return base.CompileRCFiles(Environment, RCFiles);
 			}
@@ -468,7 +468,7 @@ namespace UnrealBuildTool
 
 		public override FileItem LinkFiles(LinkEnvironment LinkEnvironment, bool bBuildImportLibraryOnly)
 		{
-			if (LinkEnvironment.Config.TargetArchitecture == "-win32")
+			if (LinkEnvironment.Config.Target.Architecture == "-win32")
             {
                 return base.LinkFiles(LinkEnvironment, bBuildImportLibraryOnly);
             }
diff --git a/Engine/Source/Programs/UnrealBuildTool/HTML5/UEBuildHTML5.cs b/Engine/Source/Programs/UnrealBuildTool/HTML5/UEBuildHTML5.cs
--- a/Engine/Source/Programs/UnrealBuildTool/HTML5/UEBuildHTML5.cs
+++ b/Engine/Source/Programs/UnrealBuildTool/HTML5/UEBuildHTML5.cs
@@ -293,7 +293,7 @@ namespace UnrealBuildTool
 		{
 			InBuildTarget.GlobalCompileEnvironment.Config.Definitions.Add("PLATFORM_HTML5=1");
             InBuildTarget.GlobalCompileEnvironment.Config.Definitions.Add("HTML5=1");
-            if (InBuildTarget.GlobalCompileEnvironment.Config.TargetArchitecture == "-win32")
+            if (InBuildTarget.GlobalCompileEnvironment.Config.Target.Architecture == "-win32")
 			{
 				InBuildTarget.GlobalCompileEnvironment.Config.Definitions.Add("PLATFORM_HTML5_WIN32=1");
 				InBuildTarget.GlobalLinkEnvironment.Config.AdditionalLibraries.Add("delayimp.lib");
diff --git a/Engine/Source/Programs/UnrealBuildTool/IOS/IOSToolChain.cs b/Engine/Source/Programs/UnrealBuildTool/IOS/IOSToolChain.cs
--- a/Engine/Source/Programs/UnrealBuildTool/IOS/IOSToolChain.cs
+++ b/Engine/Source/Programs/UnrealBuildTool/IOS/IOSToolChain.cs
@@ -232,9 +232,9 @@ namespace UnrealBuildTool
 			Result += " -c";
 
 			// What architecture(s) to build for
-			Result += GetArchitectureArgument(CompileEnvironment.Config.TargetConfiguration, CompileEnvironment.Config.TargetArchitecture);
+			Result += GetArchitectureArgument(CompileEnvironment.Config.Target.Configuration, CompileEnvironment.Config.Target.Architecture);
 
-			if (CompileEnvironment.Config.TargetArchitecture == "-simulator")
+			if (CompileEnvironment.Config.Target.Architecture == "-simulator")
 			{
 				Result += " -isysroot " + BaseSDKDirSim + "/iPhoneSimulator" + IOSSDKVersion + ".sdk";
 			}
@@ -246,7 +246,7 @@ namespace UnrealBuildTool
 			Result += " -miphoneos-version-min=" + IOSVersion;
 
 			// Optimize non- debug builds.
-			if (CompileEnvironment.Config.TargetConfiguration != CPPTargetConfiguration.Debug)
+			if (CompileEnvironment.Config.Target.Configuration != CPPTargetConfiguration.Debug)
 			{
 				Result += " -O3";
 			}
@@ -377,14 +377,14 @@ namespace UnrealBuildTool
 		string GetLinkArguments_Global( LinkEnvironment LinkEnvironment )
 		{
 			string Result = "";
-			if (LinkEnvironment.Config.TargetArchitecture == "-simulator")
+			if (LinkEnvironment.Config.Target.Architecture == "-simulator")
 			{
 				Result += " -arch i386";
 				Result += " -isysroot " + XcodeDeveloperDir + "Platforms/iPhoneSimulator.platform/Developer/SDKs/iPhoneSimulator" + IOSSDKVersion + ".sdk";
 			}
 			else
 			{
-				Result += Result += GetArchitectureArgument(LinkEnvironment.Config.TargetConfiguration, LinkEnvironment.Config.TargetArchitecture);
+				Result += Result += GetArchitectureArgument(LinkEnvironment.Config.Target.Configuration, LinkEnvironment.Config.Target.Architecture);
 				Result += " -isysroot " + XcodeDeveloperDir + "Platforms/iPhoneOS.platform/Developer/SDKs/iPhoneOS" + IOSSDKVersion + ".sdk";
 			}
 			Result += " -dead_strip";
diff --git a/Engine/Source/Programs/UnrealBuildTool/Linux/LinuxToolChain.cs b/Engine/Source/Programs/UnrealBuildTool/Linux/LinuxToolChain.cs
--- a/Engine/Source/Programs/UnrealBuildTool/Linux/LinuxToolChain.cs
+++ b/Engine/Source/Programs/UnrealBuildTool/Linux/LinuxToolChain.cs
@@ -140,7 +140,7 @@ namespace UnrealBuildTool.Linux
             //Result += " -DOPERATOR_NEW_INLINE=FORCENOINLINE";
 
             // shipping builds will cause this warning with "ensure", so disable only in those case
-            if (CompileEnvironment.Config.TargetConfiguration == CPPTargetConfiguration.Shipping)
+            if (CompileEnvironment.Config.Target.Configuration == CPPTargetConfiguration.Shipping)
             {
                 Result += " -Wno-unused-value";
 
@@ -149,7 +149,7 @@ namespace UnrealBuildTool.Linux
                 Result += " -fomit-frame-pointer";
                 Result += " -fvisibility=hidden";           // prevents from exporting all symbols (reduces the size of the binary)
             }
-            else if (CompileEnvironment.Config.TargetConfiguration == CPPTargetConfiguration.Debug)
+            else if (CompileEnvironment.Config.Target.Configuration == CPPTargetConfiguration.Debug)
             {
                 Result += " -fno-inline";                   // disable inlining for better debuggability (e.g. callstacks, "skip file" in gdb)
             }
@@ -163,13 +163,13 @@ namespace UnrealBuildTool.Linux
                 Result += " -fstack-protector";
                 //Result += " -fsanitize=address";  // Preferred clang tool for detecting address based errors but unusable for some reason with Module.Engine.7_of_42.cpp
             }
-            else if (CompileEnvironment.Config.TargetConfiguration < CPPTargetConfiguration.Shipping)
+            else if (CompileEnvironment.Config.Target.Configuration < CPPTargetConfiguration.Shipping)
             {
                 Result += " -gline-tables-only"; // include debug info for meaningful callstacks
             }
 
             // optimization level
-            if (CompileEnvironment.Config.TargetConfiguration == CPPTargetConfiguration.Debug)
+            if (CompileEnvironment.Config.Target.Configuration == CPPTargetConfiguration.Debug)
             {
                 Result += " -O0";
             }
@@ -249,7 +249,7 @@ namespace UnrealBuildTool.Linux
             string Result = "";
 
             // debugging symbols
-            if (LinkEnvironment.Config.TargetConfiguration < CPPTargetConfiguration.Shipping)
+            if (LinkEnvironment.Config.Target.Configuration < CPPTargetConfiguration.Shipping)
             {
                 Result += " -rdynamic";   // needed for backtrace_symbols()...
             }
diff --git a/Engine/Source/Programs/UnrealBuildTool/Mac/MacToolChain.cs b/Engine/Source/Programs/UnrealBuildTool/Mac/MacToolChain.cs
--- a/Engine/Source/Programs/UnrealBuildTool/Mac/MacToolChain.cs
+++ b/Engine/Source/Programs/UnrealBuildTool/Mac/MacToolChain.cs
@@ -103,7 +103,7 @@ namespace UnrealBuildTool
 			Result += " -mmacosx-version-min=" + MacOSVersion;
 
 			// Optimize non- debug builds.
-			if (CompileEnvironment.Config.TargetConfiguration != CPPTargetConfiguration.Debug)
+			if (CompileEnvironment.Config.Target.Configuration != CPPTargetConfiguration.Debug)
 			{
 				Result += " -O3";
 			}
diff --git a/Engine/Source/Programs/UnrealBuildTool/System/CPPEnvironment.cs b/Engine/Source/Programs/UnrealBuildTool/System/CPPEnvironment.cs
--- a/Engine/Source/Programs/UnrealBuildTool/System/CPPEnvironment.cs
+++ b/Engine/Source/Programs/UnrealBuildTool/System/CPPEnvironment.cs
@@ -93,8 +93,87 @@ namespace UnrealBuildTool
 		public CPPTargetPlatform EnvironmentTargetPlatform;
 	}
 
+	public abstract class NativeBuildEnvironmentConfiguration
+	{
+		public class TargetInfo
+		{
+			/** The platform to be compiled/linked for. */
+			public CPPTargetPlatform Platform;
+
+			/** The architecture that is being compiled/linked (empty string by default) */
+			public string Architecture;
+
+			/** The configuration to be compiled/linked for. */
+			public CPPTargetConfiguration Configuration;
+
+			public TargetInfo()
+			{
+				Architecture = "";
+			}
+
+			public TargetInfo(TargetInfo Target)
+			{
+				Platform      = Target.Platform;
+				Architecture  = Target.Architecture;
+				Configuration = Target.Configuration;
+			}
+		}
+
+		public TargetInfo Target;
+
+		[Obsolete("Use Target.Platform instead")]
+		public CPPTargetPlatform TargetPlatform
+		{
+			get
+			{
+				return Target.Platform;
+			}
+			set
+			{
+				Target.Platform = value;
+			}
+		}
+
+		[Obsolete("Use Target.Architecture instead")]
+		public string TargetArchitecture
+		{
+			get
+			{
+				return Target.Architecture;
+			}
+			set
+			{
+				Target.Architecture = value;
+			}
+		}
+
+		[Obsolete("Use Target.Configuration instead")]
+		public CPPTargetConfiguration TargetConfiguration
+		{
+			get
+			{
+				return Target.Configuration;
+			}
+			set
+			{
+				Target.Configuration = value;
+			}
+		}
+
+		protected NativeBuildEnvironmentConfiguration()
+		{
+			Target = new TargetInfo();
+		}
+
+		protected NativeBuildEnvironmentConfiguration(NativeBuildEnvironmentConfiguration Configuration)
+		{
+			Target = new TargetInfo(Configuration.Target);
+		}
+		
+	}
+
 	/** Encapsulates the configuration of an environment that a C++ file is compiled in. */
-	public class CPPEnvironmentConfiguration
+	public class CPPEnvironmentConfiguration : NativeBuildEnvironmentConfiguration
 	{
 		/** The directory to put the output object/debug files in. */
 		public string OutputDirectory = null;
@@ -117,15 +196,6 @@ namespace UnrealBuildTool
 		    include PCH header files first. */
 		public bool bForceIncludePrecompiledHeader = false;
 
-		/** The platform to be compiled for. */
-		public CPPTargetPlatform TargetPlatform;
-
-		/** The architecture that is being linked (empty string by default) */
-		public string TargetArchitecture = "";
-
-		/** The configuration to be compiled for. */
-		public CPPTargetConfiguration TargetConfiguration;
-
 		/** Use run time type information */
 		public bool bUseRTTI = false;
 
@@ -189,7 +259,8 @@ namespace UnrealBuildTool
 		}
 
 		/** Copy constructor. */
-		public CPPEnvironmentConfiguration(CPPEnvironmentConfiguration InCopyEnvironment)
+		public CPPEnvironmentConfiguration(CPPEnvironmentConfiguration InCopyEnvironment):
+			base(InCopyEnvironment)
 		{
 			OutputDirectory                        = InCopyEnvironment.OutputDirectory;
 			LocalShadowDirectory                   = InCopyEnvironment.LocalShadowDirectory;
@@ -197,9 +268,6 @@ namespace UnrealBuildTool
 			PrecompiledHeaderIncludeFilename       = InCopyEnvironment.PrecompiledHeaderIncludeFilename;
 			PrecompiledHeaderAction                = InCopyEnvironment.PrecompiledHeaderAction;
 			bForceIncludePrecompiledHeader         = InCopyEnvironment.bForceIncludePrecompiledHeader;
-			TargetPlatform                         = InCopyEnvironment.TargetPlatform;
-			TargetArchitecture                     = InCopyEnvironment.TargetArchitecture;
-			TargetConfiguration                    = InCopyEnvironment.TargetConfiguration;
 			bUseRTTI                               = InCopyEnvironment.bUseRTTI;
 			bFasterWithoutUnity                    = InCopyEnvironment.bFasterWithoutUnity;
 			MinFilesUsingPrecompiledHeaderOverride = InCopyEnvironment.MinFilesUsingPrecompiledHeaderOverride;
@@ -280,7 +348,7 @@ namespace UnrealBuildTool
 		 */
 		public CPPOutput CompileFiles(List<FileItem> CPPFiles, string ModuleName)
 		{
-			return UEToolChain.GetPlatformToolChain(Config.TargetPlatform).CompileCPPFiles(this, CPPFiles, ModuleName);
+			return UEToolChain.GetPlatformToolChain(Config.Target.Platform).CompileCPPFiles(this, CPPFiles, ModuleName);
 		}
 
 		/**
@@ -290,7 +358,7 @@ namespace UnrealBuildTool
 		 */
 		public CPPOutput CompileRCFiles(List<FileItem> RCFiles)
 		{
-			return UEToolChain.GetPlatformToolChain(Config.TargetPlatform).CompileRCFiles(this, RCFiles);
+			return UEToolChain.GetPlatformToolChain(Config.Target.Platform).CompileRCFiles(this, RCFiles);
 		}
 
 		/**
@@ -300,7 +368,7 @@ namespace UnrealBuildTool
 		 */
 		public bool ShouldUsePCHs()
 		{
-			return UEBuildPlatform.GetBuildPlatformForCPPTargetPlatform(Config.TargetPlatform).ShouldUsePCHFiles(Config.TargetPlatform, Config.TargetConfiguration);
+			return UEBuildPlatform.GetBuildPlatformForCPPTargetPlatform(Config.Target.Platform).ShouldUsePCHFiles(Config.Target.Platform, Config.Target.Configuration);
 		}
 
 		public void AddPrivateAssembly( string FilePath )
diff --git a/Engine/Source/Programs/UnrealBuildTool/System/CPPHeaders.cs b/Engine/Source/Programs/UnrealBuildTool/System/CPPHeaders.cs
--- a/Engine/Source/Programs/UnrealBuildTool/System/CPPHeaders.cs
+++ b/Engine/Source/Programs/UnrealBuildTool/System/CPPHeaders.cs
@@ -170,7 +170,7 @@ namespace UnrealBuildTool
 				IncludedFilesMap.Add(CPPFile, IncludedFileList);
 
 				// Gather a list of names of files directly included by this C++ file.
-				List<DependencyInclude> DirectIncludes = GetDirectIncludeDependencies(CPPFile, Config.TargetPlatform);
+				List<DependencyInclude> DirectIncludes = GetDirectIncludeDependencies(CPPFile, Config.Target.Platform);
 
 				// Build a list of the unique set of files that are included by this file.
 				string SourceFilesDirectory = Path.GetDirectoryName( CPPFile.AbsolutePath );
diff --git a/Engine/Source/Programs/UnrealBuildTool/System/LinkEnvironment.cs b/Engine/Source/Programs/UnrealBuildTool/System/LinkEnvironment.cs
--- a/Engine/Source/Programs/UnrealBuildTool/System/LinkEnvironment.cs
+++ b/Engine/Source/Programs/UnrealBuildTool/System/LinkEnvironment.cs
@@ -9,7 +9,7 @@ namespace UnrealBuildTool
 	/// <summary>
 	/// Configuration class for LinkEnvironment
 	/// </summary>
-	public class LinkEnvironmentConfiguration
+	public class LinkEnvironmentConfiguration : NativeBuildEnvironmentConfiguration
 	{
 		/** The directory to put the non-executable files in (PDBs, import library, etc) */
 		public string OutputDirectory;
@@ -23,15 +23,6 @@ namespace UnrealBuildTool
 		/** The file path for the executable file that is output by the linker. */
 		public string OutputFilePath;
 
-		/** The platform that is being linked for. */
-		public CPPTargetPlatform TargetPlatform;
-
-		/** The architecture that is being linked (empty string by default) */
-		public string TargetArchitecture = "";
-
-		/** The configuration that is being linked for. */
-		public CPPTargetConfiguration TargetConfiguration;
-
 		/** A list of the paths used to find libraries. */
 		public List<string> LibraryPaths = new List<string>();
 
@@ -104,15 +95,13 @@ namespace UnrealBuildTool
 		}
 
 		/** Copy constructor. */
-		public LinkEnvironmentConfiguration(LinkEnvironmentConfiguration InCopyEnvironment)
+		public LinkEnvironmentConfiguration(LinkEnvironmentConfiguration InCopyEnvironment):
+			base(InCopyEnvironment)
 		{
 			OutputDirectory = InCopyEnvironment.OutputDirectory;
 			IntermediateDirectory = InCopyEnvironment.IntermediateDirectory;
 			LocalShadowDirectory = InCopyEnvironment.LocalShadowDirectory;
 			OutputFilePath = InCopyEnvironment.OutputFilePath;
-			TargetPlatform = InCopyEnvironment.TargetPlatform;
-			TargetConfiguration = InCopyEnvironment.TargetConfiguration;
-			TargetArchitecture = InCopyEnvironment.TargetArchitecture;
 			LibraryPaths.AddRange(InCopyEnvironment.LibraryPaths);
 			ExcludedLibraries.AddRange(InCopyEnvironment.ExcludedLibraries);
 			AdditionalLibraries.AddRange(InCopyEnvironment.AdditionalLibraries);
@@ -139,11 +128,11 @@ namespace UnrealBuildTool
 		{
 			get
 			{
-				return TargetPlatform == CPPTargetPlatform.Win64 &&
+				return Target.Platform == CPPTargetPlatform.Win64 &&
 						bIsBuildingConsoleApplication == false &&
 						bBuildAdditionalConsoleApplication == true &&
 						bIsBuildingDLL == false &&
-						(TargetConfiguration == CPPTargetConfiguration.Debug || TargetConfiguration == CPPTargetConfiguration.Development);
+						(Target.Configuration == CPPTargetConfiguration.Debug || Target.Configuration == CPPTargetConfiguration.Development);
 			}
 		}
 	}
@@ -177,7 +166,7 @@ namespace UnrealBuildTool
 		/** Links the input files into an executable. */
 		public FileItem LinkExecutable( bool bBuildImportLibraryOnly )
 		{
-			return UEToolChain.GetPlatformToolChain(Config.TargetPlatform).LinkFiles(this, bBuildImportLibraryOnly);
+			return UEToolChain.GetPlatformToolChain(Config.Target.Platform).LinkFiles(this, bBuildImportLibraryOnly);
 		}
 
 		/// <summary>
diff --git a/Engine/Source/Programs/UnrealBuildTool/System/Unity.cs b/Engine/Source/Programs/UnrealBuildTool/System/Unity.cs
--- a/Engine/Source/Programs/UnrealBuildTool/System/Unity.cs
+++ b/Engine/Source/Programs/UnrealBuildTool/System/Unity.cs
@@ -112,9 +112,9 @@ namespace UnrealBuildTool
 			string BaseName
 			)
 		{
-			var ToolChain = UEToolChain.GetPlatformToolChain(CompileEnvironment.Config.TargetPlatform);
+			var ToolChain = UEToolChain.GetPlatformToolChain(CompileEnvironment.Config.Target.Platform);
 
-			var BuildPlatform = UEBuildPlatform.GetBuildPlatformForCPPTargetPlatform(CompileEnvironment.Config.TargetPlatform);
+			var BuildPlatform = UEBuildPlatform.GetBuildPlatformForCPPTargetPlatform(CompileEnvironment.Config.Target.Platform);
 
 			// Figure out size of all input files combined. We use this to determine whether to use larger unity threshold or not.
 			long TotalBytesInCPPFiles = CPPFiles.Sum(F => F.Info.Length);
diff --git a/Engine/Source/Programs/UnrealBuildTool/System/Utils.cs b/Engine/Source/Programs/UnrealBuildTool/System/Utils.cs
--- a/Engine/Source/Programs/UnrealBuildTool/System/Utils.cs
+++ b/Engine/Source/Programs/UnrealBuildTool/System/Utils.cs
@@ -63,7 +63,7 @@ namespace UnrealBuildTool
 			if (PrecompiledHeaderIncludeFilename.bExists)
 			{
 				// Create a Dummy wrapper around the PCH to avoid problems with #pragma once on clang
-				var ToolChain = UEToolChain.GetPlatformToolChain(ProjectCPPEnvironment.Config.TargetPlatform);
+				var ToolChain = UEToolChain.GetPlatformToolChain(ProjectCPPEnvironment.Config.Target.Platform);
 				string PCHGuardDefine = Path.GetFileNameWithoutExtension(PrecompiledHeaderIncludeFilename.AbsolutePath).ToUpper();
 				string LocalPCHHeaderNameInCode = ToolChain.ConvertPath(PrecompiledHeaderIncludeFilename.AbsolutePath);
 				string TmpPCHHeaderContents = String.Format("#ifndef __AUTO_{0}_H__\n#define __AUTO_{0}_H__\n//Last Write: {2}\n#include \"{1}\"\n#endif//__AUTO_{0}_H__", PCHGuardDefine, LocalPCHHeaderNameInCode, PrecompiledHeaderIncludeFilename.LastWriteTime);
diff --git a/Engine/Source/Programs/UnrealBuildTool/ToolChain/UEToolChain.cs b/Engine/Source/Programs/UnrealBuildTool/ToolChain/UEToolChain.cs
--- a/Engine/Source/Programs/UnrealBuildTool/ToolChain/UEToolChain.cs
+++ b/Engine/Source/Programs/UnrealBuildTool/ToolChain/UEToolChain.cs
@@ -99,7 +99,7 @@ namespace UnrealBuildTool
 						UnrealBuildTool.GetUProjectPath(), 
 						BuildConfiguration.PlatformIntermediateFolder,
 						Path.GetFileNameWithoutExtension(UnrealBuildTool.GetUProjectFile()),
-						LinkEnvironment.Config.TargetConfiguration.ToString(),
+						LinkEnvironment.Config.Target.Configuration.ToString(),
 						Path.GetFileName(OutputFile.AbsolutePath) + ".response");
 				}
 			}
diff --git a/Engine/Source/Programs/UnrealBuildTool/WinRT/WinRTToolChain.cs b/Engine/Source/Programs/UnrealBuildTool/WinRT/WinRTToolChain.cs
--- a/Engine/Source/Programs/UnrealBuildTool/WinRT/WinRTToolChain.cs
+++ b/Engine/Source/Programs/UnrealBuildTool/WinRT/WinRTToolChain.cs
@@ -88,7 +88,7 @@ namespace UnrealBuildTool
 			//
 			//	Debug
 			//
-			if (CompileEnvironment.Config.TargetConfiguration == CPPTargetConfiguration.Debug)
+			if (CompileEnvironment.Config.Target.Configuration == CPPTargetConfiguration.Debug)
 			{
 				// Disable compiler optimization.
 				Result += " /Od";
@@ -125,7 +125,7 @@ namespace UnrealBuildTool
 				//
 				// LTCG
 				//
-				if (CompileEnvironment.Config.TargetConfiguration == CPPTargetConfiguration.Shipping)
+				if (CompileEnvironment.Config.Target.Configuration == CPPTargetConfiguration.Shipping)
 				{
 					if (BuildConfiguration.bAllowLTCG)
 					{
@@ -197,7 +197,7 @@ namespace UnrealBuildTool
 					// Create debug info suitable for E&C if wanted.
 					if (BuildConfiguration.bSupportEditAndContinue
 						// We only need to do this in debug as that's the only configuration that supports E&C.
-						&& CompileEnvironment.Config.TargetConfiguration == CPPTargetConfiguration.Debug)
+						&& CompileEnvironment.Config.Target.Configuration == CPPTargetConfiguration.Debug)
 					{
 						Result += " /ZI";
 					}
@@ -215,7 +215,7 @@ namespace UnrealBuildTool
 			}
 
 			// Specify the appropriate runtime library based on the platform and config.
-			if( CompileEnvironment.Config.TargetConfiguration == CPPTargetConfiguration.Debug && BuildConfiguration.bDebugBuildsActuallyUseDebugCRT )
+			if( CompileEnvironment.Config.Target.Configuration == CPPTargetConfiguration.Debug && BuildConfiguration.bDebugBuildsActuallyUseDebugCRT )
 			{
 				Result += " /MDd";
 			}
@@ -315,7 +315,7 @@ namespace UnrealBuildTool
 				Result += string.Format(" /WINMDFILE:\"{0}\"", Path.ChangeExtension(LinkEnvironment.Config.OutputFilePath, "winmd"));
 			}
 
-			if (LinkEnvironment.Config.TargetPlatform == CPPTargetPlatform.WinRT_ARM)
+			if (LinkEnvironment.Config.Target.Platform == CPPTargetPlatform.WinRT_ARM)
 			{
 				// Link for console.
 				Result += " /SUBSYSTEM:CONSOLE";
@@ -347,7 +347,7 @@ namespace UnrealBuildTool
 			//	ReleaseLTCG
 			//
 			if (BuildConfiguration.bAllowLTCG &&
-				LinkEnvironment.Config.TargetConfiguration == CPPTargetConfiguration.Shipping)
+				LinkEnvironment.Config.Target.Configuration == CPPTargetConfiguration.Shipping)
 			{
 				// Use link-time code generation.
 				Result += " /LTCG";
@@ -361,7 +361,7 @@ namespace UnrealBuildTool
 			//
 			//	Shipping binary
 			//
-			if (LinkEnvironment.Config.TargetConfiguration == CPPTargetConfiguration.Shipping)
+			if (LinkEnvironment.Config.Target.Configuration == CPPTargetConfiguration.Shipping)
 			{
 				// Generate an EXE checksum.
 				Result += " /RELEASE";
@@ -415,7 +415,7 @@ namespace UnrealBuildTool
 			//
 			//	Shipping & LTCG
 			//
-			if (LinkEnvironment.Config.TargetConfiguration == CPPTargetConfiguration.Shipping)
+			if (LinkEnvironment.Config.Target.Configuration == CPPTargetConfiguration.Shipping)
 			{
 				// Use link-time code generation.
 				Result += " /ltcg";
@@ -649,7 +649,7 @@ namespace UnrealBuildTool
 				}
 
 				CompileAction.WorkingDirectory = Path.GetFullPath(".");
-				CompileAction.CommandPath = GetVCToolPath(CompileEnvironment.Config.TargetPlatform, CompileEnvironment.Config.TargetConfiguration, "cl");
+				CompileAction.CommandPath = GetVCToolPath(CompileEnvironment.Config.Target.Platform, CompileEnvironment.Config.Target.Configuration, "cl");
 				CompileAction.bIsVCCompiler = true;
 				CompileAction.CommandArguments = Arguments + FileArguments + CompileEnvironment.Config.AdditionalArguments;
 				CompileAction.StatusDescription = string.Format("{0}", Path.GetFileName(SourceFile.AbsolutePath));
@@ -681,8 +681,8 @@ namespace UnrealBuildTool
 			Action LinkAction = new Action(ActionType.Link);
 			LinkAction.WorkingDirectory = Path.GetFullPath(".");
 			LinkAction.CommandPath = GetVCToolPath(
-				LinkEnvironment.Config.TargetPlatform,
-				LinkEnvironment.Config.TargetConfiguration,
+				LinkEnvironment.Config.Target.Platform,
+				LinkEnvironment.Config.Target.Configuration,
 				bIsBuildingLibrary ? "lib" : "link");
 
 			// Get link arguments.
@@ -804,7 +804,7 @@ namespace UnrealBuildTool
 			{
 				// Xbox 360 LTCG does not seem to produce those.
 				if (LinkEnvironment.Config.bHasExports &&
-					(LinkEnvironment.Config.TargetConfiguration != CPPTargetConfiguration.Shipping))
+					(LinkEnvironment.Config.Target.Configuration != CPPTargetConfiguration.Shipping))
 				{
 					// Write the import library to the output directory for nFringe support.
 					FileItem ImportLibraryFile = FileItem.GetItemByPath(ImportLibraryFilePath);
diff --git a/Engine/Source/Programs/UnrealBuildTool/Windows/UEBuildWindows.cs b/Engine/Source/Programs/UnrealBuildTool/Windows/UEBuildWindows.cs
--- a/Engine/Source/Programs/UnrealBuildTool/Windows/UEBuildWindows.cs
+++ b/Engine/Source/Programs/UnrealBuildTool/Windows/UEBuildWindows.cs
@@ -520,8 +520,8 @@ namespace UnrealBuildTool
 			}
 
 			// Set up the global C++ compilation and link environment.
-			InBuildTarget.GlobalCompileEnvironment.Config.TargetConfiguration = CompileConfiguration;
-			InBuildTarget.GlobalLinkEnvironment.Config.TargetConfiguration = CompileConfiguration;
+			InBuildTarget.GlobalCompileEnvironment.Config.Target.Configuration = CompileConfiguration;
+			InBuildTarget.GlobalLinkEnvironment.Config.Target.Configuration    = CompileConfiguration;
 
 			// Create debug info based on the heuristics specified by the user.
 			InBuildTarget.GlobalCompileEnvironment.Config.bCreateDebugInfo =
diff --git a/Engine/Source/Programs/UnrealBuildTool/Windows/VCToolChain.cs b/Engine/Source/Programs/UnrealBuildTool/Windows/VCToolChain.cs
--- a/Engine/Source/Programs/UnrealBuildTool/Windows/VCToolChain.cs
+++ b/Engine/Source/Programs/UnrealBuildTool/Windows/VCToolChain.cs
@@ -76,7 +76,7 @@ namespace UnrealBuildTool
 			
 			if( !WindowsPlatform.bCompileWithClang )	// @todo clang: These options are not supported with clang-cl yet
 			{
-				if (CompileEnvironment.Config.TargetPlatform == CPPTargetPlatform.Win64)
+				if (CompileEnvironment.Config.Target.Platform == CPPTargetPlatform.Win64)
 				{
 					// Pack struct members on 8-byte boundaries.
 					Result += " /Zp8";
@@ -154,7 +154,7 @@ namespace UnrealBuildTool
 			//
 			//	Debug
 			//
-			if (CompileEnvironment.Config.TargetConfiguration == CPPTargetConfiguration.Debug)
+			if (CompileEnvironment.Config.Target.Configuration == CPPTargetConfiguration.Debug)
 			{
 				// Disable compiler optimization.
 				Result += " /Od";
@@ -168,8 +168,8 @@ namespace UnrealBuildTool
 					Result += " /Ob2";
 				}
 
-				if ((CompileEnvironment.Config.TargetPlatform == CPPTargetPlatform.Win32) || 
-					(CompileEnvironment.Config.TargetPlatform == CPPTargetPlatform.Win64))
+				if ((CompileEnvironment.Config.Target.Platform == CPPTargetPlatform.Win32) || 
+					(CompileEnvironment.Config.Target.Platform == CPPTargetPlatform.Win64))
 				{
 					// Runtime stack checks are not allowed when compiling for CLR
 					if (CompileEnvironment.Config.CLRMode == CPPCLRMode.CLRDisabled)
@@ -194,8 +194,8 @@ namespace UnrealBuildTool
 
 				// Only omit frame pointers on the PC (which is implied by /Ox) if wanted.
 				if ( BuildConfiguration.bOmitFramePointers == false
-				&& ((CompileEnvironment.Config.TargetPlatform == CPPTargetPlatform.Win32) ||
-					(CompileEnvironment.Config.TargetPlatform == CPPTargetPlatform.Win64)))
+				&& ((CompileEnvironment.Config.Target.Platform == CPPTargetPlatform.Win32) ||
+					(CompileEnvironment.Config.Target.Platform == CPPTargetPlatform.Win64)))
 				{
 					Result += " /Oy-";
 				}
@@ -209,7 +209,7 @@ namespace UnrealBuildTool
 				//
 				// LTCG
 				//
-				if (CompileEnvironment.Config.TargetConfiguration == CPPTargetConfiguration.Shipping)
+				if (CompileEnvironment.Config.Target.Configuration == CPPTargetConfiguration.Shipping)
 				{
 					if( BuildConfiguration.bAllowLTCG )
 					{
@@ -222,13 +222,13 @@ namespace UnrealBuildTool
 			//
 			//	PC
 			//
-			if ((CompileEnvironment.Config.TargetPlatform == CPPTargetPlatform.Win32) || 
-				(CompileEnvironment.Config.TargetPlatform == CPPTargetPlatform.Win64))
+			if ((CompileEnvironment.Config.Target.Platform == CPPTargetPlatform.Win32) || 
+				(CompileEnvironment.Config.Target.Platform == CPPTargetPlatform.Win64))
 			{
 				// SSE options are not allowed when using CLR compilation or the 64 bit toolchain
 				// (both enable SSE2 automatically)
 				if (CompileEnvironment.Config.CLRMode == CPPCLRMode.CLRDisabled &&
-					CompileEnvironment.Config.TargetPlatform != CPPTargetPlatform.Win64)
+					CompileEnvironment.Config.Target.Platform != CPPTargetPlatform.Win64)
 				{
 					// @todo Clang: Doesn't make sense to the Clang compiler
 					if( !WindowsPlatform.bCompileWithClang )
@@ -264,7 +264,7 @@ namespace UnrealBuildTool
             }
 
             // @todo - find a better place for this. 
-            if (CompileEnvironment.Config.TargetPlatform == CPPTargetPlatform.HTML5)
+            if (CompileEnvironment.Config.Target.Platform == CPPTargetPlatform.HTML5)
             {
                 Result += " /EHsc"; 
             }
@@ -285,7 +285,7 @@ namespace UnrealBuildTool
 						// Create debug info suitable for E&C if wanted.
 						if( BuildConfiguration.bSupportEditAndContinue &&
 							// We only need to do this in debug as that's the only configuration that supports E&C.
-							CompileEnvironment.Config.TargetConfiguration == CPPTargetConfiguration.Debug)
+							CompileEnvironment.Config.Target.Configuration == CPPTargetConfiguration.Debug)
 						{
 							Result += " /ZI";
 						}
@@ -304,7 +304,7 @@ namespace UnrealBuildTool
 			}
 
 			// Specify the appropriate runtime library based on the platform and config.
-			if( CompileEnvironment.Config.TargetConfiguration == CPPTargetConfiguration.Debug && BuildConfiguration.bDebugBuildsActuallyUseDebugCRT )
+			if( CompileEnvironment.Config.Target.Configuration == CPPTargetConfiguration.Debug && BuildConfiguration.bDebugBuildsActuallyUseDebugCRT )
 			{
 				Result += " /MDd";
 			}
@@ -419,11 +419,11 @@ namespace UnrealBuildTool
 			//
 			//	PC
 			//
-			if ((LinkEnvironment.Config.TargetPlatform == CPPTargetPlatform.Win32) || 
-				(LinkEnvironment.Config.TargetPlatform == CPPTargetPlatform.Win64))
+			if ((LinkEnvironment.Config.Target.Platform == CPPTargetPlatform.Win32) || 
+				(LinkEnvironment.Config.Target.Platform == CPPTargetPlatform.Win64))
 			{
 				// Set machine type/ architecture to be 64 bit.
-				if (LinkEnvironment.Config.TargetPlatform == CPPTargetPlatform.Win64)
+				if (LinkEnvironment.Config.Target.Platform == CPPTargetPlatform.Win64)
 				{
 					Result += " /MACHINE:x64";
 				}
@@ -447,7 +447,7 @@ namespace UnrealBuildTool
 					// on that older platform.  The compiler defaults to version 6.0+.  We'll modify the SUBSYSTEM parameter here.
 					if( WindowsPlatform.SupportWindowsXP )
 					{
-						Result += LinkEnvironment.Config.TargetPlatform == CPPTargetPlatform.Win64 ? ",5.02" : ",5.01";
+						Result += LinkEnvironment.Config.Target.Platform == CPPTargetPlatform.Win64 ? ",5.02" : ",5.01";
 					}
 				}
 
@@ -461,7 +461,7 @@ namespace UnrealBuildTool
 				Result += " /FIXED:No";
 
 				// Option is only relevant with 32 bit toolchain.
-				if (LinkEnvironment.Config.TargetPlatform == CPPTargetPlatform.Win32)
+				if (LinkEnvironment.Config.Target.Platform == CPPTargetPlatform.Win32)
 				{
 					// Disables the 2GB address space limit on 64-bit Windows and 32-bit Windows with /3GB specified in boot.ini
 					Result += " /LARGEADDRESSAWARE";
@@ -475,7 +475,7 @@ namespace UnrealBuildTool
 
 				// E&C can't use /SAFESEH.  Also, /SAFESEH isn't compatible with 64-bit linking
 				if( !BuildConfiguration.bSupportEditAndContinue &&
-					LinkEnvironment.Config.TargetPlatform != CPPTargetPlatform.Win64)
+					LinkEnvironment.Config.Target.Platform != CPPTargetPlatform.Win64)
 				{
 					// Generates a table of Safe Exception Handlers.  Documentation isn't clear whether they actually mean
 					// Structured Exception Handlers.
@@ -513,7 +513,7 @@ namespace UnrealBuildTool
 			//	Shipping & LTCG
 			//
 			if( BuildConfiguration.bAllowLTCG &&
-				LinkEnvironment.Config.TargetConfiguration == CPPTargetConfiguration.Shipping)
+				LinkEnvironment.Config.Target.Configuration == CPPTargetConfiguration.Shipping)
 			{
 				// Use link-time code generation.
 				Result += " /LTCG";
@@ -526,7 +526,7 @@ namespace UnrealBuildTool
 			//
 			//	Shipping binary
 			//
-			if (LinkEnvironment.Config.TargetConfiguration == CPPTargetConfiguration.Shipping)
+			if (LinkEnvironment.Config.Target.Configuration == CPPTargetConfiguration.Shipping)
 			{
 				// Generate an EXE checksum.
 				Result += " /RELEASE";
@@ -588,10 +588,10 @@ namespace UnrealBuildTool
 			//
 			//	PC
 			//
-			if (LinkEnvironment.Config.TargetPlatform == CPPTargetPlatform.Win32 || LinkEnvironment.Config.TargetPlatform == CPPTargetPlatform.Win64)
+			if (LinkEnvironment.Config.Target.Platform == CPPTargetPlatform.Win32 || LinkEnvironment.Config.Target.Platform == CPPTargetPlatform.Win64)
 			{
 				// Set machine type/ architecture to be 64 bit.
-				if (LinkEnvironment.Config.TargetPlatform == CPPTargetPlatform.Win64)
+				if (LinkEnvironment.Config.Target.Platform == CPPTargetPlatform.Win64)
 				{
 					Result += " /MACHINE:x64";
 				}
@@ -615,7 +615,7 @@ namespace UnrealBuildTool
 					// on that older platform.  The compiler defaults to version 6.0+.  We'll modify the SUBSYSTEM parameter here.
 					if( WindowsPlatform.SupportWindowsXP )
 					{
-						Result += LinkEnvironment.Config.TargetPlatform == CPPTargetPlatform.Win64 ? ",5.02" : ",5.01";
+						Result += LinkEnvironment.Config.Target.Platform == CPPTargetPlatform.Win64 ? ",5.02" : ",5.01";
 					}
 				}
 			}
@@ -623,7 +623,7 @@ namespace UnrealBuildTool
 			//
 			//	Shipping & LTCG
 			//
-			if (LinkEnvironment.Config.TargetConfiguration == CPPTargetConfiguration.Shipping)
+			if (LinkEnvironment.Config.Target.Configuration == CPPTargetConfiguration.Shipping)
 			{
 				// Use link-time code generation.
 				Result += " /LTCG";
@@ -898,7 +898,7 @@ namespace UnrealBuildTool
 				}
 
 				CompileAction.WorkingDirectory = Path.GetFullPath(".");
-				CompileAction.CommandPath = GetVCToolPath(CompileEnvironment.Config.TargetPlatform, CompileEnvironment.Config.TargetConfiguration, "cl");
+				CompileAction.CommandPath = GetVCToolPath(CompileEnvironment.Config.Target.Platform, CompileEnvironment.Config.Target.Configuration, "cl");
 
 				if( !WindowsPlatform.bCompileWithClang )
 				{
@@ -952,7 +952,7 @@ namespace UnrealBuildTool
 			{
 				Action CompileAction = new Action(ActionType.Compile);
 				CompileAction.WorkingDirectory = Path.GetFullPath(".");
-				CompileAction.CommandPath = GetVCToolPath(Environment.Config.TargetPlatform, Environment.Config.TargetConfiguration, "rc");
+				CompileAction.CommandPath = GetVCToolPath(Environment.Config.Target.Platform, Environment.Config.Target.Configuration, "rc");
 				CompileAction.StatusDescription = Path.GetFileName(RCFile.AbsolutePath);
 
 				// Suppress header spew
@@ -960,7 +960,7 @@ namespace UnrealBuildTool
 
 				// If we're compiling for 64-bit Windows, also add the _WIN64 definition to the resource
 				// compiler so that we can switch on that in the .rc file using #ifdef.
-				if (Environment.Config.TargetPlatform == CPPTargetPlatform.Win64)
+				if (Environment.Config.Target.Platform == CPPTargetPlatform.Win64)
 				{
 					CompileAction.CommandArguments += " /D _WIN64";
 				}
@@ -1039,8 +1039,8 @@ namespace UnrealBuildTool
 			Action LinkAction = new Action(ActionType.Link);
 			LinkAction.WorkingDirectory = Path.GetFullPath(".");
 			LinkAction.CommandPath = GetVCToolPath(
-				LinkEnvironment.Config.TargetPlatform,
-				LinkEnvironment.Config.TargetConfiguration,
+				LinkEnvironment.Config.Target.Platform,
+				LinkEnvironment.Config.Target.Configuration,
 				bIsBuildingLibrary ? "lib" : "link");
 
 			// Get link arguments.
@@ -1176,7 +1176,7 @@ namespace UnrealBuildTool
 			{
 				// Xbox 360 LTCG does not seem to produce those.
 				if (LinkEnvironment.Config.bHasExports &&
-					(LinkEnvironment.Config.TargetConfiguration != CPPTargetConfiguration.Shipping))
+					(LinkEnvironment.Config.Target.Configuration != CPPTargetConfiguration.Shipping))
 				{
 					// Write the import library to the output directory for nFringe support.
 					FileItem ImportLibraryFile = FileItem.GetItemByPath(ImportLibraryFilePath);
