[GitHub] 205 : Prevent multi conversion
#ttp 337749

diff --git a/Engine/Source/Programs/UnrealBuildTool/Configuration/UEBuildTarget.cs b/Engine/Source/Programs/UnrealBuildTool/Configuration/UEBuildTarget.cs
--- a/Engine/Source/Programs/UnrealBuildTool/Configuration/UEBuildTarget.cs
+++ b/Engine/Source/Programs/UnrealBuildTool/Configuration/UEBuildTarget.cs
@@ -2304,11 +2304,12 @@ namespace UnrealBuildTool
 		{
 			UEBuildPlatform.GetBuildPlatform(Platform).SetUpConfigurationEnvironment(this);
 
+			var UpperExtension = Path.GetExtension (OutputPath).ToUpperInvariant (); 
 			// Check to see if we're compiling a library or not
-			bool bIsBuildingDLL = Path.GetExtension(OutputPath).ToUpperInvariant() == ".DLL";
+			bool bIsBuildingDLL = UpperExtension == ".DLL";
 			GlobalCompileEnvironment.Config.bIsBuildingDLL = bIsBuildingDLL;
 			GlobalLinkEnvironment.Config.bIsBuildingDLL = bIsBuildingDLL;
-			bool bIsBuildingLibrary = Path.GetExtension(OutputPath).ToUpperInvariant() == ".LIB";
+			bool bIsBuildingLibrary = UpperExtension == ".LIB";
 			GlobalCompileEnvironment.Config.bIsBuildingLibrary = bIsBuildingLibrary;
 			GlobalLinkEnvironment.Config.bIsBuildingLibrary = bIsBuildingLibrary;
 		}
