Fix use of -rpath so that we may link all our dylibs via [REDACTED_HANDLE]/<libname>.dylib and therefore make it easier to share and/or package .dylibs into Mac apps.

diff --git a/Engine/Source/Programs/UnrealBuildTool/Mac/MacToolChain.cs b/Engine/Source/Programs/UnrealBuildTool/Mac/MacToolChain.cs
--- a/Engine/Source/Programs/UnrealBuildTool/Mac/MacToolChain.cs
+++ b/Engine/Source/Programs/UnrealBuildTool/Mac/MacToolChain.cs
@@ -569,7 +569,7 @@ namespace UnrealBuildTool
 			// we gather the list of engine dylibs.
 			List<string> EngineAndGameLibraries = new List<string>();
 
-			string DylibsPath = "@executable_path";
+			string DylibsPath = "@rpath";
 
 			string AbsolutePath = OutputFile.AbsolutePath.Replace("\\", "/");
 			if (!AbsolutePath.Contains("/Engine/Binaries/Mac/"))
@@ -578,7 +578,7 @@ namespace UnrealBuildTool
 			}
 			if (!bIsBuildingLibrary)
 			{
-				LinkCommand += " -rpath @loader_path -rpath @executable_path";
+				LinkCommand += " -rpath @loader_path/ -rpath @executable_path/";
 			}
 
 			List<string> ThirdPartyLibraries = new List<string>();
