Removed bEnableInlining setting from module build rules
- This setting wasn't really used for anything; we always enable inlining for all C++ modules

diff --git a/Engine/Source/Programs/UnrealBuildTool/Configuration/UEBuildModule.cs b/Engine/Source/Programs/UnrealBuildTool/Configuration/UEBuildModule.cs
--- a/Engine/Source/Programs/UnrealBuildTool/Configuration/UEBuildModule.cs
+++ b/Engine/Source/Programs/UnrealBuildTool/Configuration/UEBuildModule.cs
@@ -871,9 +871,6 @@ namespace UnrealBuildTool
 		/** Enable exception handling */
 		public bool bEnableExceptions = false;
 
-		/** Enable inlining */
-		public bool bEnableInlining = true;
-
 		/** Path to this module's redist static library */
 		public string RedistStaticLibraryPath = null;
 
@@ -958,7 +955,6 @@ namespace UnrealBuildTool
 			bool InFasterWithoutUnity,
 			int InMinFilesUsingPrecompiledHeaderOverride,
 			bool InEnableExceptions,
-			bool InEnableInlining,
 			bool bInBuildSourceFiles
 			)
 			: base(	InTarget,
@@ -1006,7 +1002,6 @@ namespace UnrealBuildTool
 			bFasterWithoutUnity                    = InFasterWithoutUnity;
 			MinFilesUsingPrecompiledHeaderOverride = InMinFilesUsingPrecompiledHeaderOverride;
 			bEnableExceptions                      = InEnableExceptions;
-			bEnableInlining                        = InEnableInlining;
 		}
 
 		// UEBuildModule interface.
@@ -1567,7 +1562,6 @@ namespace UnrealBuildTool
 			Result.Config.bFasterWithoutUnity                    = bFasterWithoutUnity;
 			Result.Config.MinFilesUsingPrecompiledHeaderOverride = MinFilesUsingPrecompiledHeaderOverride;
 			Result.Config.bEnableExceptions                      = bEnableExceptions;
-			Result.Config.bEnableInlining                        = bEnableInlining;
 			Result.Config.OutputDirectory                        = Path.Combine(Binary.Config.IntermediateDirectory, Name);
 
 			// Switch the optimization flag if we're building a game module
@@ -1953,7 +1947,6 @@ namespace UnrealBuildTool
 			bool InFasterWithoutUnity,
 			int InMinFilesUsingPrecompiledHeaderOverride,
 			bool InEnableExceptions,
-			bool InEnableInlining,
 			bool bInBuildSourceFiles
 			)
 		: base(InTarget,InName,InType,InModuleDirectory,InOutputDirectory,InIntelliSenseGatherer,
@@ -1962,7 +1955,7 @@ namespace UnrealBuildTool
 			InPrivateIncludePaths,InPrivateIncludePathModuleNames,InPrivateDependencyModuleNames,
             InCircularlyReferencedDependentModules, InDynamicallyLoadedModuleNames, InPlatformSpecificDynamicallyLoadedModuleNames, InOptimizeCode,
 			InAllowSharedPCH, InSharedPCHHeaderFile, InUseRTTI, InEnableBufferSecurityChecks, InFasterWithoutUnity, InMinFilesUsingPrecompiledHeaderOverride,
-			InEnableExceptions, InEnableInlining, bInBuildSourceFiles)
+			InEnableExceptions, bInBuildSourceFiles)
 		{
 			PrivateAssemblyReferences = ListFromOptionalEnumerableStringParameter(InPrivateAssemblyReferences);
 		}
diff --git a/Engine/Source/Programs/UnrealBuildTool/Configuration/UEBuildTarget.cs b/Engine/Source/Programs/UnrealBuildTool/Configuration/UEBuildTarget.cs
--- a/Engine/Source/Programs/UnrealBuildTool/Configuration/UEBuildTarget.cs
+++ b/Engine/Source/Programs/UnrealBuildTool/Configuration/UEBuildTarget.cs
@@ -1583,7 +1583,6 @@ namespace UnrealBuildTool
 					InFasterWithoutUnity: true,
 					InMinFilesUsingPrecompiledHeaderOverride: 0,
 					InEnableExceptions: false,
-					InEnableInlining: true,
 					bInBuildSourceFiles: true
 					);
 
@@ -2524,7 +2523,6 @@ namespace UnrealBuildTool
 							InFasterWithoutUnity: RulesObject.bFasterWithoutUnity,
 							InMinFilesUsingPrecompiledHeaderOverride: RulesObject.MinFilesUsingPrecompiledHeaderOverride,
 							InEnableExceptions: RulesObject.bEnableExceptions,
-							InEnableInlining: RulesObject.bEnableInlining,
 							bInBuildSourceFiles: bBuildSourceFiles
 						);
 
@@ -2562,7 +2560,6 @@ namespace UnrealBuildTool
 							InFasterWithoutUnity: RulesObject.bFasterWithoutUnity,
 							InMinFilesUsingPrecompiledHeaderOverride: RulesObject.MinFilesUsingPrecompiledHeaderOverride,
 							InEnableExceptions: RulesObject.bEnableExceptions,
-							InEnableInlining: RulesObject.bEnableInlining,
 							bInBuildSourceFiles : bBuildSourceFiles
 						);
 
diff --git a/Engine/Source/Programs/UnrealBuildTool/System/CPPEnvironment.cs b/Engine/Source/Programs/UnrealBuildTool/System/CPPEnvironment.cs
--- a/Engine/Source/Programs/UnrealBuildTool/System/CPPEnvironment.cs
+++ b/Engine/Source/Programs/UnrealBuildTool/System/CPPEnvironment.cs
@@ -211,9 +211,6 @@ namespace UnrealBuildTool
 		/** Enable exception handling */
 		public bool bEnableExceptions = false;
 
-		/** Enable inlining */
-		public bool bEnableInlining = true;
-
 		/** True if the environment contains performance critical code. */
 		public ModuleRules.CodeOptimization OptimizeCode = ModuleRules.CodeOptimization.Default;
 
@@ -272,7 +269,6 @@ namespace UnrealBuildTool
 			bFasterWithoutUnity                    = InCopyEnvironment.bFasterWithoutUnity;
 			MinFilesUsingPrecompiledHeaderOverride = InCopyEnvironment.MinFilesUsingPrecompiledHeaderOverride;
 			bEnableExceptions                      = InCopyEnvironment.bEnableExceptions;
-			bEnableInlining                        = InCopyEnvironment.bEnableInlining;
 			OptimizeCode                           = InCopyEnvironment.OptimizeCode;
 			bCreateDebugInfo                       = InCopyEnvironment.bCreateDebugInfo;
 			bIsBuildingLibrary                     = InCopyEnvironment.bIsBuildingLibrary;
diff --git a/Engine/Source/Programs/UnrealBuildTool/System/RulesCompiler.cs b/Engine/Source/Programs/UnrealBuildTool/System/RulesCompiler.cs
--- a/Engine/Source/Programs/UnrealBuildTool/System/RulesCompiler.cs
+++ b/Engine/Source/Programs/UnrealBuildTool/System/RulesCompiler.cs
@@ -140,9 +140,6 @@ namespace UnrealBuildTool
 		/** Enable exception handling */
 		public bool bEnableExceptions = false;
 
-		/** Enable inlining */
-		public bool bEnableInlining = true;
-
 		/** If true and unity builds are enabled, this module will build without unity. */
 		public bool bFasterWithoutUnity = false;
 
diff --git a/Engine/Source/Programs/UnrealBuildTool/WinRT/WinRTToolChain.cs b/Engine/Source/Programs/UnrealBuildTool/WinRT/WinRTToolChain.cs
--- a/Engine/Source/Programs/UnrealBuildTool/WinRT/WinRTToolChain.cs
+++ b/Engine/Source/Programs/UnrealBuildTool/WinRT/WinRTToolChain.cs
@@ -96,8 +96,8 @@ namespace UnrealBuildTool
 				// Favor code size (especially useful for embedded platforms).
 				Result += " /Os";
 
-				// Allow inline method expansion unless E&C support is requested or inlining is explicitly disabled.
-				if (!BuildConfiguration.bSupportEditAndContinue && CompileEnvironment.Config.bEnableInlining)
+				// Allow inline method expansion unless E&C support is requested
+				if (!BuildConfiguration.bSupportEditAndContinue)
 				{
 					// Allow inline method expansion of any suitable.
 					Result += " /Ob2";
@@ -116,11 +116,8 @@ namespace UnrealBuildTool
 				Result += " /Ox";
 				// Enable intrinsics
 				Result += " /Oi";
-				// Allow inline method expansion of any suitable unless it's explicitly disabled.
-				if (CompileEnvironment.Config.bEnableInlining)
-				{
-					Result += " /Ob2";
-				}
+				// Allow inline method expansion
+				Result += " /Ob2";
 
 				//
 				// LTCG
diff --git a/Engine/Source/Programs/UnrealBuildTool/Windows/VCToolChain.cs b/Engine/Source/Programs/UnrealBuildTool/Windows/VCToolChain.cs
--- a/Engine/Source/Programs/UnrealBuildTool/Windows/VCToolChain.cs
+++ b/Engine/Source/Programs/UnrealBuildTool/Windows/VCToolChain.cs
@@ -162,8 +162,8 @@ namespace UnrealBuildTool
 				// Favor code size (especially useful for embedded platforms).
 				Result += " /Os";
 
-				// Allow inline method expansion unless E&C support is requested or inlining is explicitly disabled.
-				if( !BuildConfiguration.bSupportEditAndContinue && CompileEnvironment.Config.bEnableInlining )
+				// Allow inline method expansion unless E&C support is requested
+				if( !BuildConfiguration.bSupportEditAndContinue )
 				{
 					Result += " /Ob2";
 				}
@@ -200,11 +200,8 @@ namespace UnrealBuildTool
 					Result += " /Oy-";
 				}
 
-				// Allow inline method expansion unless it's explicitly disabled.
-				if( CompileEnvironment.Config.bEnableInlining )
-				{
-					Result += " /Ob2";
-				}
+				// Allow inline method expansion
+				Result += " /Ob2";
 
 				//
 				// LTCG
