Removed bAllowSharedPCH from ModuleRules (not used, was deprecated a long time ago.)

diff --git a/Engine/Source/Programs/UnrealBuildTool/System/RulesCompiler.cs b/Engine/Source/Programs/UnrealBuildTool/System/RulesCompiler.cs
--- a/Engine/Source/Programs/UnrealBuildTool/System/RulesCompiler.cs
+++ b/Engine/Source/Programs/UnrealBuildTool/System/RulesCompiler.cs
@@ -108,10 +108,6 @@ namespace UnrealBuildTool
 		/// When this module's code should be optimized.
 		public CodeOptimization OptimizeCode = CodeOptimization.Default;
 
-		// NOTE: bAllowSharedPCH is deprecated and was left for backwards compatiblity with existing *.Build.cs files.  In general, you should not need to use the replacement (PCHUsage), as the defaults are correct for most modules.
-		// [System.Obsolete("bAllowSharedPCH has been deprecated.  Please remove use of this variable.", false)]
-		public bool bAllowSharedPCH = true;	// <--- deprecated!
-
 		/// Header file name for a shared PCH provided by this module.  Must be a valid relative path to a public C++ header file.
 		/// This should only be set for header files that are included by a significant number of other C++ modules.
 		public string SharedPCHHeaderFile = String.Empty;
@@ -128,7 +124,7 @@ namespace UnrealBuildTool
 			UseSharedPCHs
 		}
 
-		/// Precompiled header usage for this module (replaces bAllowSharedPCH)
+		/// Precompiled header usage for this module
 		public PCHUsageMode PCHUsage = PCHUsageMode.Default;
 
 		/** Use run time type information */
