Allow optimized code to be debugged more easily on Windows
- Sets undocumented compiler command-line switch that enables debugging of optimized code (/d2Zi+)
- More info: http://randomascii.wordpress.com/2013/09/11/debugging-optimized-codenew-in-visual-studio-2012/

diff --git a/Engine/Source/Programs/UnrealBuildTool/Windows/VCToolChain.cs b/Engine/Source/Programs/UnrealBuildTool/Windows/VCToolChain.cs
--- a/Engine/Source/Programs/UnrealBuildTool/Windows/VCToolChain.cs
+++ b/Engine/Source/Programs/UnrealBuildTool/Windows/VCToolChain.cs
@@ -187,6 +187,10 @@ namespace UnrealBuildTool
 				if( CompileEnvironment.Config.OptimizeCode >= ModuleRules.CodeOptimization.InNonDebugBuilds )
 				{
 					Result += " /Ox";
+
+					// Allow optimized code to be debugged more easily.  This makes PDBs a bit larger, but doesn't noticeably affect
+					// compile times.  The executable code is not affected at all by this switch, only the debugging information.
+					Result += " /d2Zi+";
 				}
 				
 				// Favor code speed.
