Add support for compiling for size instead of speed to Android, IOS, and HTML5
Fix for Android using an old variable for architecture
#ue4

diff --git a/Engine/Source/Programs/UnrealBuildTool/Android/AndroidToolChain.cs b/Engine/Source/Programs/UnrealBuildTool/Android/AndroidToolChain.cs
--- a/Engine/Source/Programs/UnrealBuildTool/Android/AndroidToolChain.cs
+++ b/Engine/Source/Programs/UnrealBuildTool/Android/AndroidToolChain.cs
@@ -133,7 +133,7 @@ namespace UnrealBuildTool
 		{
 			string Result = "";
 			
-			Result += (CompileEnvironment.Config.TargetArchitecture == "-armv7") ? ToolchainParamsArm : ToolchainParamsx86;
+			Result += (CompileEnvironment.Config.Target.Architecture == "-armv7") ? ToolchainParamsArm : ToolchainParamsx86;
 
 			// build up the commandline common to C and C++
 			Result += " -c";
@@ -174,7 +174,14 @@ namespace UnrealBuildTool
 			}
 			else
 			{
-				Result += " -O3";
+				if (UEBuildConfiguration.bCompileForSize)
+				{
+					Result += " -Oz";
+				}
+				else
+				{
+					Result += " -O3";
+				}
 			}
 
 			//@todo android: these are copied verbatim from UE3 and probably need adjustment
diff --git a/Engine/Source/Programs/UnrealBuildTool/HTML5/HTML5ToolChain.cs b/Engine/Source/Programs/UnrealBuildTool/HTML5/HTML5ToolChain.cs
--- a/Engine/Source/Programs/UnrealBuildTool/HTML5/HTML5ToolChain.cs
+++ b/Engine/Source/Programs/UnrealBuildTool/HTML5/HTML5ToolChain.cs
@@ -75,7 +75,7 @@ namespace UnrealBuildTool
 
             // we have to specify the full amount of memory with Asm.JS (1.5 G)
             // I wonder if there's a per game way to change this. 
-			int TotalMemory = 512 * 1024 * 1024;
+			int TotalMemory = 256 * 1024 * 1024;
             Result += " -s TOTAL_MEMORY=" + TotalMemory.ToString();
 
             // no need for exceptions
diff --git a/Engine/Source/Programs/UnrealBuildTool/IOS/IOSToolChain.cs b/Engine/Source/Programs/UnrealBuildTool/IOS/IOSToolChain.cs
--- a/Engine/Source/Programs/UnrealBuildTool/IOS/IOSToolChain.cs
+++ b/Engine/Source/Programs/UnrealBuildTool/IOS/IOSToolChain.cs
@@ -248,7 +248,14 @@ namespace UnrealBuildTool
 			// Optimize non- debug builds.
 			if (CompileEnvironment.Config.Target.Configuration != CPPTargetConfiguration.Debug)
 			{
-				Result += " -O3";
+				if (UEBuildConfiguration.bCompileForSize)
+				{
+					Result += " -Oz";
+				}
+				else
+				{
+					Result += " -O3";
+				}
 			}
 			else
 			{
