Fix for Linux using an old variable for architecture
#ue4

diff --git a/Engine/Source/Programs/UnrealBuildTool/Linux/LinuxToolChain.cs b/Engine/Source/Programs/UnrealBuildTool/Linux/LinuxToolChain.cs
--- a/Engine/Source/Programs/UnrealBuildTool/Linux/LinuxToolChain.cs
+++ b/Engine/Source/Programs/UnrealBuildTool/Linux/LinuxToolChain.cs
@@ -192,7 +192,7 @@ namespace UnrealBuildTool.Linux
             {
                 if (UsingClang())
                 {
-                    Result += String.Format(" -target {0}", CompileEnvironment.Config.TargetArchitecture);        // Set target triple
+                    Result += String.Format(" -target {0}", CompileEnvironment.Config.Target.Architecture);        // Set target triple
                 }
                 Result += String.Format(" --sysroot={0}", BaseLinuxPath);
             }
@@ -281,7 +281,7 @@ namespace UnrealBuildTool.Linux
             {
                 if (UsingClang())
                 {
-                    Result += String.Format(" -target {0}", LinkEnvironment.Config.TargetArchitecture);        // Set target triple
+                    Result += String.Format(" -target {0}", LinkEnvironment.Config.Target.Architecture);        // Set target triple
                 }
                 string SysRootPath = BaseLinuxPath.TrimEnd(new char[] { '\\', '/' });
                 Result += String.Format(" \"--sysroot={0}\"", SysRootPath);
