fix for UBT not specifying failure when there is a compile error on Mac
#ue4
#mac

diff --git a/Engine/Source/Programs/UnrealBuildTool/System/UnrealBuildTool.cs b/Engine/Source/Programs/UnrealBuildTool/System/UnrealBuildTool.cs
--- a/Engine/Source/Programs/UnrealBuildTool/System/UnrealBuildTool.cs
+++ b/Engine/Source/Programs/UnrealBuildTool/System/UnrealBuildTool.cs
@@ -1169,6 +1169,10 @@ namespace UnrealBuildTool
 							ToolChain.PostBuildSync(Target);
 						}
 					}
+					else
+					{
+						BuildResult = ECompilationResult.OtherCompilationError;
+					}
 				}
 			}
 			catch (BuildException Exception)
