Plugin linker fix

#add Fake linker-only function generated in UHT for every Module.generated.cpp.
#change UELinkerFixups file is now processed *after* running UHT.
#change Fake function called from UELinkerFixups for every statically-linked dependent project.

diff --git a/Engine/Source/Programs/UnrealBuildTool/Configuration/UEBuildTarget.cs b/Engine/Source/Programs/UnrealBuildTool/Configuration/UEBuildTarget.cs
--- a/Engine/Source/Programs/UnrealBuildTool/Configuration/UEBuildTarget.cs
+++ b/Engine/Source/Programs/UnrealBuildTool/Configuration/UEBuildTarget.cs
@@ -1338,14 +1338,6 @@ namespace UnrealBuildTool
 				}
 			}
 
-			if (ShouldCompileMonolithic() && !ProjectFileGenerator.bGenerateProjectFiles && Rules != null && Rules.Type != TargetRules.TargetType.Program)
-			{
-				// All non-program monolithic binaries implicitly depend on all static plugin libraries so they are always linked appropriately
-				// In order to do this, we create a new module here with a cpp file we emit that invokes an empty function in each library.
-				// If we do not do this, there will be no static initialization for libs if no symbols are referenced in them.
-				CreateLinkerFixupsCPPFile();
-			}
-
 			// On Mac we have actions that should be executed after all the binaries are created
 			if (GlobalLinkEnvironment.Config.Target.Platform == CPPTargetPlatform.Mac)
 			{
@@ -1425,6 +1417,14 @@ namespace UnrealBuildTool
 				}
 			}
 
+			if (ShouldCompileMonolithic() && !ProjectFileGenerator.bGenerateProjectFiles && Rules != null && Rules.Type != TargetRules.TargetType.Program)
+			{
+				// All non-program monolithic binaries implicitly depend on all static plugin libraries so they are always linked appropriately
+				// In order to do this, we create a new module here with a cpp file we emit that invokes an empty function in each library.
+				// If we do not do this, there will be no static initialization for libs if no symbols are referenced in them.
+				CreateLinkerFixupsCPPFile();
+			}
+
 			// Build the target's binaries.
 			foreach (var Binary in AppBinaries)
 			{
@@ -1516,6 +1516,22 @@ namespace UnrealBuildTool
 				LinkerFixupsFileContents.Add("{");
 
 				// Fill out the body of the function with the empty function calls. This is what causes the static libraries to be considered relevant
+				{
+					var UObjectModules = new List<UEBuildModuleCPP>();
+					foreach (var Binary in AppBinaries)
+					{
+						var DependencyModules = Binary.GetAllDependencyModules(bIncludeDynamicallyLoaded: false, bForceCircular: false);
+						foreach (var Module in DependencyModules.OfType<UEBuildModuleCPP>().Where(CPPModule => CPPModule.AutoGenerateCppInfo != null && !UObjectModules.Any(Module => Module.Name == CPPModule.Name)))
+						{
+							UObjectModules.Add(Module);
+						}
+					}
+					foreach (var Module in UObjectModules)
+					{
+						LinkerFixupsFileContents.Add("    extern void EmptyLinkFunctionForGeneratedCode" + Module.Name + "();");
+						LinkerFixupsFileContents.Add("    EmptyLinkFunctionForGeneratedCode" + Module.Name + "();");
+					}
+				}
 				foreach (var DependencyModuleName in PrivateDependencyModuleNames)
 				{
 					LinkerFixupsFileContents.Add("    extern void EmptyLinkFunctionForStaticInitialization" + DependencyModuleName + "();");
@@ -1589,6 +1605,11 @@ namespace UnrealBuildTool
 				// Now bind this new module to the executable binary so it will link the plugin libs correctly
 				NewModule.Binary = ExecutableBinary;
 				NewModule.bIncludedInTarget = true;
+
+				// Process dependencies for this new module
+				NewModule.ProcessAllCppDependencies(NewModule.CreateModuleCompileEnvironment(GlobalCompileEnvironment));
+
+				// Add module to binary
 				ExecutableBinary.AddModule(NewModule.Name);
 			}
 		}
diff --git a/Engine/Source/Programs/UnrealHeaderTool/Private/CodeGenerator.cpp b/Engine/Source/Programs/UnrealHeaderTool/Private/CodeGenerator.cpp
--- a/Engine/Source/Programs/UnrealHeaderTool/Private/CodeGenerator.cpp
+++ b/Engine/Source/Programs/UnrealHeaderTool/Private/CodeGenerator.cpp
@@ -4231,6 +4231,7 @@ void FNativeClassHeaderGenerator::ExportGeneratedCPP()
 	FStringOutputDevice GeneratedCPPEpilogue;
 	FStringOutputDevice GeneratedCPPText;
 	FStringOutputDevice GeneratedINLText;
+	FStringOutputDevice GeneratedLinkerFixupFunction;
 
 	TArray<FStringOutputDevice> GeneratedCPPFiles;
 
@@ -4275,6 +4276,8 @@ void FNativeClassHeaderGenerator::ExportGeneratedCPP()
 	// Write out our include to the .dep.h file
 	GeneratedCPPClassesIncludes.Logf(TEXT("#include \"%s\"") LINE_TERMINATOR, *FPaths::GetCleanFilename(DepHeaderPathname));
 
+	GeneratedLinkerFixupFunction.Logf(TEXT("void EmptyLinkFunctionForGeneratedCode%s() {}") LINE_TERMINATOR, *ModuleInfo->Name);
+
 	GeneratedCPPText.Log(*GeneratedPackageCPP);
 	// Add all names marked for export to a list (for sorting)
 	TArray<FString> Names;
@@ -4328,7 +4331,7 @@ void FNativeClassHeaderGenerator::ExportGeneratedCPP()
 		}
 
 		FString CppPath = ModuleInfo->GeneratedCPPFilenameBase + (GeneratedFunctionBodyTextSplit.Num() > 1 ? *FString::Printf(TEXT(".%d.cpp"), FileIdx + 1) : TEXT(".cpp"));
-		SaveHeaderIfChanged(*CppPath, *(GeneratedCPPPreamble + ModulePCHInclude + GeneratedCPPClassesIncludes + FileText + GeneratedCPPEpilogue));
+		SaveHeaderIfChanged(*CppPath, *(GeneratedCPPPreamble + ModulePCHInclude + GeneratedCPPClassesIncludes + ((GeneratedFunctionBodyTextSplit.Num() > 1) ? FString() : GeneratedLinkerFixupFunction) + FileText + GeneratedCPPEpilogue));
 
 		if (GeneratedFunctionBodyTextSplit.Num() > 1)
 		{
@@ -4346,7 +4349,7 @@ void FNativeClassHeaderGenerator::ExportGeneratedCPP()
 		}
 
 		FString CppPath = ModuleInfo->GeneratedCPPFilenameBase + TEXT(".cpp");
-		SaveHeaderIfChanged(*CppPath, *(GeneratedCPPPreamble + ModulePCHInclude + FileText + GeneratedCPPEpilogue));
+		SaveHeaderIfChanged(*CppPath, *(GeneratedCPPPreamble + ModulePCHInclude + GeneratedLinkerFixupFunction + FileText + GeneratedCPPEpilogue));
 	}
 
 	FString InlPath = ModuleInfo->GeneratedCPPFilenameBase + TEXT(".inl");
