fix for not finding device when using launch on
TTP337507
#ue4
#ios

diff --git a/Engine/Source/Programs/UnrealBuildTool/System/XcodeProjectFileGenerator.cs b/Engine/Source/Programs/UnrealBuildTool/System/XcodeProjectFileGenerator.cs
--- a/Engine/Source/Programs/UnrealBuildTool/System/XcodeProjectFileGenerator.cs
+++ b/Engine/Source/Programs/UnrealBuildTool/System/XcodeProjectFileGenerator.cs
@@ -1303,7 +1303,7 @@ namespace UnrealBuildTool
 			// write scheme files for targets
 			foreach (var Target in ProjectTargets)
 			{
-				if (Target.Type == XcodeTargetType.Project || Target.Type == XcodeTargetType.XcodeHelper)
+				if (Target.Type == XcodeTargetType.Project)
 				{
 					continue;
 				}
