Xcode indexing improvements:
- disabled use of headermaps, so it can handle multiple headers with the same name (in different folders) properly
- enabled use of UnrealEd.h as a precompiled header, which makes indexing faster

diff --git a/Engine/Source/Editor/UnrealEd/Public/UnrealEd.h b/Engine/Source/Editor/UnrealEd/Public/UnrealEd.h
--- a/Engine/Source/Editor/UnrealEd/Public/UnrealEd.h
+++ b/Engine/Source/Editor/UnrealEd/Public/UnrealEd.h
@@ -3,6 +3,8 @@
 #ifndef __UnrealEd_h__
 #define __UnrealEd_h__
 
+#ifdef __cplusplus // Xcode needs that to use this file as a precompiled header for indexing
+
 #if WITH_EDITOR
 
 #include "Engine.h"
@@ -53,4 +55,6 @@ UNREALED_API void EditorExit();
 
 #include "UnrealEdMessages.h"
 
+#endif // __cplusplus
+
 #endif	// __UnrealEd_h__
diff --git a/Engine/Source/Programs/UnrealBuildTool/Mac/UEBuildMac.cs b/Engine/Source/Programs/UnrealBuildTool/Mac/UEBuildMac.cs
--- a/Engine/Source/Programs/UnrealBuildTool/Mac/UEBuildMac.cs
+++ b/Engine/Source/Programs/UnrealBuildTool/Mac/UEBuildMac.cs
@@ -166,6 +166,10 @@ namespace UnrealBuildTool
 
         public override void ValidateUEBuildConfiguration()
         {
+			if (ProjectFileGenerator.bGenerateProjectFiles || ProjectFileGenerator.bGeneratingRocketProjectFiles)
+			{
+				UEBuildConfiguration.bCompileLeanAndMeanUE = false;
+			}
         }
 
 		/**
diff --git a/Engine/Source/Programs/UnrealBuildTool/System/XcodeProject.cs b/Engine/Source/Programs/UnrealBuildTool/System/XcodeProject.cs
--- a/Engine/Source/Programs/UnrealBuildTool/System/XcodeProject.cs
+++ b/Engine/Source/Programs/UnrealBuildTool/System/XcodeProject.cs
@@ -262,15 +262,6 @@ namespace UnrealBuildTool
 			{
 				return null;
 			}
-			else if (InitFilePath.EndsWith("Classes/Engine/Engine.h") || InitFilePath.EndsWith("Classes/VectorField/VectorField.h") 
-				|| InitFilePath.EndsWith("Classes/Engine/Canvas.h") || InitFilePath.EndsWith("Classes/Intrinsic/Model.h") || InitFilePath.EndsWith("Classes/Engine/Texture.h")
-				|| InitFilePath.EndsWith("Windows/OnlineFriendsFacebook.h") || InitFilePath.EndsWith("Windows/OnlineIdentityFacebook.h") || InitFilePath.EndsWith("HTML5/Simulator/SocketSubsystem.h")
-				|| InitFilePath.EndsWith("HTML5/Device/SocketSubsystem.h") || InitFilePath.EndsWith("ContentBrowser/Private/HistoryManager.h") || InitFilePath.EndsWith("Private/Menus/TranslationEditorMenu.h"))
-			{
-				// @todo: Get rid of this workaround when possible.
-				// Xcode has a bug (radr://15660224) that makes indexing fail if there are multiple header files with the same name, so we skip few files that cause the problem.
-				return null;
-			}
 			return new XcodeSourceFile(InitFilePath, InitProjectSubFolder);
 		}
 
diff --git a/Engine/Source/Programs/UnrealBuildTool/System/XcodeProjectFileGenerator.cs b/Engine/Source/Programs/UnrealBuildTool/System/XcodeProjectFileGenerator.cs
--- a/Engine/Source/Programs/UnrealBuildTool/System/XcodeProjectFileGenerator.cs
+++ b/Engine/Source/Programs/UnrealBuildTool/System/XcodeProjectFileGenerator.cs
@@ -444,6 +444,7 @@ namespace UnrealBuildTool
 				"\t\t\t\tGCC_WARN_CHECK_SWITCH_STATEMENTS = NO;" + ProjectFileGenerator.NewLine +
 				"\t\t\t\tSUPPORTED_PLATFORMS = \"macosx\";" + ProjectFileGenerator.NewLine +
 				"\t\t\t\tONLY_ACTIVE_ARCH = YES;" + ProjectFileGenerator.NewLine +
+				"\t\t\t\tUSE_HEADERMAP = NO;" + ProjectFileGenerator.NewLine +
 				"\t\t\t\tSDKROOT = macosx;" + ProjectFileGenerator.NewLine +
 				"\t\t\t\tSYMROOT = " + EngineSubdir + "Intermediate/Build;" + ProjectFileGenerator.NewLine +
 				"\t\t\t};" + ProjectFileGenerator.NewLine +
@@ -453,6 +454,7 @@ namespace UnrealBuildTool
 
 		private void AppendMacBuildConfig(ref StringBuilder Contents, string ConfigName, string ConfigGuid, bool bIsMacOnly)
 		{
+			string UE4Dir = Path.GetFullPath(Directory.GetCurrentDirectory() + "../../..");
 			if (bIsMacOnly)
 			{
 				Contents.Append(
@@ -465,6 +467,8 @@ namespace UnrealBuildTool
 					"\t\t\t\tPRODUCT_NAME = \"$(TARGET_NAME)\";" + ProjectFileGenerator.NewLine +
 					"\t\t\t\tMACOSX_DEPLOYMENT_TARGET = 10.9;" + ProjectFileGenerator.NewLine +
 					"\t\t\t\tSDKROOT = macosx;" + ProjectFileGenerator.NewLine +
+					"\t\t\t\tGCC_PRECOMPILE_PREFIX_HEADER = YES;" + ProjectFileGenerator.NewLine +
+					"\t\t\t\tGCC_PREFIX_HEADER = \"" + UE4Dir + "/Engine/Source/Editor/UnrealEd/Public/UnrealEd.h\";" + ProjectFileGenerator.NewLine +
 					"\t\t\t};" + ProjectFileGenerator.NewLine +
 					"\t\t\tname = " + ConfigName + ";" + ProjectFileGenerator.NewLine +
 					"\t\t};" + ProjectFileGenerator.NewLine);
@@ -489,6 +493,8 @@ namespace UnrealBuildTool
 					"\t\t\t\tSDKROOT = macosx;" + ProjectFileGenerator.NewLine +
 					"\t\t\t\t\"SDKROOT[arch=x86_64]\" = macosx;" + ProjectFileGenerator.NewLine +
 					"\t\t\t\t\"SDKROOT[arch=arm*]\" = iphoneos;" + ProjectFileGenerator.NewLine +
+					"\t\t\t\tGCC_PRECOMPILE_PREFIX_HEADER = YES;" + ProjectFileGenerator.NewLine +
+					"\t\t\t\tGCC_PREFIX_HEADER = \"" + UE4Dir + "/Engine/Source/Editor/UnrealEd/Public/UnrealEd.h\";" + ProjectFileGenerator.NewLine +
 					"\t\t\t};" + ProjectFileGenerator.NewLine +
 					"\t\t\tname = " + ConfigName + ";" + ProjectFileGenerator.NewLine +
 					"\t\t};" + ProjectFileGenerator.NewLine);
