#HTML5

   Fix HTML5 Build for CL 2095675  (GitHub pull request #194: Full binary type. Removing hardcoded binary file extensions)

diff --git a/Engine/Source/Programs/UnrealBuildTool/HTML5/UEBuildHTML5.cs b/Engine/Source/Programs/UnrealBuildTool/HTML5/UEBuildHTML5.cs
--- a/Engine/Source/Programs/UnrealBuildTool/HTML5/UEBuildHTML5.cs
+++ b/Engine/Source/Programs/UnrealBuildTool/HTML5/UEBuildHTML5.cs
@@ -142,15 +142,22 @@ namespace UnrealBuildTool
 			}
 			else
 			{
-                if (InBinaryType == UEBuildBinaryType.StaticLibrary)
+                switch (InBinaryType)
                 {
-                    return ".bc";
+                    case UEBuildBinaryType.DynamicLinkLibrary:
+                        return ".js";
+                    case UEBuildBinaryType.Executable:
+                        return ".js";
+                    case UEBuildBinaryType.StaticLibrary:
+                        return ".bc";
+                    case UEBuildBinaryType.Object:
+                        return ".bc";
+                    case UEBuildBinaryType.PrecompiledHeader:
+                        return ".gch";
                 }
-                else 
-                {
-				    return ".js";
-                }
-			}
+
+                return base.GetBinaryExtension(InBinaryType);
+            }
 		}
 
 		/**
