build script, disable client server tests for now

diff --git a/Engine/Source/Programs/UnrealBuildTool/System/RulesCompiler.cs b/Engine/Source/Programs/UnrealBuildTool/System/RulesCompiler.cs
--- a/Engine/Source/Programs/UnrealBuildTool/System/RulesCompiler.cs
+++ b/Engine/Source/Programs/UnrealBuildTool/System/RulesCompiler.cs
@@ -755,12 +755,14 @@ namespace UnrealBuildTool
         public virtual Dictionary<string, string> GUBP_GetClientServerTests_MonolithicOnly(UnrealTargetPlatform HostPlatform, UnrealTargetPlatform AltHostPlatform, UnrealTargetPlatform ServerPlatform, UnrealTargetPlatform ClientPlatform)
         {
             var Result = new Dictionary<string, string>();
+#if false // needs work
             if ((ServerPlatform == HostPlatform || ServerPlatform == AltHostPlatform) &&
                 (ClientPlatform == HostPlatform || ClientPlatform == AltHostPlatform) && 
                 Type == TargetType.Game)  // for now, we will only run these for the dev config of the host platform and only the game executable, not sure how to deal with a client only executable
             {
                 Result.Add("CookedNetTest", "BuildCookRun -run -skipcook -stage -pak -deploy -unattended -server -nullrhi -NoP4  -addcmdline=\"-nosteam\" -platform=" + ClientPlatform.ToString() + " -serverplatform=" + ServerPlatform.ToString());
             }
+#endif
             return Result;
         }
 	}
