Changed the running directory for clang on android so that we can use relative path names for library files.
Gets around issue with too many commandline parameters going to the linker :(.

diff --git a/Engine/Source/Programs/UnrealBuildTool/Android/AndroidToolChain.cs b/Engine/Source/Programs/UnrealBuildTool/Android/AndroidToolChain.cs
--- a/Engine/Source/Programs/UnrealBuildTool/Android/AndroidToolChain.cs
+++ b/Engine/Source/Programs/UnrealBuildTool/Android/AndroidToolChain.cs
@@ -117,10 +117,10 @@ namespace UnrealBuildTool
 
 
 			// toolchain params
-			ToolchainParamsArm = "-target armv7-none-linux-androideabi" +
-								   " --sysroot=\"" + Path.Combine(NDKPath, "platforms", GetNdkApiLevel(), "arch-arm") + "\"" + 
+			ToolchainParamsArm = " -target armv7-none-linux-androideabi" +
+								   " --sysroot=\"" + Path.Combine(NDKPath, "platforms", GetNdkApiLevel(), "arch-arm") + "\"" +
                                    " -gcc-toolchain \"" + Path.Combine(NDKPath, @"toolchains\arm-linux-androideabi-" + GccVersion, ArchitecturePath) + "\"";
-			ToolchainParamsx86 = "-target i686-none-linux-android" +
+			ToolchainParamsx86 = " -target i686-none-linux-android" +
 								   " --sysroot=\"" + Path.Combine(NDKPath, "platforms", GetNdkApiLevel(), "arch-x86") + "\"" +
                                    " -gcc-toolchain \"" + Path.Combine(NDKPath, @"toolchains\x86-" + GccVersion, ArchitecturePath) + "\"";
 
@@ -289,7 +289,7 @@ namespace UnrealBuildTool
 		static string GetLinkArguments(LinkEnvironment LinkEnvironment)
 		{
 			string Result = "";
-			
+
 			Result += (LinkEnvironment.Config.Target.Architecture == "-armv7") ? ToolchainParamsArm : ToolchainParamsx86;
 
 			Result += " -nostdlib";
@@ -306,6 +306,9 @@ namespace UnrealBuildTool
 				Result += " -march=atom";
 			}
 
+            // verbose output from the linker
+            // Result += " -v";
+
 			return Result;
 		}
 
@@ -563,6 +566,10 @@ namespace UnrealBuildTool
 				LinkAction.CommandPath = ClangPath;
 			}
 
+            string LinkerPath = LinkAction.WorkingDirectory;
+
+            LinkAction.WorkingDirectory = LinkEnvironment.Config.IntermediateDirectory;
+
 			// Get link arguments.
 			LinkAction.CommandArguments = LinkEnvironment.Config.bIsBuildingLibrary ? GetArArguments(LinkEnvironment) : GetLinkArguments(LinkEnvironment);
 
@@ -571,6 +578,8 @@ namespace UnrealBuildTool
 			LinkAction.ProducedItems.Add(OutputFile);
 			LinkAction.StatusDescription = string.Format("{0}", Path.GetFileName(OutputFile.AbsolutePath));
 
+            // LinkAction.bPrintDebugInfo = true;
+
 			// Add the output file to the command-line.
 			if (LinkEnvironment.Config.bIsBuildingLibrary)
 			{
@@ -585,7 +594,12 @@ namespace UnrealBuildTool
 			List<string> InputFileNames = new List<string>();
 			foreach (FileItem InputFile in LinkEnvironment.InputFiles)
 			{
-				InputFileNames.Add(string.Format("\"{0}\"", InputFile.AbsolutePath.Replace("\\", "/")));
+                string AbsolutePath = InputFile.AbsolutePath.Replace("\\", "/");
+
+                AbsolutePath = AbsolutePath.Replace(LinkEnvironment.Config.IntermediateDirectory.Replace("\\", "/"), "");
+                AbsolutePath = AbsolutePath.TrimStart(new char[] {'/'} );
+
+                InputFileNames.Add(string.Format("\"{0}\"", AbsolutePath));
 				LinkAction.PrerequisiteItems.Add(InputFile);
 			}
 
@@ -598,7 +612,14 @@ namespace UnrealBuildTool
 				// Add the library paths to the argument list.
 				foreach (string LibraryPath in LinkEnvironment.Config.LibraryPaths)
 				{
-					LinkAction.CommandArguments += string.Format(" -L\"{0}\"", LibraryPath);
+                    // LinkerPaths could be relative or absolute
+                    string AbsoluteLibraryPath = ActionThread.ExpandEnvironmentVariables(LibraryPath);
+                    // environment variables aren't expanded when using the $( style
+                    if (Path.IsPathRooted( AbsoluteLibraryPath )==false) 
+                    {
+                        AbsoluteLibraryPath = Path.Combine(LinkerPath, AbsoluteLibraryPath);
+                    }
+					LinkAction.CommandArguments += string.Format(" -L\"{0}\"", AbsoluteLibraryPath);
 				}
 
 				// add libraries in a library group
@@ -618,7 +639,7 @@ namespace UnrealBuildTool
 				}
 				LinkAction.CommandArguments += string.Format(" -Wl,--end-group");
 			}
-
+            
 			// Add the additional arguments specified by the environment.
 			LinkAction.CommandArguments += LinkEnvironment.Config.AdditionalArguments;
 			LinkAction.CommandArguments = LinkAction.CommandArguments.Replace("\\", "/");
@@ -631,7 +652,7 @@ namespace UnrealBuildTool
 
 		public override void CompileCSharpProject(CSharpEnvironment CompileEnvironment, string ProjectFileName, string DestinationFile)
 		{
-			throw new BuildException("Android cannot compile C# files");
+            throw new BuildException("Android cannot compile C# files");
 		}
 
 		public static void OutputReceivedDataEventHandler(Object Sender, DataReceivedEventArgs Line)
diff --git a/Engine/Source/Programs/UnrealBuildTool/System/ActionGraph.cs b/Engine/Source/Programs/UnrealBuildTool/System/ActionGraph.cs
--- a/Engine/Source/Programs/UnrealBuildTool/System/ActionGraph.cs
+++ b/Engine/Source/Programs/UnrealBuildTool/System/ActionGraph.cs
@@ -36,7 +36,8 @@ namespace UnrealBuildTool
 		/** Relative cost of producing items for this action. */
 		public long RelativeCost = 0;
 		public string WorkingDirectory = null;
-		public string CommandPath = null;
+        public bool bPrintDebugInfo = false;
+        public string CommandPath = null;
 		public string CommandArguments = null;
 		public string CommandDescription = null;
 		public string StatusDescription = "...";
diff --git a/Engine/Source/Programs/UnrealBuildTool/System/LocalExecutor.cs b/Engine/Source/Programs/UnrealBuildTool/System/LocalExecutor.cs
--- a/Engine/Source/Programs/UnrealBuildTool/System/LocalExecutor.cs
+++ b/Engine/Source/Programs/UnrealBuildTool/System/LocalExecutor.cs
@@ -90,6 +90,25 @@ namespace UnrealBuildTool
 			}
 		}
 
+
+        /**
+         * Used when debuging Actions outputs all action return values to debug out
+         * 
+         * @param	sender		Sending object
+         * @param	e			Event arguments (In this case, the line of string output)
+         */
+        protected void ActionDebugOutput(object sender, DataReceivedEventArgs e)
+        {
+            var Output = e.Data;
+            if (Output == null)
+            {
+                return;
+            }
+
+            Log.TraceInformation(Output);
+        }
+
+
 		/**
 		 * The actual function to run in a thread. This is potentially long and blocking
 		 */
@@ -143,6 +162,10 @@ namespace UnrealBuildTool
 				{
 					Log.TraceVerbose("Executing: {0} {1}", ExpandedCommandPath, ActionStartInfo.Arguments);
 				}
+                if (Action.bPrintDebugInfo)
+                {
+                    Log.TraceInformation("Executing: {0} {1}", ExpandedCommandPath, ActionStartInfo.Arguments);
+                }
 				// Log summary if wanted.
 				else if (Action.bShouldOutputStatusDescription)
 				{
@@ -165,14 +188,24 @@ namespace UnrealBuildTool
 					{
 						ActionProcess = new Process();
 						ActionProcess.StartInfo = ActionStartInfo;
-						bool bShouldRedirectOuput = Action.OutputEventHandler != null;
+						bool bShouldRedirectOuput = Action.OutputEventHandler != null || Action.bPrintDebugInfo;
 						if (bShouldRedirectOuput)
 						{
 							ActionStartInfo.RedirectStandardOutput = true;
 							ActionStartInfo.RedirectStandardError = true;
 							ActionProcess.EnableRaisingEvents = true;
-							ActionProcess.OutputDataReceived += Action.OutputEventHandler;
-							ActionProcess.ErrorDataReceived += Action.OutputEventHandler;
+
+                            if (Action.OutputEventHandler != null)
+                            {
+                                ActionProcess.OutputDataReceived += Action.OutputEventHandler;
+                                ActionProcess.ErrorDataReceived += Action.OutputEventHandler;
+                            }
+                            
+                            if ( Action.bPrintDebugInfo)
+                            {
+                                ActionProcess.OutputDataReceived += new DataReceivedEventHandler(ActionDebugOutput);
+                                ActionProcess.ErrorDataReceived += new DataReceivedEventHandler(ActionDebugOutput);
+                            }
 						}
 						ActionProcess.Start();
 						if (bShouldRedirectOuput)
