[GitHub] 219 : Move post build handling to ToolChain
#ttp 338156

diff --git a/Engine/Source/Programs/UnrealBuildTool/Android/UEBuildAndroid.cs b/Engine/Source/Programs/UnrealBuildTool/Android/UEBuildAndroid.cs
--- a/Engine/Source/Programs/UnrealBuildTool/Android/UEBuildAndroid.cs
+++ b/Engine/Source/Programs/UnrealBuildTool/Android/UEBuildAndroid.cs
@@ -349,7 +349,7 @@ namespace UnrealBuildTool
 		/// <param name="CompileEnvironment">The environment to compile the binary in</param>
 		/// <param name="LinkEnvironment">The environment to link the binary in</param>
 		/// <returns></returns>
-		public override IEnumerable<FileItem> Build(CPPEnvironment CompileEnvironment, LinkEnvironment LinkEnvironment)
+		public override IEnumerable<FileItem> Build(IUEToolChain ToolChain, CPPEnvironment CompileEnvironment, LinkEnvironment LinkEnvironment)
 		{
 			// generate the .apk!
 			return new FileItem[] { FileItem.GetItemByPath(Config.OutputFilePath) };
diff --git a/Engine/Source/Programs/UnrealBuildTool/Configuration/UEBuildBinary.cs b/Engine/Source/Programs/UnrealBuildTool/Configuration/UEBuildBinary.cs
--- a/Engine/Source/Programs/UnrealBuildTool/Configuration/UEBuildBinary.cs
+++ b/Engine/Source/Programs/UnrealBuildTool/Configuration/UEBuildBinary.cs
@@ -185,10 +185,11 @@ namespace UnrealBuildTool
 		/// <summary>
 		/// Builds the binary.
 		/// </summary>
+		/// <param name="ToolChain">The toolchain which to use for building</param>
 		/// <param name="CompileEnvironment">The environment to compile the binary in</param>
 		/// <param name="LinkEnvironment">The environment to link the binary in</param>
 		/// <returns></returns>
-		public abstract IEnumerable<FileItem> Build(CPPEnvironment CompileEnvironment,LinkEnvironment LinkEnvironment);
+		public abstract IEnumerable<FileItem> Build(IUEToolChain ToolChain, CPPEnvironment CompileEnvironment,LinkEnvironment LinkEnvironment);
 
 		/// <summary>
 		/// Writes an XML summary of the build environment for this binary
@@ -448,8 +449,9 @@ namespace UnrealBuildTool
 		/// <param name="CompileEnvironment">The environment to compile the binary in</param>
 		/// <param name="LinkEnvironment">The environment to link the binary in</param>
 		/// <returns></returns>
-		public override IEnumerable<FileItem> Build(CPPEnvironment CompileEnvironment, LinkEnvironment LinkEnvironment)
+		public override IEnumerable<FileItem> Build(IUEToolChain TargetToolChain, CPPEnvironment CompileEnvironment, LinkEnvironment LinkEnvironment)
 		{
+
 			// Determine the type of binary we're linking.
 			switch (Config.Type)
 			{
@@ -467,7 +469,6 @@ namespace UnrealBuildTool
 					break;
 			};
 
-
 			var OutputFiles = new List<FileItem>();
 
 			var BinaryCompileEnvironment = CompileEnvironment.DeepCopy();
@@ -531,7 +532,6 @@ namespace UnrealBuildTool
 				BinaryDependency.SetupDependentLinkEnvironment(ref BinaryLinkEnvironment);
 			}
 
-
 			// Set the link output file.
 			BinaryLinkEnvironment.Config.OutputFilePath = Config.OutputFilePath;
 
@@ -606,21 +606,7 @@ namespace UnrealBuildTool
 					OutputFiles.Add(ConsoleAppLinkEvironment.LinkExecutable(false));
 				}
 
-				// if building for Mac on a Mac, use actions to finalize the builds (otherwise, we use Deploy)
-				if (ExternalExecution.GetRuntimePlatform() == UnrealTargetPlatform.Mac &&
-					BinaryLinkEnvironment.Config.Target.Platform == CPPTargetPlatform.Mac)
-				{
-					if (!BinaryLinkEnvironment.Config.bIsBuildingDLL && !BinaryLinkEnvironment.Config.bIsBuildingLibrary)
-					{
-						MacToolChain Mac = (MacToolChain)UEToolChain.GetPlatformToolChain(CPPTargetPlatform.Mac);
-						FileItem FixDylibOutputFile = Mac.FixDylibDependencies(BinaryLinkEnvironment, Executable);
-						OutputFiles.Add(FixDylibOutputFile);
-						if (!BinaryLinkEnvironment.Config.bIsBuildingConsoleApplication)
-						{
-							OutputFiles.Add(Mac.CreateAppBundle(BinaryLinkEnvironment, Executable, FixDylibOutputFile));
-						}
-					}
-				}
+				OutputFiles.AddRange(TargetToolChain.PostBuild(Executable, BinaryLinkEnvironment));
 			}
 			
 			return OutputFiles;
@@ -718,10 +704,11 @@ namespace UnrealBuildTool
 		/// <summary>
 		/// Builds the binary.
 		/// </summary>
+		/// <param name="ToolChain">The toolchain to use for building</param>
 		/// <param name="CompileEnvironment">The environment to compile the binary in</param>
 		/// <param name="LinkEnvironment">The environment to link the binary in</param>
 		/// <returns></returns>
-		public override IEnumerable<FileItem> Build(CPPEnvironment CompileEnvironment, LinkEnvironment LinkEnvironment)
+		public override IEnumerable<FileItem> Build(IUEToolChain ToolChain, CPPEnvironment CompileEnvironment, LinkEnvironment LinkEnvironment)
 		{
 			var ProjectCSharpEnviroment = new CSharpEnvironment();
 			if (LinkEnvironment.Config.Target.Configuration == CPPTargetConfiguration.Debug)
diff --git a/Engine/Source/Programs/UnrealBuildTool/Configuration/UEBuildTarget.cs b/Engine/Source/Programs/UnrealBuildTool/Configuration/UEBuildTarget.cs
--- a/Engine/Source/Programs/UnrealBuildTool/Configuration/UEBuildTarget.cs
+++ b/Engine/Source/Programs/UnrealBuildTool/Configuration/UEBuildTarget.cs
@@ -1095,7 +1095,7 @@ namespace UnrealBuildTool
 		}
 
 		/** Builds the target, appending list of output files and returns building result. */
-		public ECompilationResult Build(List<FileItem> OutputItems)
+		public ECompilationResult Build(IUEToolChain TargetToolChain, List<FileItem> OutputItems)
 		{
 			// Set up the global compile and link environment in GlobalCompileEnvironment and GlobalLinkEnvironment.
 			SetupGlobalEnvironment();
@@ -1428,7 +1428,7 @@ namespace UnrealBuildTool
 			// Build the target's binaries.
 			foreach (var Binary in AppBinaries)
 			{
-				OutputItems.AddRange(Binary.Build(GlobalCompileEnvironment, GlobalLinkEnvironment));
+				OutputItems.AddRange(Binary.Build(TargetToolChain, GlobalCompileEnvironment, GlobalLinkEnvironment));
 			}
 
 			if (BuildConfiguration.WriteTargetInfoPath != null)
diff --git a/Engine/Source/Programs/UnrealBuildTool/IOS/UEBuildIOS.cs b/Engine/Source/Programs/UnrealBuildTool/IOS/UEBuildIOS.cs
--- a/Engine/Source/Programs/UnrealBuildTool/IOS/UEBuildIOS.cs
+++ b/Engine/Source/Programs/UnrealBuildTool/IOS/UEBuildIOS.cs
@@ -302,7 +302,7 @@ namespace UnrealBuildTool
 		/// <param name="CompileEnvironment">The environment to compile the binary in</param>
 		/// <param name="LinkEnvironment">The environment to link the binary in</param>
 		/// <returns></returns>
-		public override IEnumerable<FileItem> Build(CPPEnvironment CompileEnvironment, LinkEnvironment LinkEnvironment)
+		public override IEnumerable<FileItem> Build(IUEToolChain ToolChain, CPPEnvironment CompileEnvironment, LinkEnvironment LinkEnvironment)
 		{
 			// generate the .stub!
 			return new FileItem[] { FileItem.GetItemByPath(Config.OutputFilePath) };
diff --git a/Engine/Source/Programs/UnrealBuildTool/Mac/MacToolChain.cs b/Engine/Source/Programs/UnrealBuildTool/Mac/MacToolChain.cs
--- a/Engine/Source/Programs/UnrealBuildTool/Mac/MacToolChain.cs
+++ b/Engine/Source/Programs/UnrealBuildTool/Mac/MacToolChain.cs
@@ -969,7 +969,7 @@ namespace UnrealBuildTool
 			return RemoteOutputFile;
 		}
 
-		public FileItem FixDylibDependencies(LinkEnvironment LinkEnvironment, FileItem Executable)
+		FileItem FixDylibDependencies(LinkEnvironment LinkEnvironment, FileItem Executable)
 		{
 			Action LinkAction = new Action(ActionType.Link);
 			LinkAction.WorkingDirectory = Path.GetFullPath(".");
@@ -1054,7 +1054,7 @@ namespace UnrealBuildTool
 		 * 
 		 * @param Executable FileItem describing the executable to generate app bundle for
 		 */
-		public FileItem CreateAppBundle(LinkEnvironment LinkEnvironment, FileItem Executable, FileItem FixDylibOutputFile)
+		FileItem CreateAppBundle(LinkEnvironment LinkEnvironment, FileItem Executable, FileItem FixDylibOutputFile)
 		{
 			// Make a file item for the source and destination files
 			string FullDestPath = Executable.AbsolutePath.Substring(0, Executable.AbsolutePath.IndexOf(".app") + 4);
@@ -1233,5 +1233,26 @@ namespace UnrealBuildTool
 
 			base.PostBuildSync(Target);
 		}
+
+		public override ICollection<FileItem> PostBuild (FileItem Executable, LinkEnvironment BinaryLinkEnvironment)
+		{
+			var OutputFiles = base.PostBuild (Executable, BinaryLinkEnvironment);
+
+			// if building for Mac on a Mac, use actions to finalize the builds (otherwise, we use Deploy)
+			if (ExternalExecution.GetRuntimePlatform () != UnrealTargetPlatform.Mac) {
+				return OutputFiles;
+			}
+
+			if (BinaryLinkEnvironment.Config.bIsBuildingDLL || BinaryLinkEnvironment.Config.bIsBuildingLibrary) {
+				return OutputFiles;
+			}
+			FileItem FixDylibOutputFile = FixDylibDependencies(BinaryLinkEnvironment, Executable);
+			OutputFiles.Add(FixDylibOutputFile);
+			if (BinaryLinkEnvironment.Config.bIsBuildingConsoleApplication)
+				return OutputFiles;
+
+			OutputFiles.Add(CreateAppBundle(BinaryLinkEnvironment, Executable, FixDylibOutputFile));
+			return OutputFiles;
+		}
 	};
 }
diff --git a/Engine/Source/Programs/UnrealBuildTool/System/UnrealBuildTool.cs b/Engine/Source/Programs/UnrealBuildTool/System/UnrealBuildTool.cs
--- a/Engine/Source/Programs/UnrealBuildTool/System/UnrealBuildTool.cs
+++ b/Engine/Source/Programs/UnrealBuildTool/System/UnrealBuildTool.cs
@@ -1040,6 +1040,8 @@ namespace UnrealBuildTool
 			string ExecutorName = "Unknown";
 			ECompilationResult BuildResult = ECompilationResult.Succeeded;
 
+			var ToolChain = UEToolChain.GetPlatformToolChain(BuildPlatform.GetCPPTargetPlatform(ResetPlatform));
+
 			try
 			{
 				List<string[]> TargetSettings = ParseCommandLineFlags(Arguments);
@@ -1079,7 +1081,7 @@ namespace UnrealBuildTool
 					}
 
 					var TargetOutputItems = new List<FileItem>();
-					BuildResult = Target.Build(TargetOutputItems);
+					BuildResult = Target.Build(ToolChain, TargetOutputItems);
 
 					if(BuildResult != ECompilationResult.Succeeded)
 					{
@@ -1155,7 +1157,6 @@ namespace UnrealBuildTool
 							ActionsToExecute.Count
 							);
 
-					var ToolChain = UEToolChain.GetPlatformToolChain(BuildPlatform.GetCPPTargetPlatform(ResetPlatform));
 					ToolChain.PreBuildSync();
 
 					// Execute the actions.
diff --git a/Engine/Source/Programs/UnrealBuildTool/ToolChain/UEToolChain.cs b/Engine/Source/Programs/UnrealBuildTool/ToolChain/UEToolChain.cs
--- a/Engine/Source/Programs/UnrealBuildTool/ToolChain/UEToolChain.cs
+++ b/Engine/Source/Programs/UnrealBuildTool/ToolChain/UEToolChain.cs
@@ -32,8 +32,10 @@ namespace UnrealBuildTool
 		
 		void PreBuildSync();
 		
-        void PostBuildSync(UEBuildTarget Target);
-        
+		void PostBuildSync(UEBuildTarget Target);
+
+		ICollection<FileItem> PostBuild(FileItem Executable, LinkEnvironment ExecutableLinkEnvironment);
+
 		void SetUpGlobalEnvironment();
 	}
 
@@ -128,14 +130,18 @@ namespace UnrealBuildTool
 		{
 		}
 
-        public virtual void PostBuildSync(UEBuildTarget Target)
+		public virtual void PostBuildSync(UEBuildTarget Target)
 		{
 		}
 
-		public virtual void SetUpGlobalEnvironment()
+		public virtual ICollection<FileItem> PostBuild(FileItem Executable, LinkEnvironment ExecutableLinkEnvironment)
 		{
+			return new List<FileItem> ();
 		}
 
+		public virtual void SetUpGlobalEnvironment()
+		{
+		}
 
 		protected void RunUnrealHeaderToolIfNeeded()
 		{
