[GitHub] 222 : Enable unregistering of ToolChains
#ttp 338159

diff --git a/Engine/Source/Programs/UnrealBuildTool/ToolChain/UEToolChain.cs b/Engine/Source/Programs/UnrealBuildTool/ToolChain/UEToolChain.cs
--- a/Engine/Source/Programs/UnrealBuildTool/ToolChain/UEToolChain.cs
+++ b/Engine/Source/Programs/UnrealBuildTool/ToolChain/UEToolChain.cs
@@ -57,6 +57,11 @@ namespace UnrealBuildTool
 			}
 		}
 
+		public static void UnregisterPlatformToolChain(CPPTargetPlatform InPlatform)
+		{
+			CPPToolChainDictionary.Remove(InPlatform);
+		}
+
 		public static IUEToolChain GetPlatformToolChain(CPPTargetPlatform InPlatform)
 		{
 			if (CPPToolChainDictionary.ContainsKey(InPlatform) == true)
