[GitHub] 223 : Allow to register IUEToolChain implementations
#ttp 338160

diff --git a/Engine/Source/Programs/UnrealBuildTool/ToolChain/UEToolChain.cs b/Engine/Source/Programs/UnrealBuildTool/ToolChain/UEToolChain.cs
--- a/Engine/Source/Programs/UnrealBuildTool/ToolChain/UEToolChain.cs
+++ b/Engine/Source/Programs/UnrealBuildTool/ToolChain/UEToolChain.cs
@@ -41,9 +41,9 @@ namespace UnrealBuildTool
 
 	public abstract class UEToolChain : IUEToolChain
 	{
-		static Dictionary<CPPTargetPlatform, UEToolChain> CPPToolChainDictionary = new Dictionary<CPPTargetPlatform, UEToolChain>();
+		static Dictionary<CPPTargetPlatform, IUEToolChain> CPPToolChainDictionary = new Dictionary<CPPTargetPlatform, IUEToolChain>();
 
-		public static void RegisterPlatformToolChain(CPPTargetPlatform InPlatform, UEToolChain InToolChain)
+		public static void RegisterPlatformToolChain(CPPTargetPlatform InPlatform, IUEToolChain InToolChain)
 		{
 			if (CPPToolChainDictionary.ContainsKey(InPlatform) == true)
 			{
