UBT - Added an option to compile a shipping target configuration and save binary with the development name

diff --git a/Engine/Source/Programs/CrashReporter/CrashReportClient/CrashReportClient.Target.cs b/Engine/Source/Programs/CrashReporter/CrashReportClient/CrashReportClient.Target.cs
--- a/Engine/Source/Programs/CrashReporter/CrashReportClient/CrashReportClient.Target.cs
+++ b/Engine/Source/Programs/CrashReporter/CrashReportClient/CrashReportClient.Target.cs
@@ -22,6 +22,19 @@ public class CrashReportClientTarget : TargetRules
 		return true;
 	}
 
+	public override bool GetSupportedConfigurations(ref List<UnrealTargetConfiguration> OutConfigurations, bool bIncludeTestAndShippingConfigs)
+	{
+		if( base.GetSupportedConfigurations( ref OutConfigurations, bIncludeTestAndShippingConfigs ) )
+		{
+			OutConfigurations.Add( UnrealTargetConfiguration.Shipping );
+			return true;
+		}
+		else
+		{
+			return false;
+		}
+	}
+
 	public override void SetupBinaries(
 		TargetInfo Target,
 		ref List<UEBuildBinaryConfiguration> OutBuildBinaryConfigurations,
@@ -39,6 +52,11 @@ public class CrashReportClientTarget : TargetRules
 		}
 	}
 
+	public override bool ForceNameAsForDevelopment()
+	{
+		return true;
+	}
+
 	public override bool ShouldCompileMonolithic(UnrealTargetPlatform InPlatform, UnrealTargetConfiguration InConfiguration)
 	{
 		return true;
@@ -88,6 +106,6 @@ public class CrashReportClientTarget : TargetRules
 
 	public override List<UnrealTargetConfiguration> GUBP_ToolConfigs( UnrealTargetPlatform InHostPlatform )
 	{
-		return new List<UnrealTargetConfiguration> { UnrealTargetConfiguration.Development, UnrealTargetConfiguration.Shipping };
+		return new List<UnrealTargetConfiguration> { UnrealTargetConfiguration.Shipping };
 	}
 }
diff --git a/Engine/Source/Programs/UnrealBuildTool/Configuration/UEBuildTarget.cs b/Engine/Source/Programs/UnrealBuildTool/Configuration/UEBuildTarget.cs
--- a/Engine/Source/Programs/UnrealBuildTool/Configuration/UEBuildTarget.cs
+++ b/Engine/Source/Programs/UnrealBuildTool/Configuration/UEBuildTarget.cs
@@ -697,7 +697,7 @@ namespace UnrealBuildTool
 
 			// Construct the output path based on configuration, platform, game if not specified.
 			TargetRules.TargetType? TargetType = (Rules != null) ? Rules.Type : (TargetRules.TargetType?)null;
-			OutputPath = Path.GetFullPath(MakeBinaryPath("", AppName, UEBuildBinaryType.Executable, TargetType, bIsRocketGame, null, InAppName));
+			OutputPath = Path.GetFullPath(MakeBinaryPath("", AppName, UEBuildBinaryType.Executable, TargetType, bIsRocketGame, null, InAppName,Rules.ForceNameAsForDevelopment()));
 
 			if (bCompileMonolithic && TargetRules.IsGameType(InRulesObject.Type))
 			{
@@ -1838,7 +1838,7 @@ namespace UnrealBuildTool
 
 		/** Given a UBT-built binary name (e.g. "Core"), returns a relative path to the binary for the current build configuration (e.g. "../Binaries/Win64/Core-Win64-Debug.lib") */
 		public static string MakeBinaryPath(string ModuleName, string BinaryName, UnrealTargetPlatform Platform, 
-			UnrealTargetConfiguration Configuration, UEBuildBinaryType BinaryType, TargetRules.TargetType? TargetType, bool bIsRocketModule, PluginInfo PluginInfo, string AppName)
+			UnrealTargetConfiguration Configuration, UEBuildBinaryType BinaryType, TargetRules.TargetType? TargetType, bool bIsRocketModule, PluginInfo PluginInfo, string AppName, bool bForceNameAsForDevelopment = false)
 		{
 			// Determine the binary extension for the platform and binary type.
 			var BuildPlatform = UEBuildPlatform.GetBuildPlatform(Platform);
@@ -1895,7 +1895,7 @@ namespace UnrealBuildTool
 
 			string OutBinaryPath = "";
 			// Append binary file name
-			if (LocalConfig == UnrealTargetConfiguration.Development)
+			if (LocalConfig == UnrealTargetConfiguration.Development || bForceNameAsForDevelopment)
 			{
 				OutBinaryPath = Path.Combine(BaseDirectory, String.Format("{0}{1}{2}", BinaryName, BinarySuffix, BinaryExtension));
 			}
@@ -1908,15 +1908,15 @@ namespace UnrealBuildTool
 		}
 
 		/** Given a UBT-built binary name (e.g. "Core"), returns a relative path to the binary for the current build configuration (e.g. "../../Binaries/Win64/Core-Win64-Debug.lib") */
-		public string MakeBinaryPath(string ModuleName, string BinaryName, UEBuildBinaryType BinaryType, TargetRules.TargetType? TargetType, bool bIsRocketModule, PluginInfo PluginInfo, string AppName )
+		public string MakeBinaryPath(string ModuleName, string BinaryName, UEBuildBinaryType BinaryType, TargetRules.TargetType? TargetType, bool bIsRocketModule, PluginInfo PluginInfo, string AppName, bool bForceNameAsForDevelopment = false)
 		{
 			if (String.IsNullOrEmpty(ModuleName) && Configuration == UnrealTargetConfiguration.DebugGame && !bCompileMonolithic)
 			{
-				return MakeBinaryPath(ModuleName, BinaryName, Platform, UnrealTargetConfiguration.Development, BinaryType, TargetType, bIsRocketModule, PluginInfo, AppName);
+				return MakeBinaryPath(ModuleName, BinaryName, Platform, UnrealTargetConfiguration.Development, BinaryType, TargetType, bIsRocketModule, PluginInfo, AppName, bForceNameAsForDevelopment);
 			}
 			else
 			{
-				return MakeBinaryPath(ModuleName, BinaryName, Platform, Configuration, BinaryType, TargetType, bIsRocketModule, PluginInfo, AppName);
+				return MakeBinaryPath(ModuleName, BinaryName, Platform, Configuration, BinaryType, TargetType, bIsRocketModule, PluginInfo, AppName, bForceNameAsForDevelopment);
 			}
 		}
 
diff --git a/Engine/Source/Programs/UnrealBuildTool/System/RulesCompiler.cs b/Engine/Source/Programs/UnrealBuildTool/System/RulesCompiler.cs
--- a/Engine/Source/Programs/UnrealBuildTool/System/RulesCompiler.cs
+++ b/Engine/Source/Programs/UnrealBuildTool/System/RulesCompiler.cs
@@ -537,6 +537,17 @@ namespace UnrealBuildTool
 		{
 		}
 
+
+		/// <summary>
+		/// Returns true if this target's output path needs to be the same as for the development configuration.
+		/// Currently only used by the CrashReportClient.
+		/// </summary>
+		/// <returns>true if this target's output path needs to be the same as for the development configuration.</returns>
+		public virtual bool ForceNameAsForDevelopment()
+		{
+			return false;
+		}
+
 		/// <summary>
 		/// Setup the global environment for building this target
 		/// IMPORTANT: Game targets will *not* have this function called unless they are built as monolithic targets.
