Enabled Using Android SDK switching.

diff --git a/Engine/Source/Programs/UnrealBuildTool/Android/UEBuildAndroid.cs b/Engine/Source/Programs/UnrealBuildTool/Android/UEBuildAndroid.cs
--- a/Engine/Source/Programs/UnrealBuildTool/Android/UEBuildAndroid.cs
+++ b/Engine/Source/Programs/UnrealBuildTool/Android/UEBuildAndroid.cs
@@ -33,9 +33,34 @@ namespace UnrealBuildTool
             return false;
         }
 
+        public override bool PlatformSupportsSDKSwitching()
+        {
+            return true;
+        }
+
+        public override string GetSDKTargetPlatformName()
+        {
+            return "Android";
+        }
+
+        public override string GetRequiredSDKString()
+        {
+            return "-19";
+        }
+
+        public override String GetRequiredScriptVersionString()
+        {
+            return "1.0";
+        }
+
 		public override SDKStatus HasRequiredSDKsInstalled()
 		{
-			string NDKPath = Environment.GetEnvironmentVariable("NDKROOT");
+            if (base.HasRequiredSDKsInstalled() == SDKStatus.Valid)
+            {
+                return SDKStatus.Valid;
+            }
+
+            string NDKPath = Environment.GetEnvironmentVariable("NDKROOT");
 
 			// we don't have an NDKROOT specified
 			if (String.IsNullOrEmpty(NDKPath))
