Few more compile error fixes for Xcode indexing

diff --git a/Engine/Source/Programs/UnrealBuildTool/System/XcodeProjectFileGenerator.cs b/Engine/Source/Programs/UnrealBuildTool/System/XcodeProjectFileGenerator.cs
--- a/Engine/Source/Programs/UnrealBuildTool/System/XcodeProjectFileGenerator.cs
+++ b/Engine/Source/Programs/UnrealBuildTool/System/XcodeProjectFileGenerator.cs
@@ -927,7 +927,7 @@ namespace UnrealBuildTool
 				{
 					string Definition = CurDefinition;
 					string AlternateDefinition = Definition.Contains("=0") ? Definition.Replace("=0", "=1") : Definition.Replace("=1", "=0");
-					if (Definition.Equals("WITH_EDITORONLY_DATA=0"))
+					if (Definition.Equals("WITH_EDITORONLY_DATA=0") || Definition.Equals("WITH_DATABASE_SUPPORT=1"))
 					{
 						Definition = AlternateDefinition;
 					}
diff --git a/Engine/Source/Programs/UnrealSync/Private/UnrealSync.cpp b/Engine/Source/Programs/UnrealSync/Private/UnrealSync.cpp
--- a/Engine/Source/Programs/UnrealSync/Private/UnrealSync.cpp
+++ b/Engine/Source/Programs/UnrealSync/Private/UnrealSync.cpp
@@ -3,6 +3,7 @@
 #include "UnrealSync.h"
 #include "GUI.h"
 #include "P4DataCache.h"
+#include "P4Env.h"
 
 #include "RequiredProgramMainCPPInclude.h"
 
diff --git a/Engine/Source/Programs/UnrealVersionSelector/Private/GenericPlatform/GenericPlatformInstallation.cpp b/Engine/Source/Programs/UnrealVersionSelector/Private/GenericPlatform/GenericPlatformInstallation.cpp
--- a/Engine/Source/Programs/UnrealVersionSelector/Private/GenericPlatform/GenericPlatformInstallation.cpp
+++ b/Engine/Source/Programs/UnrealVersionSelector/Private/GenericPlatform/GenericPlatformInstallation.cpp
@@ -4,6 +4,7 @@
 #include "GenericPlatformInstallation.h"
 #include "../PlatformInstallation.h"
 #include "Runtime/Launch/Resources/Version.h"
+#include "DesktopPlatformModule.h"
 
 bool FGenericPlatformInstallation::NormalizeEngineRootDir(FString &RootDir)
 {
