Updated list of known files that should not be part of UE4XcodeHelper target in Xcode project

diff --git a/Engine/Source/Programs/UnrealBuildTool/System/XcodeProject.cs b/Engine/Source/Programs/UnrealBuildTool/System/XcodeProject.cs
--- a/Engine/Source/Programs/UnrealBuildTool/System/XcodeProject.cs
+++ b/Engine/Source/Programs/UnrealBuildTool/System/XcodeProject.cs
@@ -188,7 +188,10 @@ namespace UnrealBuildTool
 						// UE4XcodeHelper is Mac only target, so skip other platforms files
 						return false;
 					}
-					else if (SourceFile.FilePath.EndsWith("SimplygonMeshReduction.cpp") || SourceFile.FilePath.EndsWith("Android.cpp"))
+					else if (SourceFile.FilePath.EndsWith("SimplygonMeshReduction.cpp") || SourceFile.FilePath.EndsWith("MeshBoneReduction.cpp") || SourceFile.FilePath.EndsWith("Android.cpp")
+						|| SourceFile.FilePath.EndsWith("Amazon.cpp") || SourceFile.FilePath.EndsWith("FacebookModule.cpp") || SourceFile.FilePath.EndsWith("SDL_angle.c")
+						|| SourceFile.FilePath.Contains("VisualStudioSourceCodeAccess") || SourceFile.FilePath.Contains("AndroidDevice") || SourceFile.FilePath.Contains("IOSDevice")
+						|| SourceFile.FilePath.Contains("WindowsDevice") || SourceFile.FilePath.Contains("WindowsMoviePlayer") || SourceFile.FilePath.EndsWith("IOSTapJoy.cpp"))
 					{
 						// @todo: We need a way to filter out files that use SDKs we don't have
 						return false;
