Removed some excessive warnings from build script when not using carefullyredist.
Upgraded version string for build script.

diff --git a/Engine/Source/Programs/UnrealBuildTool/Android/UEBuildAndroid.cs b/Engine/Source/Programs/UnrealBuildTool/Android/UEBuildAndroid.cs
--- a/Engine/Source/Programs/UnrealBuildTool/Android/UEBuildAndroid.cs
+++ b/Engine/Source/Programs/UnrealBuildTool/Android/UEBuildAndroid.cs
@@ -21,6 +21,10 @@ namespace UnrealBuildTool
 		[XmlConfig]
 		public static string AndroidSdkApiTarget = "latest";
 
+        // this is to make sure we don't check the sdk all the time
+        static int HasSDK = -1;
+
+
 		// The current architecture - affects everything about how UBT operates on Android
 		public override string GetActiveArchitecture()
 		{
@@ -50,33 +54,37 @@ namespace UnrealBuildTool
 
         public override String GetRequiredScriptVersionString()
         {
-            return "1.0";
+            return "1.1";
         }
 
 		public override SDKStatus HasRequiredSDKsInstalled()
 		{
-            if (base.HasRequiredSDKsInstalled() == SDKStatus.Valid)
+            if (HasSDK == -1)
             {
-                return SDKStatus.Valid;
-            }
-
-            string NDKPath = Environment.GetEnvironmentVariable("NDKROOT");
-
-			// we don't have an NDKROOT specified
-			if (String.IsNullOrEmpty(NDKPath))
-			{
-				return SDKStatus.Invalid;
-			}
-
-			NDKPath = NDKPath.Replace("\"", "");
+                if (base.HasRequiredSDKsInstalled() == SDKStatus.Valid)
+                {
+                    HasSDK = 1;
+                    return HasSDK == 1 ? SDKStatus.Valid : SDKStatus.Invalid;
+                }
+                
+                string NDKPath = Environment.GetEnvironmentVariable("NDKROOT");
+
+                // we don't have an NDKROOT specified
+                if (String.IsNullOrEmpty(NDKPath))
+                {
+                    HasSDK = 0;
+                    return HasSDK == 1 ? SDKStatus.Valid : SDKStatus.Invalid;
+                }
 
-			// can't find llvm-3.3 or llvm-3.1 in the toolchains
-			if (!Directory.Exists(Path.Combine(NDKPath, @"toolchains\llvm-3.3")) && !Directory.Exists(Path.Combine(NDKPath, @"toolchains\llvm-3.1")))
-			{
-				return SDKStatus.Invalid;
-			}
+                NDKPath = NDKPath.Replace("\"", "");
 
-			return SDKStatus.Valid;
+                // can't find llvm-3.3 or llvm-3.1 in the toolchains
+                if (!Directory.Exists(Path.Combine(NDKPath, @"toolchains\llvm-3.3")) && !Directory.Exists(Path.Combine(NDKPath, @"toolchains\llvm-3.1")))
+                {
+                    HasSDK = 0;
+                }
+            }
+            return HasSDK == 1 ? SDKStatus.Valid : SDKStatus.Invalid;
 		}
 
 		public override void RegisterBuildPlatform()
