Add a parameter to UBT which allows writing a list of ThirdParty files required to build a target. Use that list to distill the correct files for Rocket.

diff --git a/Engine/Source/Programs/AutomationTool/UBTUtils.cs b/Engine/Source/Programs/AutomationTool/UBTUtils.cs
--- a/Engine/Source/Programs/AutomationTool/UBTUtils.cs
+++ b/Engine/Source/Programs/AutomationTool/UBTUtils.cs
@@ -46,7 +46,7 @@ namespace AutomationTool
 			}
 			else
 			{
-				CmdLine = String.Format("{0} {1} {2} -Project={3} {4}", Target, Platform, Config, Project, AdditionalArgs);
+				CmdLine = String.Format("{0} {1} {2} -Project={3} {4}", Target, Platform, Config, CommandUtils.MakePathSafeToUseWithCommandLine(Project), AdditionalArgs);
 			}
 			return CmdLine;
 		}
diff --git a/Engine/Source/Programs/AutomationTool/UE4Build.cs b/Engine/Source/Programs/AutomationTool/UE4Build.cs
--- a/Engine/Source/Programs/AutomationTool/UE4Build.cs
+++ b/Engine/Source/Programs/AutomationTool/UE4Build.cs
@@ -1079,6 +1079,21 @@ namespace AutomationTool
 			return UnrealBuildTool.UEBuildPlatform.GetBuildPlatform(Platform).CanUseXGE();
 		}
 
+		public string[] GetExternalFileList(BuildAgenda Agenda)
+		{
+			// Create the file list
+			string AdditionalArguments = "";
+			foreach (var Target in Agenda.Targets)
+			{
+				RunUBT(CmdEnv, UBTExecutable, Target.UprojectPath, Target.TargetName, Target.Platform.ToString(), Target.Config.ToString(), "-generateexternalfilelist" + AdditionalArguments + " " + Target.UBTArgs);
+				AdditionalArguments = " -mergeexternalfilelist";
+			}
+
+			// Read it in, and return the filenames as an array of strings
+			string FileListPath = CommandUtils.CombinePaths(CmdEnv.LocalRoot, @"/Engine/Intermediate/Build/ExternalFiles.xml");
+			return UnrealBuildTool.Utils.ReadClass<UnrealBuildTool.ExternalFileList>(FileListPath).FileNames.ToArray();
+		}
+
 		/// <summary>
 		/// Executes a build.
 		/// </summary>
diff --git a/Engine/Source/Programs/UnrealBuildTool/Configuration/UEBuildConfiguration.cs b/Engine/Source/Programs/UnrealBuildTool/Configuration/UEBuildConfiguration.cs
--- a/Engine/Source/Programs/UnrealBuildTool/Configuration/UEBuildConfiguration.cs
+++ b/Engine/Source/Programs/UnrealBuildTool/Configuration/UEBuildConfiguration.cs
@@ -67,6 +67,14 @@ namespace UnrealBuildTool
 		[XmlConfig]
 		public static bool bCompileLeanAndMeanUE;
 
+		/** Whether to generate a list of external files that are required to build a target */
+		[XmlConfig]
+		public static bool bGenerateExternalFileList;
+
+		/** Whether to merge to the existing list of external files */
+		[XmlConfig]
+		public static bool bMergeExternalFileList;
+
 		/** Whether to generate a manifest file that contains the files to add to Perforce */
 		[XmlConfig]
 		public static bool bGenerateManifest;
diff --git a/Engine/Source/Programs/UnrealBuildTool/Configuration/UEBuildTarget.cs b/Engine/Source/Programs/UnrealBuildTool/Configuration/UEBuildTarget.cs
--- a/Engine/Source/Programs/UnrealBuildTool/Configuration/UEBuildTarget.cs
+++ b/Engine/Source/Programs/UnrealBuildTool/Configuration/UEBuildTarget.cs
@@ -125,6 +125,14 @@ namespace UnrealBuildTool
 		}
 	}
 
+	/// <summary>
+	/// A list of extenal files required to build a given target
+	/// </summary>
+	public class ExternalFileList
+	{
+		public readonly List<string> FileNames = new List<string>();
+	}
+
 	public class OnlyModule
 	{
 		public OnlyModule(string InitOnlyModuleName)
@@ -1006,6 +1014,94 @@ namespace UnrealBuildTool
 			}
 		}
 
+		/** Generates a list of external files which are required to build this target */
+		public void GenerateExternalFileList()
+		{
+			string FileListPath = "../Intermediate/Build/ExternalFiles.xml";
+
+			// Find all the external modules
+			HashSet<string> ModuleNames = new HashSet<string>();
+			foreach(UEBuildBinary Binary in AppBinaries)
+			{
+				foreach(UEBuildModule Module in Binary.GetAllDependencyModules(bIncludeDynamicallyLoaded: false, bForceCircular: false))
+				{
+					UEBuildExternalModule ExternalModule = Module as UEBuildExternalModule;
+					if(ExternalModule != null)
+					{
+						ModuleNames.Add(ExternalModule.Name);
+					}
+				}
+			}
+
+			// Create a set of filenames
+			HashSet<string> FileNames = new HashSet<string>(StringComparer.CurrentCultureIgnoreCase);
+
+			// Add all their include paths
+			foreach(string ModuleName in ModuleNames)
+			{
+				// Create the module rules
+				string ModuleRulesFileName;
+				ModuleRules Rules = RulesCompiler.CreateModuleRules(ModuleName, GetTargetInfo(), out ModuleRulesFileName);
+
+				// Add the rules file itself
+				FileNames.Add(ModuleRulesFileName);
+
+				// Add all the libraries
+				foreach(string LibraryName in Rules.PublicAdditionalLibraries)
+				{
+					foreach(string LibraryPath in Rules.PublicLibraryPaths.Where(x => !x.StartsWith("$(")))
+					{
+						string LibraryFileName = Path.Combine(LibraryPath, LibraryName);
+						if(File.Exists(LibraryFileName))
+						{
+							FileNames.Add(LibraryFileName);
+						}
+					}
+				}
+
+				// Find all the include paths
+				List<string> AllIncludePaths = new List<string>();
+				AllIncludePaths.AddRange(Rules.PublicIncludePaths);
+				AllIncludePaths.AddRange(Rules.PublicSystemIncludePaths);
+
+				// Add all the include paths
+				foreach(string IncludePath in AllIncludePaths.Where(x => !x.StartsWith("$(")))
+				{
+					foreach(string IncludeFileName in Directory.EnumerateFiles(IncludePath, "*", SearchOption.AllDirectories))
+					{
+						string Extension = Path.GetExtension(IncludeFileName).ToLower();
+						if(Extension == ".h" || Extension == ".inl")
+						{
+							FileNames.Add(IncludeFileName);
+						}
+					}
+				}
+			}
+
+			// Normalize all the filenames
+			HashSet<string> NormalizedFileNames = new HashSet<string>(StringComparer.CurrentCultureIgnoreCase);
+			foreach(string FileName in FileNames)
+			{
+				string NormalizedFileName = Path.GetFullPath(FileName).Replace('\\', '/');
+				NormalizedFileNames.Add(NormalizedFileName);
+			}
+
+			// Add the existing filenames
+			if(UEBuildConfiguration.bMergeExternalFileList)
+			{
+				foreach(string FileName in Utils.ReadClass<ExternalFileList>(FileListPath).FileNames)
+				{
+					NormalizedFileNames.Add(FileName);
+				}
+			}
+
+			// Write the output list
+			ExternalFileList FileList = new ExternalFileList();
+			FileList.FileNames.AddRange(NormalizedFileNames);
+			FileList.FileNames.Sort();
+			Utils.WriteClass<ExternalFileList>(FileList, FileListPath, "");
+		}
+
 		/** Generates a public manifest file for writing out */
         public void GenerateManifest(List<UEBuildBinary> Binaries, CPPTargetPlatform Platform, List<string> SpecialRocketLibFilesThatAreBuildProducts)
 		{
@@ -1344,6 +1440,13 @@ namespace UnrealBuildTool
 				MacToolChain.SetupBundleDependencies(AppBinaries, GameName);
 			}
 
+			// Generate the external file list 
+			if(UEBuildConfiguration.bGenerateExternalFileList)
+			{
+				GenerateExternalFileList();
+				return ECompilationResult.Succeeded;
+			}
+
 			// If we're only generating the manifest, return now
 			if (UEBuildConfiguration.bGenerateManifest || UEBuildConfiguration.bCleanProject)
 			{
diff --git a/Engine/Source/Programs/UnrealBuildTool/System/UnrealBuildTool.cs b/Engine/Source/Programs/UnrealBuildTool/System/UnrealBuildTool.cs
--- a/Engine/Source/Programs/UnrealBuildTool/System/UnrealBuildTool.cs
+++ b/Engine/Source/Programs/UnrealBuildTool/System/UnrealBuildTool.cs
@@ -756,6 +756,14 @@ namespace UnrealBuildTool
 						{
 							UEBuildConfiguration.bPrepForDeployment = true;
 						}
+						else if (LowercaseArg == "-generateexternalfilelist")
+						{
+							UEBuildConfiguration.bGenerateExternalFileList = true;
+						}
+						else if (LowercaseArg == "-mergeexternalfilelist")
+						{
+							UEBuildConfiguration.bMergeExternalFileList = true;
+						}
 						else if (LowercaseArg == "-generatemanifest")
 						{
 							// Generate a manifest file containing all the files required to be in Perforce
@@ -1144,7 +1152,7 @@ namespace UnrealBuildTool
 				if (BuildResult == ECompilationResult.Succeeded &&
 					(
 						(BuildConfiguration.bXGEExport && UEBuildConfiguration.bGenerateManifest) ||
-						(!GeneratingActionGraph && !ProjectFileGenerator.bGenerateProjectFiles && !UEBuildConfiguration.bGenerateManifest && !UEBuildConfiguration.bCleanProject)
+						(!GeneratingActionGraph && !ProjectFileGenerator.bGenerateProjectFiles && !UEBuildConfiguration.bGenerateManifest && !UEBuildConfiguration.bCleanProject && !UEBuildConfiguration.bGenerateExternalFileList)
 					))
 				{
 					// Plan the actions to execute for the build.
