Allow generating external files to find Unix style libraries (ie. with an implicit 'lib' before the name). Also include shadow files in the list of dependencies, and prevent trying to deploy.

diff --git a/Engine/Source/Programs/UnrealBuildTool/Configuration/UEBuildTarget.cs b/Engine/Source/Programs/UnrealBuildTool/Configuration/UEBuildTarget.cs
--- a/Engine/Source/Programs/UnrealBuildTool/Configuration/UEBuildTarget.cs
+++ b/Engine/Source/Programs/UnrealBuildTool/Configuration/UEBuildTarget.cs
@@ -1036,6 +1036,9 @@ namespace UnrealBuildTool
 			// Create a set of filenames
 			HashSet<string> FileNames = new HashSet<string>(StringComparer.CurrentCultureIgnoreCase);
 
+			// Get the platform we're building for
+			IUEBuildPlatform BuildPlatform = UEBuildPlatform.GetBuildPlatform(Platform);
+
 			// Add all their include paths
 			foreach(string ModuleName in ModuleNames)
 			{
@@ -1046,16 +1049,38 @@ namespace UnrealBuildTool
 				// Add the rules file itself
 				FileNames.Add(ModuleRulesFileName);
 
+				// Get a list of all the library paths
+				List<string> LibraryPaths = new List<string>();
+				LibraryPaths.Add(Directory.GetCurrentDirectory());
+				LibraryPaths.AddRange(Rules.PublicLibraryPaths.Where(x => !x.StartsWith("$(")).Select(x => Path.GetFullPath(x.Replace('/', '\\'))));
+
 				// Add all the libraries
+				string LibraryExtension = BuildPlatform.GetBinaryExtension(UEBuildBinaryType.StaticLibrary);
 				foreach(string LibraryName in Rules.PublicAdditionalLibraries)
 				{
-					foreach(string LibraryPath in Rules.PublicLibraryPaths.Where(x => !x.StartsWith("$(")))
+					foreach(string LibraryPath in LibraryPaths)
 					{
 						string LibraryFileName = Path.Combine(LibraryPath, LibraryName);
 						if(File.Exists(LibraryFileName))
 						{
 							FileNames.Add(LibraryFileName);
 						}
+
+						string UnixLibraryFileName = Path.Combine(LibraryPath, "lib" + LibraryName + LibraryExtension);
+						if(File.Exists(UnixLibraryFileName))
+						{
+							FileNames.Add(UnixLibraryFileName);
+						}
+					}
+				}
+
+				// Add all the additional shadow files
+				foreach(string AdditionalShadowFile in Rules.PublicAdditionalShadowFiles)
+				{
+					string ShadowFileName = Path.GetFullPath(AdditionalShadowFile);
+					if(File.Exists(ShadowFileName))
+					{
+						FileNames.Add(ShadowFileName);
 					}
 				}
 
diff --git a/Engine/Source/Programs/UnrealBuildTool/System/UnrealBuildTool.cs b/Engine/Source/Programs/UnrealBuildTool/System/UnrealBuildTool.cs
--- a/Engine/Source/Programs/UnrealBuildTool/System/UnrealBuildTool.cs
+++ b/Engine/Source/Programs/UnrealBuildTool/System/UnrealBuildTool.cs
@@ -925,7 +925,7 @@ namespace UnrealBuildTool
 							// If we build w/ bXGEExport true, we didn't REALLY build at this point, 
 							// so don't bother with doing the PrepTargetForDeployment call. 
 							if ((Result == ECompilationResult.Succeeded) && (BuildConfiguration.bDeployAfterCompile == true) && (BuildConfiguration.bXGEExport == false) &&
-								(UEBuildConfiguration.bGenerateManifest == false) && (UEBuildConfiguration.bCleanProject == false))
+								(UEBuildConfiguration.bGenerateManifest == false) && (UEBuildConfiguration.bGenerateExternalFileList == false) && (UEBuildConfiguration.bCleanProject == false))
 							{
 								var DeployHandler = UEBuildDeploy.GetBuildDeploy(CheckPlatform);
 								if (DeployHandler != null)
