Fix for ToolsForCompile_OnMac build error

diff --git a/Engine/Source/Programs/UnrealBuildTool/Mac/MacToolChain.cs b/Engine/Source/Programs/UnrealBuildTool/Mac/MacToolChain.cs
--- a/Engine/Source/Programs/UnrealBuildTool/Mac/MacToolChain.cs
+++ b/Engine/Source/Programs/UnrealBuildTool/Mac/MacToolChain.cs
@@ -731,7 +731,7 @@ namespace UnrealBuildTool
 					}
 				}
 
-				LinkCommand += string.Format(" -filelist \"{0}\"", ConvertPath(ResponsePath));
+				LinkCommand += string.Format(" @\"{0}\"", ConvertPath(ResponsePath));
 			}
 
 			if (LinkEnvironment.Config.bIsBuildingDLL)
