Add Rocket static libraries to the linker using absolute paths, to fix Android linking from a different working directory.

diff --git a/Engine/Source/Programs/UnrealBuildTool/Configuration/UEBuildModule.cs b/Engine/Source/Programs/UnrealBuildTool/Configuration/UEBuildModule.cs
--- a/Engine/Source/Programs/UnrealBuildTool/Configuration/UEBuildModule.cs
+++ b/Engine/Source/Programs/UnrealBuildTool/Configuration/UEBuildModule.cs
@@ -1611,7 +1611,7 @@ namespace UnrealBuildTool
 
 			if(RedistStaticLibraryPath != null)
 			{
-				LinkEnvironment.Config.AdditionalLibraries.Add(RedistStaticLibraryPath);
+				LinkEnvironment.Config.AdditionalLibraries.Add(Path.GetFullPath(RedistStaticLibraryPath));
 			}
 		}
 
