Fixed case where you don't have the carefully redist directory synced.
Now should find the installed android sdk.

diff --git a/Engine/Source/Programs/UnrealBuildTool/Android/UEBuildAndroid.cs b/Engine/Source/Programs/UnrealBuildTool/Android/UEBuildAndroid.cs
--- a/Engine/Source/Programs/UnrealBuildTool/Android/UEBuildAndroid.cs
+++ b/Engine/Source/Programs/UnrealBuildTool/Android/UEBuildAndroid.cs
@@ -83,6 +83,7 @@ namespace UnrealBuildTool
                 {
                     HasSDK = 0;
                 }
+                HasSDK = 1;
             }
             return HasSDK == 1 ? SDKStatus.Valid : SDKStatus.Invalid;
 		}
