Rewrote the sdk checking code for andorid so that it's easier to work with.

diff --git a/Engine/Source/Programs/UnrealBuildTool/Android/UEBuildAndroid.cs b/Engine/Source/Programs/UnrealBuildTool/Android/UEBuildAndroid.cs
--- a/Engine/Source/Programs/UnrealBuildTool/Android/UEBuildAndroid.cs
+++ b/Engine/Source/Programs/UnrealBuildTool/Android/UEBuildAndroid.cs
@@ -57,33 +57,48 @@ namespace UnrealBuildTool
             return "1.1";
         }
 
+        /// <summary>
+        /// checks if the sdk is installed or has been synced
+        /// </summary>
+        /// <returns></returns>
+        private bool HasAnySDK()
+        {
+            if (base.HasRequiredSDKsInstalled() == SDKStatus.Valid)
+            {
+                return true;
+            }
+
+            string NDKPath = Environment.GetEnvironmentVariable("NDKROOT");
+
+            // we don't have an NDKROOT specified
+            if (String.IsNullOrEmpty(NDKPath))
+            {
+                return false;
+            }
+
+            NDKPath = NDKPath.Replace("\"", "");
+
+            // can't find llvm-3.3 or llvm-3.1 in the toolchains
+            if (!Directory.Exists(Path.Combine(NDKPath, @"toolchains\llvm-3.3")) && !Directory.Exists(Path.Combine(NDKPath, @"toolchains\llvm-3.1")))
+            {
+                return false;
+            }
+            return true;
+        }
+
+
 		public override SDKStatus HasRequiredSDKsInstalled()
 		{
             if (HasSDK == -1)
             {
-                if (base.HasRequiredSDKsInstalled() == SDKStatus.Valid)
+                if (HasAnySDK())
                 {
                     HasSDK = 1;
-                    return HasSDK == 1 ? SDKStatus.Valid : SDKStatus.Invalid;
-                }
-                
-                string NDKPath = Environment.GetEnvironmentVariable("NDKROOT");
-
-                // we don't have an NDKROOT specified
-                if (String.IsNullOrEmpty(NDKPath))
-                {
-                    HasSDK = 0;
-                    return HasSDK == 1 ? SDKStatus.Valid : SDKStatus.Invalid;
                 }
-
-                NDKPath = NDKPath.Replace("\"", "");
-
-                // can't find llvm-3.3 or llvm-3.1 in the toolchains
-                if (!Directory.Exists(Path.Combine(NDKPath, @"toolchains\llvm-3.3")) && !Directory.Exists(Path.Combine(NDKPath, @"toolchains\llvm-3.1")))
+                else
                 {
                     HasSDK = 0;
                 }
-                HasSDK = 1;
             }
             return HasSDK == 1 ? SDKStatus.Valid : SDKStatus.Invalid;
 		}
