Fix for absolute paths in compiler response file.

diff --git a/Engine/Source/Programs/UnrealBuildTool/Configuration/UEBuildModule.cs b/Engine/Source/Programs/UnrealBuildTool/Configuration/UEBuildModule.cs
--- a/Engine/Source/Programs/UnrealBuildTool/Configuration/UEBuildModule.cs
+++ b/Engine/Source/Programs/UnrealBuildTool/Configuration/UEBuildModule.cs
@@ -571,7 +571,7 @@ namespace UnrealBuildTool
 				}
 
 				// Add the module's directory to the include path, so we can root #includes to it
-				IncludePaths.Add(ModuleDirectory);
+				IncludePaths.Add(Utils.CleanDirectorySeparators(Utils.MakePathRelativeTo(ModuleDirectory, Path.Combine(ProjectFileGenerator.RootRelativePath, "Engine/Source")), '/'));
 			}
 		}
 		
