UBT: Making iOS SDK version parsing culture invariant

diff --git a/Engine/Source/Programs/UnrealBuildTool/IOS/IOSToolChain.cs b/Engine/Source/Programs/UnrealBuildTool/IOS/IOSToolChain.cs
--- a/Engine/Source/Programs/UnrealBuildTool/IOS/IOSToolChain.cs
+++ b/Engine/Source/Programs/UnrealBuildTool/IOS/IOSToolChain.cs
@@ -158,7 +158,7 @@ namespace UnrealBuildTool
 				}
 			}
 
-			IOSSDKVersionFloat = float.Parse(IOSSDKVersion);
+			IOSSDKVersionFloat = float.Parse(IOSSDKVersion, System.Globalization.CultureInfo.InvariantCulture);
 
 			if (ExternalExecution.GetRuntimePlatform() != UnrealTargetPlatform.Mac)
 			{
