Fix test for foreign projects when generating project files; UEBuildTarget sets the .uproject path for every target in the solution in turn, and it gets set to the last project enumerated. It's not possible to detect foreign projects based on what's in the UBT command line.

diff --git a/Engine/Source/Programs/UnrealBuildTool/System/Project.cs b/Engine/Source/Programs/UnrealBuildTool/System/Project.cs
--- a/Engine/Source/Programs/UnrealBuildTool/System/Project.cs
+++ b/Engine/Source/Programs/UnrealBuildTool/System/Project.cs
@@ -163,6 +163,15 @@ namespace UnrealBuildTool
 			set;
 		}
 
+		/// Returns true if this is a foreign project, and requires UBT to be passed the path to the .uproject file 
+		/// on the command line.
+		public bool IsForeignProject
+		{
+			get;
+			set;
+		}
+
+
 
 		/// All of the targets in this project.  All non-stub projects must have at least one target.
 		public readonly List<ProjectTarget> ProjectTargets = new List<ProjectTarget>();
diff --git a/Engine/Source/Programs/UnrealBuildTool/System/ProjectFileGenerator.cs b/Engine/Source/Programs/UnrealBuildTool/System/ProjectFileGenerator.cs
--- a/Engine/Source/Programs/UnrealBuildTool/System/ProjectFileGenerator.cs
+++ b/Engine/Source/Programs/UnrealBuildTool/System/ProjectFileGenerator.cs
@@ -1614,6 +1614,7 @@ namespace UnrealBuildTool
 
 						bool bProjectAlreadyExisted;
 						var ProjectFile = FindOrAddProject(ProjectFilePath, IncludeInGeneratedProjects: true, bAlreadyExisted: out bProjectAlreadyExisted);
+						ProjectFile.IsForeignProject = bGeneratingGameProjectFiles && UnrealBuildTool.HasUProjectFile() && Utils.IsFileUnderDirectory(TargetFilePath, UnrealBuildTool.GetUProjectPath());
 						ProjectFile.IsGeneratedProject = true;
 						ProjectFile.IsStubProject = false;
 
diff --git a/Engine/Source/Programs/UnrealBuildTool/System/VCProject.cs b/Engine/Source/Programs/UnrealBuildTool/System/VCProject.cs
--- a/Engine/Source/Programs/UnrealBuildTool/System/VCProject.cs
+++ b/Engine/Source/Programs/UnrealBuildTool/System/VCProject.cs
@@ -893,10 +893,9 @@ namespace UnrealBuildTool
 			{
 				return;
 			}
-
+	
 			string UProjectPath = "";
-			bool bIsProjectTarget = UnrealBuildTool.HasUProjectFile() && Utils.IsFileUnderDirectory(TargetFilePath, UnrealBuildTool.GetUProjectPath());
-			if (bIsProjectTarget)
+			if (IsForeignProject)
 			{
 				UProjectPath = "\"$(SolutionDir)$(SolutionName).uproject\"";
 			}
@@ -1015,7 +1014,7 @@ namespace UnrealBuildTool
 					// NMake Build command line
 					VCProjectFileContent.Append("		<NMakeBuildCommandLine>");
 					VCProjectFileContent.Append(EscapePath(NormalizeProjectPath(Path.Combine(BatchFilesDirectoryName, "Build.bat"))) + ProjectPlatformConfiguration.ToString());
-					if (bIsProjectTarget)
+					if (IsForeignProject)
 					{
 						VCProjectFileContent.Append(" " + UProjectPath + (UnrealBuildTool.RunningRocket() ? " -rocket" : ""));
 					}
@@ -1024,7 +1023,7 @@ namespace UnrealBuildTool
 					// NMake ReBuild command line
 					VCProjectFileContent.Append("		<NMakeReBuildCommandLine>");
 					VCProjectFileContent.Append(EscapePath(NormalizeProjectPath(Path.Combine(BatchFilesDirectoryName, "Rebuild.bat"))) + ProjectPlatformConfiguration.ToString());
-					if (bIsProjectTarget)
+					if (IsForeignProject)
 					{
 						VCProjectFileContent.Append(" " + UProjectPath + (UnrealBuildTool.RunningRocket() ? " -rocket" : ""));
 					}
@@ -1033,7 +1032,7 @@ namespace UnrealBuildTool
 					// NMake Clean command line
 					VCProjectFileContent.Append("		<NMakeCleanCommandLine>");
 					VCProjectFileContent.Append(EscapePath(NormalizeProjectPath(Path.Combine(BatchFilesDirectoryName, "Clean.bat"))) + ProjectPlatformConfiguration.ToString());
-					if (bIsProjectTarget)
+					if (IsForeignProject)
 					{
 						VCProjectFileContent.Append(" " + UProjectPath + (UnrealBuildTool.RunningRocket() ? " -rocket" : ""));
 					}
