- Allow non-existant fields in BuildConfiguration.xml, because it's a shared file (the one in Documents) and another branch may add fields

diff --git a/Engine/Source/Programs/UnrealBuildTool/Utilities/XmlConfigLoader.cs b/Engine/Source/Programs/UnrealBuildTool/Utilities/XmlConfigLoader.cs
--- a/Engine/Source/Programs/UnrealBuildTool/Utilities/XmlConfigLoader.cs
+++ b/Engine/Source/Programs/UnrealBuildTool/Utilities/XmlConfigLoader.cs
@@ -603,9 +603,15 @@ namespace UnrealBuildTool
 				{
 					FieldInfo Field = ClassType.GetField(XmlField.Name);
 
-					if (Field == null || !IsConfigurableField(Field))
+					// allow settings in the .xml that don't exist, as another branch may have it, and can share this file from Documents
+					if (Field == null)
 					{
-						throw new BuildException("BuildConfiguration Loading: field '{0}' doesn't exist or is either non-public, non-static or not-xml-configurable.", XmlField.Name);
+						continue;
+					}
+					
+					if (!IsConfigurableField(Field))
+					{
+						throw new BuildException("BuildConfiguration Loading: field '{0}' is either non-public, non-static or not-xml-configurable.", XmlField.Name);
 					}
 
 					if(Field.FieldType.IsArray)
