fixed issue with path from PYTHON environment variable not adding python.exe
TTP339237
#ue4
#html5

diff --git a/Engine/Source/Programs/UnrealBuildTool/HTML5/HTML5ToolChain.cs b/Engine/Source/Programs/UnrealBuildTool/HTML5/HTML5ToolChain.cs
--- a/Engine/Source/Programs/UnrealBuildTool/HTML5/HTML5ToolChain.cs
+++ b/Engine/Source/Programs/UnrealBuildTool/HTML5/HTML5ToolChain.cs
@@ -32,6 +32,10 @@ namespace UnrealBuildTool
 					{
 						PythonPath = Utils.IsRunningOnMono ? "python" : "python.exe";
 					}
+					else
+					{
+						PythonPath += Utils.IsRunningOnMono ? "/python" : "/python.exe";
+					}
                     EMCCPath = "\"" + EMCCPath + "\""; 
 					// set some environment variable we'll need
 					//Environment.SetEnvironmentVariable("EMCC_DEBUG", "cache");
