[GitHub] 244 : UBT Visual Studio Version Autocheck
#ttp 339180

diff --git a/Engine/Source/Programs/UnrealBuildTool/Windows/UEBuildWindows.cs b/Engine/Source/Programs/UnrealBuildTool/Windows/UEBuildWindows.cs
--- a/Engine/Source/Programs/UnrealBuildTool/Windows/UEBuildWindows.cs
+++ b/Engine/Source/Programs/UnrealBuildTool/Windows/UEBuildWindows.cs
@@ -24,8 +24,40 @@ namespace UnrealBuildTool
 
 	public class WindowsPlatform : UEBuildPlatform
 	{
+		/// Property caching.
+		private static WindowsCompiler? CachedCompiler;
+
 		/// Version of the compiler toolchain to use on Windows platform
-		public static readonly WindowsCompiler Compiler = WindowsCompiler.VisualStudio2013;
+		public static WindowsCompiler Compiler
+		{
+			get
+			{
+				// Cache the result because Compiler is often used.
+				if (CachedCompiler.HasValue)
+				{
+					return CachedCompiler.Value;
+				}
+
+				// First check if Visual Sudio 2013 is installed.
+				// If both 2012 and 2013 are installed, prefer 2013
+				if (!String.IsNullOrEmpty(WindowsPlatform.GetVSComnToolsPath(WindowsCompiler.VisualStudio2013)))
+				{
+					CachedCompiler = WindowsCompiler.VisualStudio2013;
+				}
+				// Next try Visual Studio 2012
+				else if (!String.IsNullOrEmpty(WindowsPlatform.GetVSComnToolsPath(WindowsCompiler.VisualStudio2012)))
+				{
+					CachedCompiler = WindowsCompiler.VisualStudio2012;
+				}
+				// Finally assume 2013 is installed to defer errors somewhere else like VCToolChain
+				else
+				{
+					CachedCompiler = WindowsCompiler.VisualStudio2013;
+				}
+
+				return CachedCompiler.Value;
+			}
+		}
 
 		/// True if we are using "clang-cl" to compile instead of MSVC on Windows platform
 		public static readonly bool bCompileWithClang = false;
diff --git a/Engine/Source/Programs/UnrealBuildTool/Windows/VCToolChain.cs b/Engine/Source/Programs/UnrealBuildTool/Windows/VCToolChain.cs
--- a/Engine/Source/Programs/UnrealBuildTool/Windows/VCToolChain.cs
+++ b/Engine/Source/Programs/UnrealBuildTool/Windows/VCToolChain.cs
@@ -1506,18 +1506,7 @@ namespace UnrealBuildTool
 
 			if (string.IsNullOrEmpty(BaseVSToolPath))
 			{
-				if (WindowsPlatform.Compiler == WindowsCompiler.VisualStudio2012 && !String.IsNullOrEmpty(WindowsPlatform.GetVSComnToolsPath(WindowsCompiler.VisualStudio2013)))
-				{
-					throw new BuildException("Visual Studio 2012 must be installed in order to build this target.  However, Visual Studio 2013 was detected.  To compile with Visual Studio 2013, change the 'WindowsPlatform.Compiler' variable in UnrealBuildTool");
-				}
-				else if (WindowsPlatform.Compiler == WindowsCompiler.VisualStudio2013 && !String.IsNullOrEmpty(WindowsPlatform.GetVSComnToolsPath(WindowsCompiler.VisualStudio2012)))
-				{
-					throw new BuildException("Visual Studio 2013 must be installed in order to build this target.  However, Visual Studio 2012 was detected.  To compile with Visual Studio 2012, change the 'WindowsPlatform.Compiler' variable in UnrealBuildTool");
-				}
-				else
-				{
-					throw new BuildException("Visual Studio 2012 or Visual Studio 2013 must be installed in order to build this target.");
-				}
+				throw new BuildException("Visual Studio 2012 or Visual Studio 2013 must be installed in order to build this target.");
 			}
 
 			return BaseVSToolPath;
