Disabled passBuildSettingsInEnvironment for all targets in Xcode project

diff --git a/Engine/Source/Programs/UnrealBuildTool/System/XcodeProjectFileGenerator.cs b/Engine/Source/Programs/UnrealBuildTool/System/XcodeProjectFileGenerator.cs
--- a/Engine/Source/Programs/UnrealBuildTool/System/XcodeProjectFileGenerator.cs
+++ b/Engine/Source/Programs/UnrealBuildTool/System/XcodeProjectFileGenerator.cs
@@ -403,7 +403,7 @@ namespace UnrealBuildTool
 
 			if (Target.Type == XcodeTargetType.Legacy)
 			{
-				Contents.Append("\t\t\tpassBuildSettingsInEnvironment = 1;" + ProjectFileGenerator.NewLine);
+				Contents.Append("\t\t\tpassBuildSettingsInEnvironment = 0;" + ProjectFileGenerator.NewLine);
 			}
 
 			Contents.Append("\t\t\tproductName = \"" + Target.DisplayName + "\";" + ProjectFileGenerator.NewLine);
