fix for location of program plist in xcode project
#ue4
#ios

diff --git a/Engine/Source/Programs/UnrealBuildTool/System/XcodeProjectFileGenerator.cs b/Engine/Source/Programs/UnrealBuildTool/System/XcodeProjectFileGenerator.cs
--- a/Engine/Source/Programs/UnrealBuildTool/System/XcodeProjectFileGenerator.cs
+++ b/Engine/Source/Programs/UnrealBuildTool/System/XcodeProjectFileGenerator.cs
@@ -563,7 +563,7 @@ namespace UnrealBuildTool
 			{
 				Contents.Append(
 					"\t\t\t\tCODE_SIGN_RESOURCE_RULES_PATH = \"" + EngineRelative + "Engine/Build/iOS/XcodeSupportFiles/CustomResourceRules.plist\";" + ProjectFileGenerator.NewLine +
-					"\t\t\t\tINFOPLIST_FILE = \"" + EngineRelative + "Engine/Source/Programs/" + TargetName + "/Resources/IOS/" + TargetName + "-Info.plist\";" + ProjectFileGenerator.NewLine +
+					"\t\t\t\tINFOPLIST_FILE = \"" + EngineRelative + TargetName + "/Intermediate/IOS/" + TargetName + "-Info.plist\";" + ProjectFileGenerator.NewLine +
 					"\t\t\t\tSYMROOT = \"" + EngineRelative + "Engine/Binaries/IOS/Payload\";" + ProjectFileGenerator.NewLine +
 					"\t\t\t\tOBJROOT = \"" + EngineRelative + "Engine/Intermediate/IOS/build\";" + ProjectFileGenerator.NewLine +
 					"\t\t\t\tCONFIGURATION_BUILD_DIR = \"" + EngineRelative + "Engine/Binaries/IOS/Payload\";" + ProjectFileGenerator.NewLine);
