UBT: Clarified message about what code generation is for (now says Reflection Code Generation)

diff --git a/Engine/Source/Programs/UnrealBuildTool/System/ExternalExecution.cs b/Engine/Source/Programs/UnrealBuildTool/System/ExternalExecution.cs
--- a/Engine/Source/Programs/UnrealBuildTool/System/ExternalExecution.cs
+++ b/Engine/Source/Programs/UnrealBuildTool/System/ExternalExecution.cs
@@ -539,7 +539,7 @@ namespace UnrealBuildTool
 						return false;
 					}
 
-					Log.TraceInformation( "Code generation finished for {0} and took {1}", ActualTargetName, (double)s.ElapsedMilliseconds/1000.0 );
+					Log.TraceInformation( "Reflection code generation finished for {0} and took {1}", ActualTargetName, (double)s.ElapsedMilliseconds/1000.0 );
 
 					// Now that UHT has successfully finished generating code, we need to update all cached FileItems in case their last write time has changed.
 					// Otherwise UBT might not detect changes UHT made.
