- SSH/Rsync can now be used in UBT, but is not currently recommended yet (one step is much slower, and IPhonePackager issues)
- Integrated IPhonePackager code directly into IOS Automation (this is not 100% working yet, it's a work-in-progress). Original IPP still in same place untouched
- Added a return code to RunLocalProcess
- Added IPhonePackager UAT script that will pass along params to the internal IPP code (RunUAT IPhonePackager -cmd="<ipp commandline here>" [-nocompile])

diff --git a/Engine/Source/Programs/AutomationTool/IOS/IOS.Automation.csproj b/Engine/Source/Programs/AutomationTool/IOS/IOS.Automation.csproj
--- a/Engine/Source/Programs/AutomationTool/IOS/IOS.Automation.csproj
+++ b/Engine/Source/Programs/AutomationTool/IOS/IOS.Automation.csproj
@@ -34,15 +34,27 @@
     <TreatWarningsAsErrors>true</TreatWarningsAsErrors>
   </PropertyGroup>
   <ItemGroup>
+    <Reference Include="BouncyCastle.Crypto">
+      <HintPath>..\..\..\..\Binaries\ThirdParty\IOS\BouncyCastle.Crypto.dll</HintPath>
+      <Private>False</Private>
+    </Reference>
     <Reference Include="Ionic.Zip.Reduced">
       <HintPath>..\..\..\..\Binaries\DotNET\Ionic.Zip.Reduced.dll</HintPath>
       <Private>False</Private>
     </Reference>
+    <Reference Include="RPCUtility">
+      <HintPath>..\..\..\..\Binaries\DotNET\RPCUtility.exe</HintPath>
+      <Private>False</Private>
+    </Reference>
     <Reference Include="System" />
     <Reference Include="System.Core" />
+    <Reference Include="System.Drawing" />
+    <Reference Include="System.Runtime.Remoting" />
     <Reference Include="System.Runtime.Serialization" />
+    <Reference Include="System.Security" />
     <Reference Include="System.Web" />
     <Reference Include="System.Web.Extensions" />
+    <Reference Include="System.Windows.Forms" />
     <Reference Include="System.Xml.Linq" />
     <Reference Include="System.Data.DataSetExtensions" />
     <Reference Include="Microsoft.CSharp" />
@@ -54,9 +66,60 @@
       <Link>Properties\MetaData.cs</Link>
     </Compile>
     <Compile Include="IOSPlatform.Automation.cs" />
+    <Compile Include="IPhonePackager\CodeSigning.cs" />
+    <Compile Include="IPhonePackager\CompileTime.cs" />
+    <Compile Include="IPhonePackager\Config.cs" />
+    <Compile Include="IPhonePackager\ConfigureMobileGame.cs">
+      <SubType>Form</SubType>
+    </Compile>
+    <Compile Include="IPhonePackager\ConfigureMobileGame.Designer.cs">
+      <DependentUpon>ConfigureMobileGame.cs</DependentUpon>
+    </Compile>
+    <Compile Include="IPhonePackager\CookTime.cs" />
+    <Compile Include="IPhonePackager\DeployTime.cs" />
+    <Compile Include="IPhonePackager\FileOperations.cs" />
+    <Compile Include="IPhonePackager\GenerateSigningRequestDialog.cs">
+      <SubType>Form</SubType>
+    </Compile>
+    <Compile Include="IPhonePackager\GenerateSigningRequestDialog.Designer.cs">
+      <DependentUpon>GenerateSigningRequestDialog.cs</DependentUpon>
+    </Compile>
+    <Compile Include="IPhonePackager\GraphicalResignTool.cs">
+      <SubType>Form</SubType>
+    </Compile>
+    <Compile Include="IPhonePackager\GraphicalResignTool.Designer.cs">
+      <DependentUpon>GraphicalResignTool.cs</DependentUpon>
+    </Compile>
+    <Compile Include="IPhonePackager\iPhonePackager.cs" />
+    <Compile Include="IPhonePackager\MachObjects.cs" />
+    <Compile Include="IPhonePackager\MobileProvisionUtilities.cs" />
+    <Compile Include="IPhonePackager\PasswordDialog.cs">
+      <SubType>Form</SubType>
+    </Compile>
+    <Compile Include="IPhonePackager\PasswordDialog.Designer.cs">
+      <DependentUpon>PasswordDialog.cs</DependentUpon>
+    </Compile>
+    <Compile Include="IPhonePackager\SlowProgressDialog.cs">
+      <SubType>Form</SubType>
+    </Compile>
+    <Compile Include="IPhonePackager\SlowProgressDialog.Designer.cs">
+      <DependentUpon>SlowProgressDialog.cs</DependentUpon>
+    </Compile>
+    <Compile Include="IPhonePackager\ToolsHub.cs">
+      <SubType>Form</SubType>
+    </Compile>
+    <Compile Include="IPhonePackager\ToolsHub.Designer.cs">
+      <DependentUpon>ToolsHub.cs</DependentUpon>
+    </Compile>
+    <Compile Include="IPhonePackager\Utilities.cs" />
     <Compile Include="Properties\AssemblyInfo.cs" />
+    <Compile Include="Properties\Resources.Designer.cs" />
   </ItemGroup>
   <ItemGroup>
+    <ProjectReference Include="..\..\IOS\DeploymentInterface\DeploymentInterface.csproj">
+      <Project>{2da5397a-f93f-41e0-a1ad-b220d40fd739}</Project>
+      <Name>DeploymentInterface</Name>
+    </ProjectReference>
     <ProjectReference Include="..\AutomationTool.csproj">
       <Project>{a35b0261-d868-475e-a098-eb14ce6dc89f}</Project>
       <Name>AutomationTool</Name>
@@ -68,5 +131,31 @@
       <Private>False</Private>
     </ProjectReference>
   </ItemGroup>
+  <ItemGroup>
+    <EmbeddedResource Include="IPhonePackager\ConfigureMobileGame.resx">
+      <DependentUpon>ConfigureMobileGame.cs</DependentUpon>
+    </EmbeddedResource>
+    <EmbeddedResource Include="IPhonePackager\GenerateSigningRequestDialog.resx">
+      <DependentUpon>GenerateSigningRequestDialog.cs</DependentUpon>
+    </EmbeddedResource>
+    <EmbeddedResource Include="IPhonePackager\GraphicalResignTool.resx">
+      <DependentUpon>GraphicalResignTool.cs</DependentUpon>
+    </EmbeddedResource>
+    <EmbeddedResource Include="IPhonePackager\PasswordDialog.resx">
+      <DependentUpon>PasswordDialog.cs</DependentUpon>
+    </EmbeddedResource>
+    <EmbeddedResource Include="IPhonePackager\SlowProgressDialog.resx">
+      <DependentUpon>SlowProgressDialog.cs</DependentUpon>
+    </EmbeddedResource>
+    <EmbeddedResource Include="IPhonePackager\ToolsHub.resx">
+      <DependentUpon>ToolsHub.cs</DependentUpon>
+    </EmbeddedResource>
+    <EmbeddedResource Include="Properties\Resources.resx">
+      <SubType>Designer</SubType>
+    </EmbeddedResource>
+  </ItemGroup>
+  <ItemGroup>
+    <None Include="Resources\Image1.bmp" />
+  </ItemGroup>
   <Import Project="$(MSBuildToolsPath)\Microsoft.CSharp.targets" />
 </Project>
\ No newline at end of file
diff --git a/Engine/Source/Programs/AutomationTool/IOS/IOSPlatform.Automation.cs b/Engine/Source/Programs/AutomationTool/IOS/IOSPlatform.Automation.cs
--- a/Engine/Source/Programs/AutomationTool/IOS/IOSPlatform.Automation.cs
+++ b/Engine/Source/Programs/AutomationTool/IOS/IOSPlatform.Automation.cs
@@ -6,6 +6,7 @@ using System.Text;
 using System.IO;
 using AutomationTool;
 using UnrealBuildTool;
+using System.Text.RegularExpressions;
 using Ionic.Zip;
 using Ionic.Zlib;
 
@@ -22,6 +23,65 @@ public class IOSPlatform : Platform
 	{
 	}
 
+	// Run the integrated IPP code
+	// @todo ipp: How to log the output in a nice UAT way?
+	protected static int RunIPP(string IPPArguments)
+	{
+		List<string> Args = new List<string>();
+
+		StringBuilder Token = null;
+		int Index = 0;
+		bool bInQuote = false;
+		while (Index < IPPArguments.Length)
+		{
+			if (IPPArguments[Index] == '\"')
+			{
+				bInQuote = !bInQuote;
+			}
+			else if (IPPArguments[Index] == ' ' && !bInQuote)
+			{
+				if (Token != null)
+				{
+					Args.Add(Token.ToString());
+					Token = null;
+				}
+			}
+			else
+			{
+				if (Token == null)
+				{
+					Token = new StringBuilder();
+				}
+				Token.Append(IPPArguments[Index]);
+			}
+			Index++;
+		}
+
+		if (Token != null)
+		{
+			Args.Add(Token.ToString());
+		}
+	
+		return RunIPP(Args.ToArray());
+	}
+
+	protected static int RunIPP(string[] Args)
+	{
+		return iPhonePackager.Program.RealMain(Args);
+	}
+
+	public override int RunCommand(string Command)
+	{
+		// run generic IPP commands through the interface
+		if (Command.StartsWith("IPP:"))
+		{
+			return RunIPP(Command.Substring(4));
+		}
+
+		// non-zero error code
+		return 4;
+	}
+
 	protected string MakeIPAFileName( UnrealTargetConfiguration TargetConfiguration, string ProjectGameExeFilename )
 	{
 		var ProjectIPA = Path.ChangeExtension(ProjectGameExeFilename, null);
@@ -124,47 +184,93 @@ public class IOSPlatform : Platform
 
 			bool cookonthefly = Params.CookOnTheFly || Params.SkipCookOnTheFly;
 
-			string IPPArguments = "RepackageFromStage \"" + (Params.IsCodeBasedProject ? Params.RawProjectPath : "Engine") + "\"";
-			IPPArguments += " -config " + TargetConfiguration.ToString();
-
-			if (TargetConfiguration == UnrealTargetConfiguration.Shipping)
+			// rename the .ipa if not code based
+			if (!Params.IsCodeBasedProject)
 			{
-				IPPArguments += " -compress=best";
+				ProjectIPA = Path.Combine(Path.GetDirectoryName(Params.RawProjectPath), "Binaries", "IOS", Params.ShortProjectName + ".ipa");
+				if (TargetConfiguration != UnrealTargetConfiguration.Development)
+				{
+					ProjectIPA = Path.Combine(Path.GetDirectoryName(Params.RawProjectPath), "Binaries", "IOS", Params.ShortProjectName + "-" + PlatformType.ToString() + "-" + TargetConfiguration.ToString() + ".ipa");
+				}
 			}
 
-			// Determine if we should sign
-			bool bNeedToSign = GetCodeSignDesirability(Params);
+			// delete the .ipa to make sure it was made
+			DeleteFile(ProjectIPA);
 
-			if (!String.IsNullOrEmpty(Params.BundleName))
+			if (RemoteToolChain.bUseRPCUtil)
 			{
-				// Have to sign when a bundle name is specified
-				bNeedToSign = true;
-				IPPArguments += " -bundlename " + Params.BundleName;
-			}
+				string IPPArguments = "RepackageFromStage \"" + (Params.IsCodeBasedProject ? Params.RawProjectPath : "Engine") + "\"";
+				IPPArguments += " -config " + TargetConfiguration.ToString();
 
-			if (bNeedToSign)
-			{
-				IPPArguments += " -sign";
-			}
+				if (TargetConfiguration == UnrealTargetConfiguration.Shipping)
+				{
+					IPPArguments += " -compress=best";
+				}
 
-			IPPArguments += (cookonthefly ? " -cookonthefly" : "");
-			IPPArguments += " -stagedir \"" + CombinePaths(Params.BaseStageDirectory, "IOS") + "\"";
-			IPPArguments += " -projectdir \"" + Path.GetDirectoryName(Params.RawProjectPath) + "\"";
+				// Determine if we should sign
+				bool bNeedToSign = GetCodeSignDesirability(Params);
 
-			// rename the .ipa if not code based
-			if (!Params.IsCodeBasedProject)
-			{
-				ProjectIPA = Path.Combine(Path.GetDirectoryName(Params.RawProjectPath), "Binaries", "IOS", Params.ShortProjectName + ".ipa");
-				if (TargetConfiguration != UnrealTargetConfiguration.Development)
+				if (!String.IsNullOrEmpty(Params.BundleName))
 				{
-					ProjectIPA = Path.Combine(Path.GetDirectoryName(Params.RawProjectPath), "Binaries", "IOS", Params.ShortProjectName + "-" + PlatformType.ToString() + "-" + TargetConfiguration.ToString() + ".ipa");
+					// Have to sign when a bundle name is specified
+					bNeedToSign = true;
+					IPPArguments += " -bundlename " + Params.BundleName;
 				}
+
+				if (bNeedToSign)
+				{
+					IPPArguments += " -sign";
+				}
+
+				IPPArguments += (cookonthefly ? " -cookonthefly" : "");
+				IPPArguments += " -stagedir \"" + CombinePaths(Params.BaseStageDirectory, "IOS") + "\"";
+				IPPArguments += " -projectdir \"" + Path.GetDirectoryName(Params.RawProjectPath) + "\"";
+
+				RunAndLog(CmdEnv, IPPExe, IPPArguments);
 			}
+			else
+			{
+				List<string> IPPArguments = new List<string>();
+				IPPArguments.Add("RepackageFromStage");
+				IPPArguments.Add(Params.IsCodeBasedProject ? Params.RawProjectPath : "Engine");
+				IPPArguments.Add("-config");
+				IPPArguments.Add(TargetConfiguration.ToString());
+
+				if (TargetConfiguration == UnrealTargetConfiguration.Shipping)
+				{
+					IPPArguments.Add("-compress=best");
+				}
 
-            // delete the .ipa to make sure it was made
-            DeleteFile(ProjectIPA);
+				// Determine if we should sign
+				bool bNeedToSign = GetCodeSignDesirability(Params);
 
-			RunAndLog(CmdEnv, IPPExe, IPPArguments);
+				if (!String.IsNullOrEmpty(Params.BundleName))
+				{
+					// Have to sign when a bundle name is specified
+					bNeedToSign = true;
+					IPPArguments.Add("-bundlename");
+					IPPArguments.Add(Params.BundleName);
+				}
+
+				if (bNeedToSign)
+				{
+					IPPArguments.Add("-sign");
+				}
+
+				if (cookonthefly)
+				{
+					IPPArguments.Add(" -cookonthefly");
+				}
+				IPPArguments.Add(" -stagedir");
+				IPPArguments.Add(CombinePaths(Params.BaseStageDirectory, "IOS"));
+				IPPArguments.Add(" -projectdir");
+				IPPArguments.Add(Path.GetDirectoryName(Params.RawProjectPath));
+
+				if (RunIPP(IPPArguments.ToArray()) != 0)
+				{
+					throw new AutomationException("IPP Failed");
+				}
+			}
 
 			// verify the .ipa exists
 			if (!FileExists(ProjectIPA))
diff --git a/Engine/Source/Programs/AutomationTool/IOS/IPhonePackager/CodeSigning.cs b/Engine/Source/Programs/AutomationTool/IOS/IPhonePackager/CodeSigning.cs
--- a/Engine/Source/Programs/AutomationTool/IOS/IPhonePackager/CodeSigning.cs
+++ b/Engine/Source/Programs/AutomationTool/IOS/IPhonePackager/CodeSigning.cs
@@ -11,6 +11,7 @@ using System.IO;
 using MachObjectHandling;
 using System.Diagnostics;
 using System.Xml;
+using UnrealBuildTool;
 
 namespace iPhonePackager
 {
@@ -103,10 +104,10 @@ namespace iPhonePackager
 			return bOverridesExists && (Provision != null) && (Cert != null);
 		}
 
-		protected virtual byte[] GetMobileProvision(string CFBundleIdentifier)
+		protected virtual byte[] GetMobileProvision()
 		{
 			// find the movile provision file in the library
-			string MobileProvisionFilename = MobileProvision.FindCompatibleProvision(CFBundleIdentifier);
+			string MobileProvisionFilename = MobileProvision.FindCompatibleProvision();
 
 			byte[] Result = null;
 			try
@@ -122,9 +123,9 @@ namespace iPhonePackager
 			return Result;
 		}
 
-		public void LoadMobileProvision(string CFBundleIdentifier)
+		public void LoadMobileProvision()
 		{
-			byte[] MobileProvisionFile = GetMobileProvision(CFBundleIdentifier);
+			byte[] MobileProvisionFile = GetMobileProvision();
 
 			if (MobileProvisionFile != null)
 			{
@@ -290,15 +291,8 @@ namespace iPhonePackager
 			// Load Info.plist, which guides nearly everything else
 			Info = LoadInfoPList();
 
-			// Get the name of the bundle
-			string CFBundleIdentifier;
-			if (!Info.GetString("CFBundleIdentifier", out CFBundleIdentifier))
-			{
-				throw new InvalidDataException("Info.plist must contain the key CFBundleIdentifier");
-			}
-
 			// Load the mobile provision, which provides entitlements and a partial cert which can be used to find an installed certificate
-			LoadMobileProvision(CFBundleIdentifier);
+			LoadMobileProvision();
 			if (Provision == null)
 			{
 				return;
@@ -306,8 +300,9 @@ namespace iPhonePackager
 
 			// Install the Apple trust chain certs (required to do a CMS signature with full chain embedded)
 			List<string> TrustChainCertFilenames = new List<string>();
-			TrustChainCertFilenames.Add("AppleWorldwideDeveloperRelationsCA.pem");
-			TrustChainCertFilenames.Add("AppleRootCA.pem");
+			string CertPath = Path.GetFullPath(Path.Combine("Engine/Build/IOS"));
+			TrustChainCertFilenames.Add(Path.Combine(CertPath, "AppleWorldwideDeveloperRelationsCA.pem"));
+			TrustChainCertFilenames.Add(Path.Combine(CertPath, "AppleRootCA.pem"));
 			InstallCertificates(TrustChainCertFilenames);
 
 			// Find and load the signing cert
diff --git a/Engine/Source/Programs/AutomationTool/IOS/IPhonePackager/CompileTime.cs b/Engine/Source/Programs/AutomationTool/IOS/IPhonePackager/CompileTime.cs
--- a/Engine/Source/Programs/AutomationTool/IOS/IPhonePackager/CompileTime.cs
+++ b/Engine/Source/Programs/AutomationTool/IOS/IPhonePackager/CompileTime.cs
@@ -3,6 +3,7 @@
  */
 
 using System;
+using System.Collections;
 using System.Collections.Generic;
 using System.Diagnostics;
 using System.IO;
@@ -12,6 +13,7 @@ using System.Text;
 using System.Threading;
 using System.Windows.Forms;
 using Ionic.Zlib;
+using UnrealBuildTool;
 
 namespace iPhonePackager
 {
@@ -102,7 +104,7 @@ namespace iPhonePackager
 			iPhone_SigningDevRootMac = "/UE4/Builds";
 
 			// get the path to mirror into on the Mac
-			string BinariesDir = Path.GetFullPath(Path.Combine(Environment.CurrentDirectory, @"..\.."));
+			string BinariesDir = Path.GetFullPath(Path.Combine(Environment.CurrentDirectory, "Engine\\Binaries"));
 			string Root = Path.GetPathRoot(BinariesDir);
 			string BranchPath = MachineName + "/" + BinariesDir.Substring( Root.Length );
 			BranchPath = BranchPath.Replace('\\', '/');
@@ -196,7 +198,7 @@ namespace iPhonePackager
 			// Copy the no sign resource rules file over
 			if (!Config.bForDistribution)
 			{
-				FileOperations.CopyRequiredFile(@"..\..\..\Build\IOS\XcodeSupportFiles\CustomResourceRules.plist", Path.Combine(Config.PCStagingRootDir, "CustomResourceRules.plist"));
+				FileOperations.CopyRequiredFile(Config.RootRelativePath + @"Engine\Build\IOS\XcodeSupportFiles\CustomResourceRules.plist", Path.Combine(Config.PCStagingRootDir, "CustomResourceRules.plist"));
 			}
 
 			// Copy the mobile provision file over
@@ -222,7 +224,7 @@ namespace iPhonePackager
 			FileOperations.CopyRequiredFile(Config.RootRelativePath + @"Engine\Intermediate\IOS\UE4.xcodeproj\project.pbxproj", Path.Combine(Config.PCXcodeStagingDir, @"project.pbxproj.datecheck"));
 			
 			// needs Mac line endings so it can be executed
-			string SrcPath = @"..\..\..\Build\IOS\XcodeSupportFiles\prepackage.sh";
+			string SrcPath = Config.RootRelativePath + @"Engine\Build\IOS\XcodeSupportFiles\prepackage.sh";
 			string DestPath = Path.Combine(Config.PCXcodeStagingDir, @"prepackage.sh");
 			Program.Log(" ... '" + SrcPath + "' -> '" + DestPath + "'");
 			string SHContents = File.ReadAllText(SrcPath);
@@ -248,7 +250,7 @@ namespace iPhonePackager
 			case "deletemacstagingfiles":
 				Program.Log( " ... deleting staging files on the Mac" );
 				DisplayCommandLine = "rm -rf Payload";
-				CommandLine = "\"" + MacStagingRootDir + "\" " + DisplayCommandLine;
+				CommandLine = DisplayCommandLine;
 				WorkingFolder = MacStagingRootDir;
 				break;
 
@@ -257,7 +259,7 @@ namespace iPhonePackager
 
 				DisplayCommandLine = String.Format("mkdir -p ~/Library/MobileDevice/Provisioning\\ Profiles");
 
-				CommandLine = "\"" + MacXcodeStagingDir + "\" " + DisplayCommandLine;
+				CommandLine = DisplayCommandLine;
 				WorkingFolder = MacXcodeStagingDir;
 				break;
 
@@ -267,14 +269,14 @@ namespace iPhonePackager
 
 				DisplayCommandLine = String.Format("cp -f {0} ~/Library/MobileDevice/Provisioning\\ Profiles", MacMobileProvisionFilename);
 
-				CommandLine = "\"" + MacXcodeStagingDir + "\" " + DisplayCommandLine;
+				CommandLine = DisplayCommandLine;
 				WorkingFolder = MacXcodeStagingDir;
 				break;
 
 			case "removeprovision":
 				Program.Log(" ... removing .mobileprovision");
 				DisplayCommandLine = String.Format("rm -f ~/Library/MobileDevice/Provisioning\\ Profiles/{0}", MacMobileProvisionFilename);
-				CommandLine = "\"" + MacXcodeStagingDir + "\" " + DisplayCommandLine;
+				CommandLine = DisplayCommandLine;
 				WorkingFolder = MacXcodeStagingDir;
 				break;
 
@@ -282,14 +284,14 @@ namespace iPhonePackager
 				// Note: The executable must have already been copied over
 				Program.Log(" ... setting executable bit");
 				DisplayCommandLine = "chmod a+x \'" + RemoteExecutablePath + "\'";
-				CommandLine = "\"" + MacStagingRootDir + "\" " + DisplayCommandLine;
+				CommandLine = DisplayCommandLine;
 				WorkingFolder = MacStagingRootDir;
 				break;
 
 			case "prepackage":
 				Program.Log(" ... running prepackage script remotely ");
 				DisplayCommandLine = String.Format("sh prepackage.sh {0} IOS {1} {2}", Program.GameName, Program.GameConfiguration, Program.Architecture);
-				CommandLine = "\"" + MacXcodeStagingDir + "\" " + DisplayCommandLine;
+				CommandLine = DisplayCommandLine;
 				WorkingFolder = MacXcodeStagingDir;
 				break;
 
@@ -297,42 +299,42 @@ namespace iPhonePackager
 				Program.Log(" ... making application (codesign, etc...)");
 				Program.Log("  Using signing identity '{0}'", Config.CodeSigningIdentity);
 				DisplayCommandLine = CurrentBaseXCodeCommandLine;
-				CommandLine = "\"" + MacXcodeStagingDir + "/..\" " + DisplayCommandLine;
+				CommandLine = DisplayCommandLine;
 				WorkingFolder = "\"" + MacXcodeStagingDir + "/..\"";
 				break;
 
 			case "validation":
 				Program.Log( " ... validating distribution package" );
 				DisplayCommandLine = XcodeDeveloperDir + "Platforms/iPhoneOS.platform/Developer/usr/bin/Validation " + RemoteAppDirectory;
-				CommandLine = "\"" + MacStagingRootDir + "\" " + DisplayCommandLine;
+				CommandLine = DisplayCommandLine;
 				WorkingFolder = MacStagingRootDir;
 				break;
 
 			case "deleteipa":
 				Program.Log(" ... deleting IPA on Mac");
 				DisplayCommandLine = "rm -f " + Config.IPAFilenameOnMac;
-				CommandLine = "\"" + MacStagingRootDir + "\" " + DisplayCommandLine;
+				CommandLine = DisplayCommandLine;
 				WorkingFolder = MacStagingRootDir;
 				break;
 
 			case "kill":
 				Program.Log( " ... killing" );
 				DisplayCommandLine = "killall " + Program.GameName;
-				CommandLine = ". " + DisplayCommandLine;
+				CommandLine = DisplayCommandLine;
 				WorkingFolder = ".";
 				break;
 
 			case "strip":
 				Program.Log( " ... stripping" );
-				DisplayCommandLine = XcodeDeveloperDir + "Platforms/iPhoneOS.platform/Developer/usr/bin/strip '" + RemoteExecutablePath + "'";
-				CommandLine = "\"" + MacStagingRootDir + "\" " + DisplayCommandLine;
+				DisplayCommandLine = XcodeDeveloperDir + "Platforms/iPhoneOS.platform/Developer/usr/bin/strip " + RemoteExecutablePath;
+				CommandLine = DisplayCommandLine;
 				WorkingFolder = MacStagingRootDir;
 				break;
 
 			case "resign":
 				Program.Log("... resigning");
 				DisplayCommandLine = "bash -c '" + "chmod a+x ResignScript" + ";" + "./ResignScript" + "'";
-				CommandLine = "\"" + MacStagingRootDir + "\" " + DisplayCommandLine;
+				CommandLine = DisplayCommandLine;
 				WorkingFolder = MacStagingRootDir;
 				break;
 
@@ -349,7 +351,7 @@ namespace iPhonePackager
 					Config.IPAFilenameOnMac,
 					dSYMName);
 
-				CommandLine = "\"" + MacStagingRootDir + "\" " + DisplayCommandLine;
+				CommandLine = DisplayCommandLine;
 				WorkingFolder = MacStagingRootDir;
 				break;
 
@@ -360,7 +362,7 @@ namespace iPhonePackager
 				string dSYMPath = Program.GameName + ".app.dSYM";
 				DisplayCommandLine = String.Format("dsymutil -o {0} {1}", dSYMPath, ExePath);
 
-				CommandLine = "\"" + MacStagingRootDir + "\"" + DisplayCommandLine;
+				CommandLine = DisplayCommandLine;
 				WorkingFolder = MacStagingRootDir;
 				break;
 
@@ -373,36 +375,48 @@ namespace iPhonePackager
 			Program.Log( " ... " + DisplayCommandLine );
 			Program.Log(" ... full command: " +  MacName + " " + CommandLine);
 
-			Process RPCUtil = new Process();
-			RPCUtil.StartInfo.FileName = @"..\RPCUtility.exe";
-			RPCUtil.StartInfo.UseShellExecute = false;
-			RPCUtil.StartInfo.Arguments = MacName + " " + CommandLine;
-			RPCUtil.StartInfo.RedirectStandardOutput = true;
-			RPCUtil.StartInfo.RedirectStandardError = true;
-			RPCUtil.OutputDataReceived += new DataReceivedEventHandler(OutputReceivedRemoteProcessCall);
-			RPCUtil.ErrorDataReceived += new DataReceivedEventHandler(OutputReceivedRemoteProcessCall);
+			Hashtable Results = RPCUtilHelper.Command(WorkingFolder, CommandLine, null);
 
-			RPCUtil.Start();
-			
-			RPCUtil.BeginOutputReadLine();
-			RPCUtil.BeginErrorReadLine();
-
-			RPCUtil.WaitForExit();
+			string Output = Results["CommandOutput"] as string;
+			Program.Log("[RPC] " + Output);
 
-			if (RPCUtil.ExitCode != 0)
+			if (Output.Contains("** BUILD FAILED **"))
 			{
-				Program.Error("RPCCommand {0} failed with return code {1}", RPCCommand, RPCUtil.ExitCode);
-				switch (RPCCommand.ToLowerInvariant())
-				{
-					case "installprovision":
-						Program.Error("Ensure your access permissions for '~/Library/MobileDevice/Provisioning Profiles' are set correctly.");
-						break;
-					default:
-						break;
-				}
+				Program.Error("Xcode build failed!");
 			}
 
-			return (RPCUtil.ExitCode == 0);
+			return (Int64)Results["ExitCode"] == 0;
+
+// 			Process RPCUtil = new Process();
+// 			RPCUtil.StartInfo.FileName = @"..\RPCUtility.exe";
+// 			RPCUtil.StartInfo.UseShellExecute = false;
+// 			RPCUtil.StartInfo.Arguments = MacName + " " + CommandLine;
+// 			RPCUtil.StartInfo.RedirectStandardOutput = true;
+// 			RPCUtil.StartInfo.RedirectStandardError = true;
+// 			RPCUtil.OutputDataReceived += new DataReceivedEventHandler(OutputReceivedRemoteProcessCall);
+// 			RPCUtil.ErrorDataReceived += new DataReceivedEventHandler(OutputReceivedRemoteProcessCall);
+// 
+// 			RPCUtil.Start();
+// 			
+// 			RPCUtil.BeginOutputReadLine();
+// 			RPCUtil.BeginErrorReadLine();
+// 
+// 			RPCUtil.WaitForExit();
+// 
+// 			if (RPCUtil.ExitCode != 0)
+// 			{
+// 				Program.Error("RPCCommand {0} failed with return code {1}", RPCCommand, RPCUtil.ExitCode);
+// 				switch (RPCCommand.ToLowerInvariant())
+// 				{
+// 					case "installprovision":
+// 						Program.Error("Ensure your access permissions for '~/Library/MobileDevice/Provisioning Profiles' are set correctly.");
+// 						break;
+// 					default:
+// 						break;
+// 				}
+// 			}
+// 
+// 			return (RPCUtil.ExitCode == 0);
 		}
 
 
diff --git a/Engine/Source/Programs/AutomationTool/IOS/IPhonePackager/GraphicalResignTool.cs b/Engine/Source/Programs/AutomationTool/IOS/IPhonePackager/GraphicalResignTool.cs
--- a/Engine/Source/Programs/AutomationTool/IOS/IPhonePackager/GraphicalResignTool.cs
+++ b/Engine/Source/Programs/AutomationTool/IOS/IPhonePackager/GraphicalResignTool.cs
@@ -247,7 +247,7 @@ namespace iPhonePackager
 			/// </summary>
 			public byte[] CustomInfoPList = null;
 
-			protected override byte[] GetMobileProvision(string CFBundleIdentifier)
+			protected override byte[] GetMobileProvision()
 			{
 				return (CustomMobileProvision != null) ? CustomMobileProvision : FileSystem.ReadAllBytes("embedded.mobileprovision");
 			}
diff --git a/Engine/Source/Programs/AutomationTool/IOS/IPhonePackager/MobileProvisionUtilities.cs b/Engine/Source/Programs/AutomationTool/IOS/IPhonePackager/MobileProvisionUtilities.cs
--- a/Engine/Source/Programs/AutomationTool/IOS/IPhonePackager/MobileProvisionUtilities.cs
+++ b/Engine/Source/Programs/AutomationTool/IOS/IPhonePackager/MobileProvisionUtilities.cs
@@ -28,7 +28,7 @@ namespace iPhonePackager
 		public bool bDebug;
 		public Utilities.PListHelper Data;
 
-		public static string FindCompatibleProvision(string CFBundleIdentifier)
+		public static string FindCompatibleProvision()
 		{
 			// remap the gamename if necessary
 			string GameName = Program.GameName;
@@ -47,30 +47,6 @@ namespace iPhonePackager
 				Directory.CreateDirectory(Config.ProvisionDirectory);
 			}
 
-			#region remove after we provide an install mechanism
-			// copy all of the provisions from the game directory to the library
-			foreach (string Provision in Directory.EnumerateFiles(Config.ProjectRootDirectory + "/Build/IOS/", "*.mobileprovision", SearchOption.AllDirectories))
-			{
-				if (!File.Exists(Config.ProvisionDirectory + Path.GetFileName(Provision)))
-				{
-					File.Copy(Provision, Config.ProvisionDirectory + Path.GetFileName(Provision), true);
-					FileInfo DestFileInfo = new FileInfo(Config.ProvisionDirectory + Path.GetFileName(Provision));
-					DestFileInfo.Attributes = DestFileInfo.Attributes & ~FileAttributes.ReadOnly;
-				}
-			}
-
-			// copy all of the provisions from the egine directory to the library
-			foreach (string Provision in Directory.EnumerateFiles(Config.EngineBuildDirectory, "*.mobileprovision", SearchOption.AllDirectories))
-			{
-				if (!File.Exists(Config.ProvisionDirectory + Path.GetFileName(Provision)))
-				{
-					File.Copy(Provision, Config.ProvisionDirectory + Path.GetFileName(Provision), true);
-					FileInfo DestFileInfo = new FileInfo(Config.ProvisionDirectory + Path.GetFileName(Provision));
-					DestFileInfo.Attributes = DestFileInfo.Attributes & ~FileAttributes.ReadOnly;
-				}
-			}
-			#endregion
-
 			// cache the provision library
 			Dictionary<string, MobileProvision> ProvisionLibrary = new Dictionary<string, MobileProvision>();
 			foreach (string Provision in Directory.EnumerateFiles(Config.ProvisionDirectory, "*.mobileprovision"))
@@ -82,7 +58,7 @@ namespace iPhonePackager
 			// check the cache for a provision matching the app id (com.company.Game)
 			foreach (KeyValuePair<string, MobileProvision> Pair in ProvisionLibrary)
 			{
-				if (Pair.Value.ApplicationIdentifier.Contains(CFBundleIdentifier) && (!Config.bForDistribution || (Pair.Value.ProvisionedDeviceIDs.Count == 0 && !Pair.Value.bDebug)))
+				if ((Pair.Value.ProvisionName == GameName || Pair.Value.ApplicationIdentifier.Contains(GameName)) && (!Config.bForDistribution || (Pair.Value.ProvisionedDeviceIDs.Count == 0 && !Pair.Value.bDebug)))
 				{
 					// check to see if we have a certificate for this provision
 					if (CodeSignatureBuilder.FindCertificate(Pair.Value) != null)
@@ -92,6 +68,73 @@ namespace iPhonePackager
 				}
 			}
 
+			#region remove after we provide an install mechanism
+			// check to see if we can find a provision in the game's build directory
+			string MobileProvisionFilename = FileOperations.FindPrefixedFile(Config.ProjectRootDirectory + "/Build/IOS/", GameName + ".mobileprovision");
+
+			if (!File.Exists(MobileProvisionFilename))
+			{
+				MobileProvisionFilename = FileOperations.FindPrefixedFile(Config.ProjectRootDirectory + "/Build/IOS/NotForLicensees/", GameName + ".mobileprovision");
+				if (!File.Exists(MobileProvisionFilename))
+				{
+					MobileProvisionFilename = FileOperations.FindAnyFileWithExtension(Config.ProjectRootDirectory + "/Build/IOS/", ".mobileprovision");
+				}
+			}
+			if (File.Exists(MobileProvisionFilename))
+			{
+				Directory.CreateDirectory(Config.ProvisionDirectory);
+				File.Copy(MobileProvisionFilename, Config.ProvisionDirectory + Path.GetFileName(MobileProvisionFilename), true);
+				FileInfo DestFileInfo = new FileInfo(Config.ProvisionDirectory + Path.GetFileName(MobileProvisionFilename));
+				DestFileInfo.Attributes = DestFileInfo.Attributes & ~FileAttributes.ReadOnly;
+				MobileProvisionFilename = Config.ProvisionDirectory + Path.GetFileName(MobileProvisionFilename);
+
+				// verify the provision is good
+				MobileProvision p = MobileProvisionParser.ParseFile(MobileProvisionFilename);
+				if ((p.ProvisionName == GameName || p.ApplicationIdentifier.Contains(GameName)) && (!Config.bForDistribution || (p.ProvisionedDeviceIDs.Count == 0 && !p.bDebug)))
+				{
+					// check to see if we have a certificate for this provision
+					if (CodeSignatureBuilder.FindCertificate(p) != null)
+					{
+						return MobileProvisionFilename;
+					}
+				}
+			}
+
+			// check for a distribution provision in the game's build directory
+			if (Config.bForDistribution)
+			{
+				MobileProvisionFilename = FileOperations.FindPrefixedFile(Config.ProjectRootDirectory + "/Build/IOS/", GameName + "_Distro.mobileprovision");
+
+				if (!File.Exists(MobileProvisionFilename))
+				{
+					MobileProvisionFilename = FileOperations.FindPrefixedFile(Config.ProjectRootDirectory + "/Build/IOS/NotForLicensees/", GameName + "_Distro.mobileprovision");
+					if (!File.Exists(MobileProvisionFilename))
+					{
+						MobileProvisionFilename = FileOperations.FindAnyFileWithExtension(Config.ProjectRootDirectory + "/Build/IOS/", "_Distro.mobileprovision");
+					}
+				}
+				if (File.Exists(MobileProvisionFilename))
+				{
+					Directory.CreateDirectory(Config.ProvisionDirectory);
+					File.Copy(MobileProvisionFilename, Config.ProvisionDirectory + Path.GetFileName(MobileProvisionFilename), true);
+					FileInfo DestFileInfo = new FileInfo(Config.ProvisionDirectory + Path.GetFileName(MobileProvisionFilename));
+					DestFileInfo.Attributes = DestFileInfo.Attributes & ~FileAttributes.ReadOnly;
+					MobileProvisionFilename = Config.ProvisionDirectory + Path.GetFileName(MobileProvisionFilename);
+
+					// verify the provision is good
+					MobileProvision p = MobileProvisionParser.ParseFile(MobileProvisionFilename);
+					if ((p.ProvisionName == GameName || p.ApplicationIdentifier.Contains(GameName)) && (!Config.bForDistribution || (p.ProvisionedDeviceIDs.Count == 0 && !p.bDebug)))
+					{
+						// check to see if we have a certificate for this provision
+						if (CodeSignatureBuilder.FindCertificate(p) != null)
+						{
+							return MobileProvisionFilename;
+						}
+					}
+				}
+			}
+			#endregion
+
 			// check the cache for a provision matching the wild card
 			foreach (KeyValuePair<string, MobileProvision> Pair in ProvisionLibrary)
 			{
@@ -105,8 +148,75 @@ namespace iPhonePackager
 				}
 			}
 
+			#region remove after we provide an install mechanism
+			// check to see if we can find a provision in the engine's build directory
+			MobileProvisionFilename = FileOperations.FindPrefixedFile(Config.EngineBuildDirectory, "UE4Game.mobileprovision");
+
+			if (!File.Exists(MobileProvisionFilename))
+			{
+				MobileProvisionFilename = FileOperations.FindPrefixedFile(Config.EngineBuildDirectory + "/NotForLicensees/", "UE4Game.mobileprovision");
+				if (!File.Exists(MobileProvisionFilename))
+				{
+					MobileProvisionFilename = FileOperations.FindAnyFileWithExtension(Config.EngineBuildDirectory, ".mobileprovision");
+				}
+			}
+			if (File.Exists(MobileProvisionFilename))
+			{
+				Directory.CreateDirectory(Config.ProvisionDirectory);
+				File.Copy(MobileProvisionFilename, Config.ProvisionDirectory + Path.GetFileName(MobileProvisionFilename), true);
+				FileInfo DestFileInfo = new FileInfo(Config.ProvisionDirectory + Path.GetFileName(MobileProvisionFilename));
+				DestFileInfo.Attributes = DestFileInfo.Attributes & ~FileAttributes.ReadOnly;
+				MobileProvisionFilename = Config.ProvisionDirectory + Path.GetFileName(MobileProvisionFilename);
+
+				// verify the provision is good
+				MobileProvision p = MobileProvisionParser.ParseFile(MobileProvisionFilename);
+				if ((p.ProvisionName.Contains("Wildcard") || p.ApplicationIdentifier.Contains("*")) && (!Config.bForDistribution || (p.ProvisionedDeviceIDs.Count == 0 && !p.bDebug)))
+				{
+					// check to see if we have a certificate for this provision
+					if (CodeSignatureBuilder.FindCertificate(p) != null)
+					{
+						return MobileProvisionFilename;
+					}
+				}
+			}
+
+			// check for a distribution provision in the game's build directory
+			if (Config.bForDistribution)
+			{
+				MobileProvisionFilename = FileOperations.FindPrefixedFile(Config.EngineBuildDirectory, "UE4Game_Distro.mobileprovision");
+
+				if (!File.Exists(MobileProvisionFilename))
+				{
+					MobileProvisionFilename = FileOperations.FindPrefixedFile(Config.EngineBuildDirectory + "/NotForLicensees/", "UE4Game_Distro.mobileprovision");
+					if (!File.Exists(MobileProvisionFilename))
+					{
+						MobileProvisionFilename = FileOperations.FindAnyFileWithExtension(Config.EngineBuildDirectory, ".mobileprovision");
+					}
+				}
+				if (File.Exists(MobileProvisionFilename))
+				{
+					Directory.CreateDirectory(Config.ProvisionDirectory);
+					File.Copy(MobileProvisionFilename, Config.ProvisionDirectory + Path.GetFileName(MobileProvisionFilename), true);
+					FileInfo DestFileInfo = new FileInfo(Config.ProvisionDirectory + Path.GetFileName(MobileProvisionFilename));
+					DestFileInfo.Attributes = DestFileInfo.Attributes & ~FileAttributes.ReadOnly;
+					MobileProvisionFilename = Config.ProvisionDirectory + Path.GetFileName(MobileProvisionFilename);
+
+					// verify the provision is good
+					MobileProvision p = MobileProvisionParser.ParseFile(MobileProvisionFilename);
+					if ((p.ProvisionName == GameName || p.ApplicationIdentifier.Contains(GameName)) && (!Config.bForDistribution || (p.ProvisionedDeviceIDs.Count == 0 && !p.bDebug)))
+					{
+						// check to see if we have a certificate for this provision
+						if (CodeSignatureBuilder.FindCertificate(p) != null)
+						{
+							return MobileProvisionFilename;
+						}
+					}
+				}
+			}
+			#endregion
+
 			// check to see if there is already an embedded provision
-			string MobileProvisionFilename = Path.Combine(Config.RepackageStagingDirectory, "embedded.mobileprovision");
+			MobileProvisionFilename = Path.Combine(Config.RepackageStagingDirectory, "embedded.mobileprovision");
 
 			return MobileProvisionFilename;
 		}
diff --git a/Engine/Source/Programs/AutomationTool/IOS/IPhonePackager/ToolsHub.cs b/Engine/Source/Programs/AutomationTool/IOS/IPhonePackager/ToolsHub.cs
--- a/Engine/Source/Programs/AutomationTool/IOS/IPhonePackager/ToolsHub.cs
+++ b/Engine/Source/Programs/AutomationTool/IOS/IPhonePackager/ToolsHub.cs
@@ -141,7 +141,7 @@ namespace iPhonePackager
 
 					// Copy the file into the destination location
 					string EffectivePrefix = bIsDistribution ? "Distro_" : Config.SigningPrefix;
-					string DestinationFilename = Path.Combine(Config.ProvisionDirectory, EffectivePrefix + Program.GameName + ".mobileprovision");
+					string DestinationFilename = Path.Combine(Config.BuildDirectory, EffectivePrefix + Program.GameName + ".mobileprovision");
 
 					if (File.Exists(DestinationFilename))
 					{
diff --git a/Engine/Source/Programs/AutomationTool/IOS/IPhonePackager/iPhonePackager.cs b/Engine/Source/Programs/AutomationTool/IOS/IPhonePackager/iPhonePackager.cs
--- a/Engine/Source/Programs/AutomationTool/IOS/IPhonePackager/iPhonePackager.cs
+++ b/Engine/Source/Programs/AutomationTool/IOS/IPhonePackager/iPhonePackager.cs
@@ -443,17 +443,22 @@ namespace iPhonePackager
 			}
 		}
 
+		[STAThread]
+		static int Main(string[] args)
+		{
+			return RealMain(args);
+		}
+		
 		/**
 		 * Main control loop
 		 */
-		[STAThread]
-		static int Main(string[] args)
+		public static int RealMain(string[] args)
 		{
 			// remember the working directory at start, as the game path could be relative to this path
 			string InitialCurrentDirectory = Environment.CurrentDirectory;
 
 			// set the working directory to the location of the application (so relative paths always work)
-			Environment.CurrentDirectory = Path.GetDirectoryName(Application.ExecutablePath);
+//			Environment.CurrentDirectory = Path.GetDirectoryName(Application.ExecutablePath);
 
 			AppDomain.CurrentDomain.AssemblyResolve += new ResolveEventHandler( CurrentDomain_AssemblyResolve );
 
diff --git a/Engine/Source/Programs/AutomationTool/Platform.cs b/Engine/Source/Programs/AutomationTool/Platform.cs
--- a/Engine/Source/Programs/AutomationTool/Platform.cs
+++ b/Engine/Source/Programs/AutomationTool/Platform.cs
@@ -376,6 +376,15 @@ namespace AutomationTool
 
 		}
 
+		/// <summary>
+		/// General purpose command to run generic string commands inside the platform interfeace
+		/// </summary>
+		/// <param name="Command"></param>
+		public virtual int RunCommand(string Command)
+		{
+			return 0;
+		}
+
 		public virtual bool ShouldUseManifestForUBTBuilds(string AddArgs)
 		{
 			return true;
diff --git a/Engine/Source/Programs/AutomationTool/Program.cs b/Engine/Source/Programs/AutomationTool/Program.cs
--- a/Engine/Source/Programs/AutomationTool/Program.cs
+++ b/Engine/Source/Programs/AutomationTool/Program.cs
@@ -16,6 +16,7 @@ namespace AutomationTool
 		// This needs to be static, otherwise SetConsoleCtrlHandler will result in a crash on exit.
 		static ProcessManager.CtrlHandlerDelegate ProgramCtrlHandler = new ProcessManager.CtrlHandlerDelegate(CtrlHandler);
 
+		[STAThread]
 		static int Main()
 		{
 			var CommandLine = SharedUtils.ParseCommandLine();
diff --git a/Engine/Source/Programs/AutomationTool/Scripts/AutomationScripts.Automation.csproj b/Engine/Source/Programs/AutomationTool/Scripts/AutomationScripts.Automation.csproj
--- a/Engine/Source/Programs/AutomationTool/Scripts/AutomationScripts.Automation.csproj
+++ b/Engine/Source/Programs/AutomationTool/Scripts/AutomationScripts.Automation.csproj
@@ -72,7 +72,7 @@
     <Compile Include="FixupRedirects.Automation.cs" />
     <Compile Include="GenerateDSYM.Automation.cs" />
     <Compile Include="GUBP.Automation.cs" />
-    <Compile Include="IOSBuildCookPackage.Automation.cs" />
+    <Compile Include="IPhonePackager.Automation.cs" />
     <Compile Include="LauncherLocalization.Automation.cs" />
     <Compile Include="Localisation.Automation.cs" />
     <Compile Include="MegaXGE.Automation.cs" />
diff --git a/Engine/Source/Programs/AutomationTool/Scripts/IOSBuildCookPackage.Automation.cs b/Engine/Source/Programs/AutomationTool/Scripts/IOSBuildCookPackage.Automation.cs
deleted file mode 100644
diff --git a/Engine/Source/Programs/AutomationTool/Scripts/IPhonePackager.Automation.cs b/Engine/Source/Programs/AutomationTool/Scripts/IPhonePackager.Automation.cs
new file mode 100644
--- /dev/null
+++ b/Engine/Source/Programs/AutomationTool/Scripts/IPhonePackager.Automation.cs
@@ -0,0 +1,26 @@
+﻿// Copyright 1998-2014 Epic Games, Inc. All Rights Reserved.
+using System;
+using System.Collections.Generic;
+using System.Text;
+using System.IO;
+using AutomationTool;
+using UnrealBuildTool;
+
+[Help("UAT command to call into the integrated IPhonePackager code")]
+class IPhonePackager : BuildCommand
+{
+	public override void ExecuteBuild()
+	{
+		Log("************************* Calling IPP");
+
+		Platform IOS = Platform.Platforms[UnrealTargetPlatform.IOS];
+		string Command = ParseParamValue("cmd", "");
+
+		// check the return value
+		int ReturnValue = IOS.RunCommand("IPP:" + Command);
+		if (ReturnValue != 0)
+		{
+			throw new AutomationException("Internal IPP returned {0}", ReturnValue);
+		}
+	}
+}
diff --git a/Engine/Source/Programs/UnrealBuildTool/IOS/IOSToolChain.cs b/Engine/Source/Programs/UnrealBuildTool/IOS/IOSToolChain.cs
--- a/Engine/Source/Programs/UnrealBuildTool/IOS/IOSToolChain.cs
+++ b/Engine/Source/Programs/UnrealBuildTool/IOS/IOSToolChain.cs
@@ -1163,9 +1163,22 @@ namespace UnrealBuildTool
 						UnrealBuildTool.GenerateProjectFiles(new XcodeProjectFileGenerator(), new string[] { "-NoIntellisense", "-iosdeployonly" });
 					}
 
+					// now that 
 					Process StubGenerateProcess = new Process();
-					StubGenerateProcess.StartInfo.WorkingDirectory = Path.GetFullPath("..\\Binaries\\DotNET\\IOS");
-					StubGenerateProcess.StartInfo.FileName = Path.Combine(StubGenerateProcess.StartInfo.WorkingDirectory, "iPhonePackager.exe");
+
+					if (RemoteToolChain.bUseRPCUtil)
+					{
+						StubGenerateProcess.StartInfo.WorkingDirectory = Path.GetFullPath("..\\Binaries\\DotNET\\IOS");
+						StubGenerateProcess.StartInfo.FileName = Path.Combine(StubGenerateProcess.StartInfo.WorkingDirectory, "iPhonePackager.exe");
+					}
+					else
+					{
+						// UAT wants the working directory to be the root (or does it matter?)
+						StubGenerateProcess.StartInfo.WorkingDirectory = Path.GetFullPath("..\\Binaries\\DotNET");
+						StubGenerateProcess.StartInfo.FileName = Path.Combine(StubGenerateProcess.StartInfo.WorkingDirectory, "AutomationTool.exe");
+					}
+
+					string Arguments = "";
 						
 					string PathToApp = RulesCompiler.GetTargetFilename(AppName);
 
@@ -1190,29 +1203,39 @@ namespace UnrealBuildTool
 					if (bUseDangerouslyFastMode)
 					{
 						// the quickness!
-						StubGenerateProcess.StartInfo.Arguments = "DangerouslyFast " + PathToApp;
+						Arguments += "DangerouslyFast " + PathToApp;
 					}
 					else
 					{
-						StubGenerateProcess.StartInfo.Arguments = "PackageIPA \"" + PathToApp + "\" -createstub";
+						Arguments += "PackageIPA \"" + PathToApp + "\" -createstub";
 						// if we are making the dsym, then we can strip the debug info from the executable
 						if (BuildConfiguration.bStripSymbolsOnIOS || (Target.Configuration == UnrealTargetConfiguration.Shipping))
 						{
-							StubGenerateProcess.StartInfo.Arguments += " -strip";
+							Arguments += " -strip";
 						}
 					}
-					StubGenerateProcess.StartInfo.Arguments += " -config " + Target.Configuration + " -mac " + RemoteServerName;
+					Arguments += " -config " + Target.Configuration + " -mac " + RemoteServerName;
 
 					var BuildPlatform = UEBuildPlatform.GetBuildPlatform(Target.Platform);
 					string Architecture = BuildPlatform.GetActiveArchitecture();
 					if (Architecture != "")
 					{
 						// pass along the architecture if we need, skipping the initial -, so we have "-architecture simulator"
-						StubGenerateProcess.StartInfo.Arguments += " -architecture " + Architecture.Substring(1);
+						Arguments += " -architecture " + Architecture.Substring(1);
 					}
 
 					// programmers that use Xcode packaging mode should use the following commandline instead, as it will package for Xcode on each compile
-					//				StubGenerateProcess.StartInfo.Arguments = "PackageApp " + GameName + " " + Configuration;
+					//				Arguments = "PackageApp " + GameName + " " + Configuration;
+
+					if (RemoteToolChain.bUseRPCUtil)
+					{
+						StubGenerateProcess.StartInfo.Arguments = Arguments;
+					}
+					else
+					{
+						Arguments = Arguments.Replace("\"", "\\\"");
+						StubGenerateProcess.StartInfo.Arguments = "IPhonePackager -cmd=\"" + Arguments + "\" -nocompile";
+					}
 
 					StubGenerateProcess.OutputDataReceived += new DataReceivedEventHandler(OutputReceivedDataEventHandler);
 					StubGenerateProcess.ErrorDataReceived += new DataReceivedEventHandler(OutputReceivedDataEventHandler);
diff --git a/Engine/Source/Programs/UnrealBuildTool/System/ActionGraph.cs b/Engine/Source/Programs/UnrealBuildTool/System/ActionGraph.cs
--- a/Engine/Source/Programs/UnrealBuildTool/System/ActionGraph.cs
+++ b/Engine/Source/Programs/UnrealBuildTool/System/ActionGraph.cs
@@ -615,14 +615,6 @@ namespace UnrealBuildTool
 						}
 					}
 
-					if (ProducedItem.bIsRemoteFile)
-					{
-						DateTime LastWriteTime;
-						long UnusedLength;
-						ProducedItem.bExists = RPCUtilHelper.GetRemoteFileInfo(ProducedItem.AbsolutePath, out LastWriteTime, out UnusedLength);
-						ProducedItem.LastWriteTime = LastWriteTime;
-					}
-
 					// If the produced file doesn't exist or has zero size, consider it outdated.  The zero size check is to detect cases
 					// where aborting an earlier compile produced invalid zero-sized obj files, but that may cause actions where that's
 					// legitimate output to always be considered outdated.
@@ -785,13 +777,17 @@ namespace UnrealBuildTool
 					{
 						if (ProducedItem.bIsRemoteFile)
 						{
-							try
-							{
-								RPCUtilHelper.MakeDirectory(Path.GetDirectoryName(ProducedItem.AbsolutePath).Replace("\\", "/"));
-							}
-							catch (System.Exception Ex)
+							// we don't need to do this in the SSH mode, the action will have an output file, and it will use that to make the directory while executing the command
+							if (RemoteToolChain.bUseRPCUtil)
 							{
-								throw new BuildException( Ex, "Error while creating remote directory for '{0}'.  (Exception: {1})", ProducedItem.AbsolutePath, Ex.Message );
+								try
+								{
+									RPCUtilHelper.MakeDirectory(Path.GetDirectoryName(ProducedItem.AbsolutePath).Replace("\\", "/"));
+								}
+								catch (System.Exception Ex)
+								{
+									throw new BuildException(Ex, "Error while creating remote directory for '{0}'.  (Exception: {1})", ProducedItem.AbsolutePath, Ex.Message);
+								}
 							}
 						}
 						else
diff --git a/Engine/Source/Programs/UnrealBuildTool/System/RPCUtilHelper.cs b/Engine/Source/Programs/UnrealBuildTool/System/RPCUtilHelper.cs
--- a/Engine/Source/Programs/UnrealBuildTool/System/RPCUtilHelper.cs
+++ b/Engine/Source/Programs/UnrealBuildTool/System/RPCUtilHelper.cs
@@ -26,6 +26,12 @@ namespace UnrealBuildTool
 		/** A socket per command thread */
 		private static Hashtable CommandThreadSockets = new Hashtable();
 
+		/** Time difference between remote and local idea's of UTC time */
+		private static TimeSpan TimeDifferenceFromRemote = new TimeSpan(0);
+
+		/** The number of commands the remote side should be able to run at once */
+		private static int MaxRemoteCommandsAllowed = 0;
+
 		static RPCUtilHelper()
 		{
 			AppDomain.CurrentDomain.AssemblyResolve += new ResolveEventHandler(CurrentDomain_AssemblyResolve);
@@ -66,22 +72,88 @@ namespace UnrealBuildTool
 
 			return (null);
 		}
-	
+
+		private static DateTime RemoteToLocalTime(string RemoteTime)
+		{
+			try
+			{
+				// convert string to integer
+				int RemoteTimeSinceEpoch = int.Parse(RemoteTime);
+
+				// convert the seconds into TimeSpan
+				TimeSpan RemoteSpanSinceEpoch = new TimeSpan(0, 0, RemoteTimeSinceEpoch);
+
+				// put the remote time into local time (Jan 1, 1970 was the Epoch for Mac), and adjust it for time difference between machines
+				return new DateTime(1970, 1, 1, 0, 0, 0, DateTimeKind.Utc) + RemoteSpanSinceEpoch + TimeDifferenceFromRemote;
+
+			}
+			catch (Exception)
+			{
+				// use MinValue for any errors
+				return DateTime.MinValue;
+			}
+		}
+
 		static public void Initialize(string InMacName)
 		{
 			MacName = InMacName;
 
-			if (CommandHelper.PingRemoteHost(MacName))
+			// when not using RPCUtil, we do NOT want to ping the host
+			if (!RemoteToolChain.bUseRPCUtil || CommandHelper.PingRemoteHost(MacName))
 			{
-				if (BuildConfiguration.bFlushBuildDirOnRemoteMac)
+				try
 				{
-					Command("/", "rm", "-rf /UE4/Builds/" + Environment.MachineName, null);
+					if (!RemoteToolChain.bUseRPCUtil)
+					{
+						// ask for current time, free memory and num CPUs
+
+						string[] Commands = new string[] 
+						{
+							"+\"%s\"",
+							"vm_stat | grep Statistics | awk '{print $8}'",
+							"vm_stat | grep 'Pages free:' | awk '{print $3}'",
+							"sysctl -a | grep hw.logicalcpu: | awk '{print $2}'",
+						};
+						Hashtable Results = Command("/", "date", string.Join(" && ", Commands), null);
+						string[] Lines = ((string)Results["CommandOutput"]).Split("\r\n".ToCharArray(), StringSplitOptions.RemoveEmptyEntries);
+
+						// convert it to a time, with TimeDifferenceFromRemote as 0
+						DateTime RemoteTimebase = RemoteToLocalTime(Lines[0]);
+						if (RemoteTimebase == DateTime.MinValue)
+						{
+							throw new BuildException("Failed to parse remote time on " + MacName);
+						}
+
+						// calculate the difference
+						TimeDifferenceFromRemote = DateTime.UtcNow - RemoteTimebase;
+
+						// now figure out max number of commands to run at once
+						int PageSize = int.Parse(Lines[1]);
+						Int64 AvailableMem = Int64.Parse(Lines[2].Replace(".", "")) * PageSize;
+						int NumProcesses = (int)Math.Max(1, AvailableMem / (RemoteToolChain.MemoryPerCompileMB * 1024 * 1024));
+
+						// now, combine that with actual number of cores
+						int NumCores = int.Parse(Lines[3]);
+						MaxRemoteCommandsAllowed = Math.Min(NumProcesses, NumCores);
+
+						Console.WriteLine("Remote time is {0}, difference is {1}", RemoteTimebase.ToString(), TimeDifferenceFromRemote.ToString());
+					}
+					
+					if (BuildConfiguration.bFlushBuildDirOnRemoteMac)
+					{
+						Command("/", "rm", "-rf /UE4/Builds/" + Environment.MachineName, null);
+					}
+				}
+				catch(Exception Ex)
+				{
+					Console.WriteLine("Ex {0}", Ex.ToString());
+					throw new BuildException("Failed to run init commands on " + MacName);
 				}
 			}
-			else
-			{
-				throw new BuildException("Failed to ping Mac named " + MacName);
-			}
+ 			else
+ 			{
+ 				throw new BuildException("Failed to ping Mac named " + MacName);
+ 			}
 		}
 
 		/**
@@ -162,12 +234,42 @@ namespace UnrealBuildTool
 		 */
 		public static bool GetRemoteFileInfo(string RemotePath, out DateTime ModificationTime, out long Length)
 		{
-			return RPCUtility.CommandHelper.GetFileInfo(GetSocket(), RemotePath, DateTime.UtcNow, out ModificationTime, out Length);
+			if (RemoteToolChain.bUseRPCUtil)
+			{
+				return RPCUtility.CommandHelper.GetFileInfo(GetSocket(), RemotePath, DateTime.UtcNow, out ModificationTime, out Length);
+			}
+			else
+			{
+				string CommandArgs = string.Format("-c 'if [ -e \"{0}\" ]; then eval $(stat -s \"{0}\") && echo $st_mtime,$st_size; fi'", RemotePath);
+				Hashtable Results = Command("/", "bash", CommandArgs, null);
+
+				string Output = Results["CommandOutput"] as string;
+				string[] Tokens =  Output.Split(",".ToCharArray());
+				if (Tokens.Length == 2)
+				{
+					ModificationTime = RemoteToLocalTime(Tokens[0]);
+					Length = long.Parse(Tokens[1]);
+					return true;
+				}
+
+				// any failures will fall through to here
+				ModificationTime = DateTime.MinValue;
+				Length = 0;
+				return false;
+			
+			}
 		}
 
 		public static void MakeDirectory(string Directory)
 		{
-			RPCUtility.CommandHelper.MakeDirectory(GetSocket(), Directory);
+			if (RemoteToolChain.bUseRPCUtil)
+			{
+				RPCUtility.CommandHelper.MakeDirectory(GetSocket(), Directory);
+			}
+			else
+			{
+				Command("/", "bash", "-c 'mkdir \"" + Directory + "\"'", null);
+			}
 		}
 
 		[Flags]
@@ -181,15 +283,27 @@ namespace UnrealBuildTool
 
 		public static void CopyFile(string Source, string Dest, bool bIsUpload)
 		{
-			if (bIsUpload)
+			if (RemoteToolChain.bUseRPCUtil)
 			{
-//				Hashtable CommandResult = 
-				RPCUtility.CommandHelper.RPCUpload(GetSocket(), Source, Dest);
-//				Log.TraceInformation(CommandResult["CommandOutput"] as string);
+				if (bIsUpload)
+				{
+					RPCUtility.CommandHelper.RPCUpload(GetSocket(), Source, Dest);
+				}
+				else
+				{
+					RPCUtility.CommandHelper.RPCDownload(GetSocket(), Source, Dest);
+				}
 			}
 			else
 			{
-				RPCUtility.CommandHelper.RPCDownload(GetSocket(), Source, Dest);
+				if (bIsUpload)
+				{
+					RemoteToolChain.UploadFile(Source, Dest);
+				}
+				else
+				{
+					RemoteToolChain.DownloadFile(Source, Dest);
+				}
 			}
 		}
 
@@ -279,64 +393,135 @@ namespace UnrealBuildTool
 
 		public static void BatchFileInfo(FileItem[] Files)
 		{
-			// build a list of file paths to get info about
-			StringBuilder FileList = new StringBuilder();
-			foreach (FileItem File in Files)
+			if (RemoteToolChain.bUseRPCUtil)
 			{
-				FileList.AppendFormat("{0}\n", File.AbsolutePath);
-			}
+				// build a list of file paths to get info about
+				StringBuilder FileList = new StringBuilder();
+				foreach (FileItem File in Files)
+				{
+					FileList.AppendFormat("{0}\n", File.AbsolutePath);
+				}
+
+				DateTime Now = DateTime.Now;
 
-			// execute the command!
-			Int64[] FileSizeAndDates = RPCUtility.CommandHelper.RPCBatchFileInfo(GetSocket(), FileList.ToString());
+				// execute the command!
+				Int64[] FileSizeAndDates = RPCUtility.CommandHelper.RPCBatchFileInfo(GetSocket(), FileList.ToString());
 
-			// now update the source times
-			for (int Index = 0; Index < Files.Length; Index++)
+				Console.WriteLine("BatchFileInfo version 1 took {0}", (DateTime.Now - Now).ToString());
+
+				// now update the source times
+				for (int Index = 0; Index < Files.Length; Index++)
+				{
+					Files[Index].Length = FileSizeAndDates[Index * 2 + 0];
+					Files[Index].LastWriteTime = new DateTimeOffset(RPCUtility.CommandHelper.FromRemoteTime(FileSizeAndDates[Index * 2 + 1]));
+					Files[Index].bExists = FileSizeAndDates[Index * 2 + 0] >= 0;
+				}
+			}
+			else
 			{
-				Files[Index].Length = FileSizeAndDates[Index * 2 + 0];
-				Files[Index].LastWriteTime = new DateTimeOffset(RPCUtility.CommandHelper.FromRemoteTime(FileSizeAndDates[Index * 2 + 1]));
-				Files[Index].bExists = FileSizeAndDates[Index * 2 + 0] >= 0;
+				// build a list of file paths to get info about
+				StringBuilder Commands = new StringBuilder();
+				Commands.Append("#!/bin/bash\n");
+				foreach (FileItem File in Files)
+				{
+					Commands.AppendFormat("if [ -e \"{0}\" ]; then eval $(stat -s \"{0}\") && echo $st_mtime && echo $st_size; else echo 0 && echo -1; fi\n", File.AbsolutePath);
+				}
+
+				// write out locally
+				string LocalCommandsFile = Path.GetTempFileName();
+				System.IO.File.WriteAllText(LocalCommandsFile, Commands.ToString());
+
+				string RemoteDir = "/var/tmp/" + Environment.MachineName;
+				string RemoteCommandsFile = Path.GetFileName(LocalCommandsFile) + ".sh";
+
+				DateTime Now = DateTime.Now;
+				RemoteToolChain.UploadFile(LocalCommandsFile, RemoteDir + "/" + RemoteCommandsFile);
+
+				// execute the file, not a commandline
+				Hashtable Results = Command(RemoteDir, "sh", RemoteCommandsFile + " && rm " + RemoteCommandsFile, null);
+				
+				Console.WriteLine("BatchFileInfo took {0}", (DateTime.Now - Now).ToString());
+
+				string[] Lines = ((string)Results["CommandOutput"]).Split("\r\n".ToCharArray(),	StringSplitOptions.RemoveEmptyEntries);
+				if (Lines.Length != Files.Length * 2)
+				{
+					throw new BuildException("Received the wrong number of results from BatchFileInfo");
+				}
+
+				for (int Index = 0; Index < Files.Length; Index++)
+				{
+					Files[Index].LastWriteTime = new DateTimeOffset(RemoteToLocalTime(Lines[Index * 2 + 0]));
+					Files[Index].Length = long.Parse(Lines[Index * 2 + 1]);
+					Files[Index].bExists = Files[Index].Length >= 0;
+				}
 			}
 		}
 
 		public static int GetCommandSlots()
 		{
-			return RPCUtility.CommandHelper.GetCommandSlots(GetSocket());
+			if (RemoteToolChain.bUseRPCUtil)
+			{
+				return RPCUtility.CommandHelper.GetCommandSlots(GetSocket());
+			}
+			else
+			{
+				return MaxRemoteCommandsAllowed;
+			}
 		}
 
-		public static Hashtable Command(string WorkingDirectory, string Command, string CommandArgs, string RemoteOutputPath)
+		public static Hashtable Command(string WorkingDirectory, string CommandWithArgs, string RemoteOutputPath)
 		{
-	
-			int RetriesRemaining = 6;
-			do
+			int FirstSpace = CommandWithArgs.IndexOf(' ');
+			if (FirstSpace == -1)
 			{
-				// a $ on the commandline will actually be converted, so we need to quote it
-				CommandArgs = CommandArgs.Replace("$", "\\$");
+				return Command(WorkingDirectory, CommandWithArgs, " ", RemoteOutputPath);
+			}
 
-				try
-				{
-					Hashtable Results = RPCUtility.CommandHelper.RPCCommand(GetSocket(), WorkingDirectory, Command, CommandArgs, RemoteOutputPath);
-					return Results;
-				}
-				catch (Exception Ex)
+			return Command(WorkingDirectory, CommandWithArgs.Substring(0, FirstSpace), CommandWithArgs.Substring(FirstSpace + 1), RemoteOutputPath);
+
+		}
+
+		public static Hashtable Command(string WorkingDirectory, string Command, string CommandArgs, string RemoteOutputPath)
+		{
+			if (RemoteToolChain.bUseRPCUtil)
+			{
+				int RetriesRemaining = 6;
+				do
 				{
-					if (RetriesRemaining > 0)
+					// a $ on the commandline will actually be converted, so we need to quote it
+					CommandArgs = CommandArgs.Replace("$", "\\$");
+
+					try
 					{
-						Int32 RetryTimeoutMS = 1000;
-						Debug.WriteLine("Retrying command after sleeping for " + RetryTimeoutMS + " milliseconds. Command is:" + Command + " " + CommandArgs);
-						Thread.Sleep(RetryTimeoutMS);
+						Hashtable Results = RPCUtility.CommandHelper.RPCCommand(GetSocket(), WorkingDirectory, Command, CommandArgs, RemoteOutputPath);
+						return Results;
 					}
-					else
+					catch (Exception Ex)
 					{
-						Log.TraceInformation("Out of retries, too many exceptions:" + Ex.ToString());
-						// We've tried enough times, just throw the error
-						throw new Exception("Deep Exception, retries exhausted... ", Ex);
+						if (RetriesRemaining > 0)
+						{
+							Int32 RetryTimeoutMS = 1000;
+							Debug.WriteLine("Retrying command after sleeping for " + RetryTimeoutMS + " milliseconds. Command is:" + Command + " " + CommandArgs);
+							Thread.Sleep(RetryTimeoutMS);
+						}
+						else
+						{
+							Log.TraceInformation("Out of retries, too many exceptions:" + Ex.ToString());
+							// We've tried enough times, just throw the error
+							throw new Exception("Deep Exception, retries exhausted... ", Ex);
+						}
+						RetriesRemaining--;
 					}
-					RetriesRemaining--;
 				}
-			}
-			while (RetriesRemaining > 0);
+				while (RetriesRemaining > 0);
+
+				return null;
 
-			return null;
+			}
+			else
+			{
+				return RemoteToolChain.SSHCommand(WorkingDirectory, Command, CommandArgs, RemoteOutputPath);
+			}
 		}
 	}
 }
diff --git a/Engine/Source/Programs/UnrealBuildTool/System/Utils.cs b/Engine/Source/Programs/UnrealBuildTool/System/Utils.cs
--- a/Engine/Source/Programs/UnrealBuildTool/System/Utils.cs
+++ b/Engine/Source/Programs/UnrealBuildTool/System/Utils.cs
@@ -352,8 +352,10 @@ namespace UnrealBuildTool
 		/**
 		/* Try to launch a local process, and produce a friendly error message if it fails.
 		/*/
-		public static void RunLocalProcess(Process LocalProcess)
+		public static int RunLocalProcess(Process LocalProcess)
 		{
+			int ExitCode = -1;
+
 			// release all process resources
 			using(LocalProcess)
 			{
@@ -368,12 +370,15 @@ namespace UnrealBuildTool
 					LocalProcess.BeginOutputReadLine();
 					LocalProcess.BeginErrorReadLine();
 					LocalProcess.WaitForExit();
+					ExitCode = LocalProcess.ExitCode;
 				}
 				catch(Exception ex)
 				{
 					throw new BuildException(ex, "Failed to start local process for action (\"{0}\"): {1} {2}", ex.Message, LocalProcess.StartInfo.FileName, LocalProcess.StartInfo.Arguments);
 				}
 			}
+
+			return ExitCode;
 		}
 
 
diff --git a/Engine/Source/Programs/UnrealBuildTool/ToolChain/RemoteToolChain.cs b/Engine/Source/Programs/UnrealBuildTool/ToolChain/RemoteToolChain.cs
--- a/Engine/Source/Programs/UnrealBuildTool/ToolChain/RemoteToolChain.cs
+++ b/Engine/Source/Programs/UnrealBuildTool/ToolChain/RemoteToolChain.cs
@@ -1,6 +1,7 @@
 ﻿// Copyright 1998-2014 Epic Games, Inc. All Rights Reserved.
 
 using System;
+using System.Collections;
 using System.Collections.Generic;
 using System.IO;
 using System.Diagnostics;
@@ -12,7 +13,7 @@ using System.Text.RegularExpressions;
 
 namespace UnrealBuildTool
 {
-	abstract class RemoteToolChain : UEToolChain
+	public abstract class RemoteToolChain : UEToolChain
     {
 		protected void RegisterRemoteToolChain(UnrealTargetPlatform InPlatform, CPPTargetPlatform CPPPlatform)
 		{
@@ -35,6 +36,30 @@ namespace UnrealBuildTool
 		/** The path (on the Mac) to the your particular development directory, where files will be copied to from the PC */
 		public static string UserDevRootMac = "/UE4/Builds/";
 
+		/** Whether or not to connect to UnrealRemoteTool using RPCUtility */
+		[XmlConfig]
+		public static bool bUseRPCUtil = true;
+		
+		/** Path to rsync executable and parameters for your rsync utility */
+		[XmlConfig]
+		public static string RSyncExe = "C:\\Program Files (x86)\\DeltaCopy\\rsync.exe";
+
+		/** Path to rsync executable and parameters for your rsync utility */
+		[XmlConfig]
+		public static string SSHExe = "C:\\Program Files (x86)\\DeltaCopy\\ssh.exe";
+
+		/** The authentication used for Rsync (for the -e rsync flag) */
+		[XmlConfig]
+		public static string RsyncAuthentication = "ssh -i id_rsa";
+
+		/** The authentication used for SSH (probably similar to RsyncAuthentication) */
+		[XmlConfig]
+		public static string SSHAuthentication = "-i id_rsa";
+
+		/** Username on the remote machine to connect to with RSync */
+		[XmlConfig]
+		public static string RSyncUsername = "josh.adams";
+
 		/** The directory that this local branch is in, without drive information (strip off X:\ from X:\UE4\iOS) */
 		public static string BranchDirectory = Environment.MachineName + "\\" + Path.GetFullPath(".\\").Substring(3);
 		public static string BranchDirectoryMac = BranchDirectory.Replace("\\", "/");
@@ -46,6 +71,9 @@ namespace UnrealBuildTool
 		/** The platform this toolchain is compiling for */
 		protected UnrealTargetPlatform RemoteToolChainPlatform;
 
+		/** The average amound of memory a compile takes, used so that we don't compile too many things at once */
+		public static int MemoryPerCompileMB = 1000;
+
 		static RemoteToolChain()
         {
             ErrorMessageTokens = new List<string>();
@@ -200,26 +228,12 @@ namespace UnrealBuildTool
 			{
 				RsyncDirs.Add(Entry);
 			}
-			
-			
-// 			// If not, create it now
-//             string RemoteFilePath = ConvertPath(LocalFileItem.AbsolutePath);
-// 
-//             // add this file to the list of files we will RPCBatchUpload later
-//             string Entry = LocalFileItem.AbsolutePath + " " + RemoteFilePath;
-//             if (!BatchUploadCommands.Contains(Entry))
-//             {
-//                 BatchUploadCommands.Add(Entry);
-//             }
-
-// 			// strip it down to the root (rsync wants it to end in a / for -a to work)
-// 			Entry = Utils.MakePathRelativeTo(Path.GetDirectoryName(LocalFileItem.AbsolutePath), "../..").Replace('\\', '/') + "/";
-// 			if (!RsyncDirs.Contains(Entry))
-// 			{
-// 				RsyncDirs.Add(Entry);
-// 			}
 
 			string Ext = Path.GetExtension(LocalFileItem.AbsolutePath);
+			if (Ext == "")
+			{
+				Ext = Path.GetFileName(LocalFileItem.AbsolutePath);
+			}
 			if (!RsyncExtensions.Contains(Ext))
             {
 				RsyncExtensions.Add(Ext);
@@ -333,7 +347,26 @@ namespace UnrealBuildTool
 				Log.TraceInformation(Line.Data);
 			}
 		}
+
+		private static Dictionary<Object,StringBuilder> SSHOutputMap = new Dictionary<object,StringBuilder>();
+		static public void OutputReceivedForSSH(Object Sender, DataReceivedEventArgs Line)
+		{
+			if ((Line != null) && (Line.Data != null) && (Line.Data != ""))
+			{
+				StringBuilder SSHOutput = SSHOutputMap[Sender];
+				if (SSHOutput.Length != 0)
+				{
+					SSHOutput.Append(Environment.NewLine);
+				}
+				SSHOutput.Append(Line.Data);
+			}
+		}
 		
+		private static string ConvertPathToCygwin(string InPath)
+		{
+			return "/cygdrive/" + Utils.CleanDirectorySeparators(InPath.Replace(":", ""), '/');
+		}
+
 		public override void PreBuildSync()
         {
 			// no need to sync on the Mac!
@@ -342,21 +375,33 @@ namespace UnrealBuildTool
 				return;
 			}
 
-			bool bUseRPCUtil = true;
-
 			if (bUseRPCUtil)
 			{
-				List<string> BatchUploadCommands = new List<string>();
+				string ExtString = "";
+
+				// look only for useful extensions
+				foreach (string Ext in RsyncExtensions)
+				{
+					// for later ls
+					ExtString += Ext.StartsWith(".") ? ("*" + Ext) : Ext;
+					ExtString += " ";
+				}
 
+				List<string> BatchUploadCommands = new List<string>();
 				// for each directory we visited, add all the files in that directory
 				foreach (string Dir in RsyncDirs)
 				{
+					List<string> LocalFilenames = new List<string>();
+
 					// look only for useful extensions
 					foreach (string Ext in RsyncExtensions)
 					{
 						string[] Files = Directory.GetFiles(Dir, "*" + Ext);
 						foreach (string SyncFile in Files)
 						{
+							// remember all local files
+							LocalFilenames.Add(Path.GetFileName(SyncFile));
+
 							string RemoteFilePath = ConvertPath(SyncFile);
 							// an upload command is local name and remote name
 							BatchUploadCommands.Add(SyncFile + ";" + RemoteFilePath);
@@ -376,28 +421,44 @@ namespace UnrealBuildTool
 				}
 
 				// write out directories to copy
-				File.WriteAllLines("D:\\dev\\CarefullyRedist\\DeltaCopy\\RSyncPaths.txt", RelativeRsyncDirs.ToArray());
-				File.WriteAllLines("D:\\dev\\CarefullyRedist\\DeltaCopy\\IncludeFrom.txt", RsyncExtensions);
+				string RSyncPathsFile = Path.GetTempFileName();
+				string IncludeFromFile = Path.GetTempFileName();
+				File.WriteAllLines(RSyncPathsFile, RelativeRsyncDirs.ToArray());
+				File.WriteAllLines(IncludeFromFile, RsyncExtensions);
 
 				// source and destination paths in the format rsync wants
-				string CygPath = "/cygdrive/" + Utils.CleanDirectorySeparators(Path.GetFullPath("../../").Replace(":", ""), '/');
-				string RemotePath = ConvertPath(Path.GetFullPath("../../"));
+				string CygRootPath = ConvertPathToCygwin(Path.GetFullPath("../../"));
+				string RemotePath = ConvertPath(Path.GetFullPath("../../")).Replace(" ", "\\ ");
 
+				// get the executable dir for Rsync
+				string ExeDir = Path.GetDirectoryName(RSyncExe);
+	
 				Process RsyncProcess = new Process();
-				RsyncProcess.StartInfo.WorkingDirectory = Path.GetFullPath("D:\\dev\\CarefullyRedist\\DeltaCopy");
-				RsyncProcess.StartInfo.FileName = "D:\\dev\\CarefullyRedist\\DeltaCopy\\rsync.exe";
+				if (ExeDir != "")
+				{
+					RsyncProcess.StartInfo.WorkingDirectory = ExeDir;
+				}
+
+				// --exclude='*'  ??? why???
+				RsyncProcess.StartInfo.FileName = RSyncExe;
 				RsyncProcess.StartInfo.Arguments = string.Format(
-					"-vzae \"ssh -i id_rsa\" --delete --files-from=RsyncPaths.txt " +
-					"--include-from=IncludeFrom.txt --include='*/' --exclude='*' {0} josh.adams@{2}:{1}",
-					CygPath,
+					"-vzae \"{0}\" --rsync-path=\"mkdir -p {2} && rsync\" --delete --files-from=\"{4}\" --include-from=\"{5}\" --include='*/' --exclude='*.o' {1} {6}@{3}:'{2}'",
+					RsyncAuthentication,
+					CygRootPath,
 					RemotePath,
-					RemoteServerName);
+					RemoteServerName,
+					ConvertPathToCygwin(RSyncPathsFile),
+					ConvertPathToCygwin(IncludeFromFile),
+					RSyncUsername);
 
 				RsyncProcess.OutputDataReceived += new DataReceivedEventHandler(OutputReceivedForRsync);
 				RsyncProcess.ErrorDataReceived += new DataReceivedEventHandler(OutputReceivedForRsync);
 
 				// run rsync
 				Utils.RunLocalProcess(RsyncProcess);
+
+				File.Delete(IncludeFromFile);
+				File.Delete(RSyncPathsFile);
 			}
 
 			// we can now clear out the set of files
@@ -405,6 +466,142 @@ namespace UnrealBuildTool
 			RsyncExtensions.Clear();
 		}
 
+		static public bool UploadFile(string LocalPath, string RemotePath)
+		{
+			string RemoteDir = Path.GetDirectoryName(RemotePath).Replace("\\", "/");
+			string RemoteFilename = Path.GetFileName(RemotePath);
+
+			// get the executable dir for Rsync
+			string ExeDir = Path.GetDirectoryName(RSyncExe);
+
+			Process RsyncProcess = new Process();
+			if (ExeDir != "")
+			{
+				RsyncProcess.StartInfo.WorkingDirectory = ExeDir;
+			}
+
+			// make simple rsync commandline to send a file
+			RsyncProcess.StartInfo.FileName = RSyncExe;
+			RsyncProcess.StartInfo.Arguments = string.Format(
+				"-zae \"{0}\" --rsync-path=\"mkdir -p '{1}' && rsync\" \"{2}\" {3}@{4}:\"{1}/{5}\"",
+				RsyncAuthentication,
+				RemoteDir,
+				ConvertPathToCygwin(LocalPath),
+				RSyncUsername,
+				RemoteServerName,
+				RemoteFilename
+				);
+
+			RsyncProcess.OutputDataReceived += new DataReceivedEventHandler(OutputReceivedForRsync);
+			RsyncProcess.ErrorDataReceived += new DataReceivedEventHandler(OutputReceivedForRsync);
+
+			// run rsync (0 means success)
+			return Utils.RunLocalProcess(RsyncProcess) == 0;
+		}
+
+		static public bool DownloadFile(string RemotePath, string LocalPath)
+		{
+			// get the executable dir for Rsync
+			string ExeDir = Path.GetDirectoryName(RSyncExe);
+
+			Process RsyncProcess = new Process();
+			if (ExeDir != "")
+			{
+				RsyncProcess.StartInfo.WorkingDirectory = ExeDir;
+			}
+
+			// make sure directory exists to download to
+			Directory.CreateDirectory(Path.GetDirectoryName(LocalPath));
+
+			// make simple rsync commandline to send a file
+			RsyncProcess.StartInfo.FileName = RSyncExe;
+			RsyncProcess.StartInfo.Arguments = string.Format(
+				"-zae \"{0}\" {2}@{3}:\"{4}\" \"{1}\"",
+				RsyncAuthentication,
+				ConvertPathToCygwin(LocalPath),
+				RSyncUsername,
+				RemoteServerName,
+				RemotePath
+				);
+
+			RsyncProcess.OutputDataReceived += new DataReceivedEventHandler(OutputReceivedForRsync);
+			RsyncProcess.ErrorDataReceived += new DataReceivedEventHandler(OutputReceivedForRsync);
+
+			//Console.WriteLine("COPY: {0} {1}", RsyncProcess.StartInfo.FileName, RsyncProcess.StartInfo.Arguments);
+
+			// run rsync (0 means success)
+			return Utils.RunLocalProcess(RsyncProcess) == 0;
+		}
+
+		static public Hashtable SSHCommand(string WorkingDirectory, string Command, string CommandArgs, string RemoteOutputPath)
+		{
+			Console.WriteLine("Doing {0} {1}", Command, CommandArgs);
+
+			// make the commandline for other end
+			string RemoteCommandline = "cd " + WorkingDirectory;
+			if (!string.IsNullOrWhiteSpace(RemoteOutputPath))
+			{
+				RemoteCommandline += " && mkdir -p " + Path.GetDirectoryName(RemoteOutputPath).Replace("\\", "/");
+			}
+
+			// get the executable dir for SSH
+			string ExeDir = Path.GetDirectoryName(SSHExe);
+
+			Process SSHProcess = new Process();
+			if (ExeDir != "")
+			{
+				SSHProcess.StartInfo.WorkingDirectory = ExeDir;
+			}
+
+			// long commands go as a file
+			if (CommandArgs.Length > 256)
+			{
+				// upload the commandline text file
+				string CommandLineFile = Path.GetTempFileName();
+				File.WriteAllText(CommandLineFile, Command + " " + CommandArgs);
+
+				string RemoteCommandlineDir = "/var/tmp/" + Environment.MachineName;
+				string RemoteCommandlinePath = RemoteCommandlineDir + "/" + Path.GetFileName(CommandLineFile);
+
+				DateTime Now = DateTime.Now;
+				UploadFile(CommandLineFile, RemoteCommandlinePath);
+				Console.WriteLine("Upload took {0}", (DateTime.Now - Now).ToString());
+
+				// execute the file, not a commandline
+				RemoteCommandline += string.Format(" && bash < {0} && rm {0}", RemoteCommandlinePath);
+			}
+			else
+			{
+				RemoteCommandline += " && " + Command + " " + CommandArgs;
+			}
+
+			SSHProcess.StartInfo.FileName = SSHExe;
+			SSHProcess.StartInfo.Arguments = string.Format(
+				"{0} {1}@{2} \"{3}\"",
+				SSHAuthentication,
+				RSyncUsername,
+				RemoteServerName,
+				RemoteCommandline.Replace("\"", "\\\""));
+
+			// add this process to the map
+			SSHOutputMap[SSHProcess] = new StringBuilder("");
+			SSHProcess.OutputDataReceived += new DataReceivedEventHandler(OutputReceivedForSSH);
+			SSHProcess.ErrorDataReceived += new DataReceivedEventHandler(OutputReceivedForSSH);
+
+			DateTime Start = DateTime.Now;
+			Int64 ExitCode = Utils.RunLocalProcess(SSHProcess);
+			Console.WriteLine("Execute took {0}", (DateTime.Now - Start).ToString());
+
+			// now we have enough to fill out the HashTable
+			Hashtable Return = new Hashtable();
+			Return["CommandOutput"] = SSHOutputMap[SSHProcess].ToString();
+			Return["ExitCode"] = (object)ExitCode;
+
+			SSHOutputMap.Remove(SSHProcess);
+
+			return Return;
+		}
+
 		public override void PostBuildSync(UEBuildTarget Target)
         {
 
