UBT: property telling if particular target is modular or monolithic.

Needs to be in sync with the rest of logic.

diff --git a/Engine/Source/Programs/UnrealBuildTool/System/RulesCompiler.cs b/Engine/Source/Programs/UnrealBuildTool/System/RulesCompiler.cs
--- a/Engine/Source/Programs/UnrealBuildTool/System/RulesCompiler.cs
+++ b/Engine/Source/Programs/UnrealBuildTool/System/RulesCompiler.cs
@@ -59,7 +59,25 @@ namespace UnrealBuildTool
 					 Type == TargetRules.TargetType.RocketGame;
 			}
 		}
-	}
+
+        /// <summary>
+        /// True if the target type is a monolithic binary
+        /// </summary>
+        public bool IsMonolithic
+        {
+            get
+            {
+                if (!Type.HasValue)
+                {
+                    throw new BuildException("Trying to access TargetInfo.IsMonolithic when TargetInfo.Type is not set. Make sure IsMonolithic is used only in ModuleRules.");
+                }
+                return Type == TargetRules.TargetType.Client ||
+                     Type == TargetRules.TargetType.Game ||
+                     Type == TargetRules.TargetType.Server ||
+                     Type == TargetRules.TargetType.RocketGame;
+            }
+        }
+    }
 
 
 	/// <summary>
