Fix building of projects with space in the path name.

(Fixes 4.3 TTP #338978).

diff --git a/Engine/Source/Programs/UnrealBuildTool/Linux/LinuxToolChain.cs b/Engine/Source/Programs/UnrealBuildTool/Linux/LinuxToolChain.cs
--- a/Engine/Source/Programs/UnrealBuildTool/Linux/LinuxToolChain.cs
+++ b/Engine/Source/Programs/UnrealBuildTool/Linux/LinuxToolChain.cs
@@ -641,7 +641,7 @@ namespace UnrealBuildTool.Linux
                 Writer.WriteLine("};");
             };
 
-            LinkAction.CommandArguments += string.Format(" -Wl,--version-script={0}", LinkerScriptPath);
+            LinkAction.CommandArguments += string.Format(" -Wl,--version-script=\"{0}\"", LinkerScriptPath);
 
             // Only execute linking on the local PC.
             LinkAction.bCanExecuteRemotely = false;
