- Removed Metal support, must be manually enabled in UBT code. For safety :) Will try to come up with a better solution that is based on Project Settings

diff --git a/Engine/Source/Programs/UnrealBuildTool/IOS/UEBuildIOS.cs b/Engine/Source/Programs/UnrealBuildTool/IOS/UEBuildIOS.cs
--- a/Engine/Source/Programs/UnrealBuildTool/IOS/UEBuildIOS.cs
+++ b/Engine/Source/Programs/UnrealBuildTool/IOS/UEBuildIOS.cs
@@ -122,12 +122,13 @@ namespace UnrealBuildTool
 			}
 
 			// needs IOS8 for Metal
-			if (IOSToolChain.IOSSDKVersionFloat >= 8.0 && UEBuildConfiguration.bCompileAgainstEngine)
-			{
-				InBuildTarget.GlobalCompileEnvironment.Config.Definitions.Add("HAS_METAL=1");
-				InBuildTarget.ExtraModuleNames.Add("MetalRHI");
-			}
-			else
+			// @todo metal enable: Uncomment the following lines if you want to give Metal a try!
+// 			if (IOSToolChain.IOSSDKVersionFloat >= 8.0 && UEBuildConfiguration.bCompileAgainstEngine)
+// 			{
+// 				InBuildTarget.GlobalCompileEnvironment.Config.Definitions.Add("HAS_METAL=1");
+// 				InBuildTarget.ExtraModuleNames.Add("MetalRHI");
+// 			}
+// 			else
 			{
 				InBuildTarget.GlobalCompileEnvironment.Config.Definitions.Add("HAS_METAL=0");
 			}
