Linux: dirty hack for modular builds.

diff --git a/Engine/Source/Programs/UnrealBuildTool/Linux/LinuxToolChain.cs b/Engine/Source/Programs/UnrealBuildTool/Linux/LinuxToolChain.cs
--- a/Engine/Source/Programs/UnrealBuildTool/Linux/LinuxToolChain.cs
+++ b/Engine/Source/Programs/UnrealBuildTool/Linux/LinuxToolChain.cs
@@ -276,6 +276,9 @@ namespace UnrealBuildTool.Linux
 
             // RPATH for third party libs
             Result += " -Wl,-rpath=${ORIGIN}/../../../Engine/Binaries/Linux";
+            // FIXME: really ugly temp solution. Modules need to be able to specify this
+            Result += " -Wl,-rpath=${ORIGIN}/../../../Engine/Binaries/ThirdParty/jemalloc/Linux/x86_64-unknown-linux-gnu";
+            Result += " -Wl,-rpath=${ORIGIN}/../../../Engine/Binaries/ThirdParty/ICU/icu4c-53_1/Linux/x86_64-unknown-linux-gnu";
 
             if (CrossCompiling())
             {
