Linux: more rpath hardcoding.

diff --git a/Engine/Source/Programs/UnrealBuildTool/Linux/LinuxToolChain.cs b/Engine/Source/Programs/UnrealBuildTool/Linux/LinuxToolChain.cs
--- a/Engine/Source/Programs/UnrealBuildTool/Linux/LinuxToolChain.cs
+++ b/Engine/Source/Programs/UnrealBuildTool/Linux/LinuxToolChain.cs
@@ -279,6 +279,8 @@ namespace UnrealBuildTool.Linux
             // FIXME: really ugly temp solution. Modules need to be able to specify this
             Result += " -Wl,-rpath=${ORIGIN}/../../../Engine/Binaries/ThirdParty/jemalloc/Linux/x86_64-unknown-linux-gnu";
             Result += " -Wl,-rpath=${ORIGIN}/../../../Engine/Binaries/ThirdParty/ICU/icu4c-53_1/Linux/x86_64-unknown-linux-gnu";
+            Result += " -Wl,-rpath=${ORIGIN}/../../../Engine/Binaries/ThirdParty/SDL2/Linux/x86_64-unknown-linux-gnu";
+            Result += " -Wl,-rpath=${ORIGIN}/../../../Engine/Binaries/ThirdParty/zlib/Linux/x86_64-unknown-linux-gnu";
 
             if (CrossCompiling())
             {
