UBT - ForceNameAsForDevelopment is valid only for Shipping config, so the Debug will built with the regular name

diff --git a/Engine/Source/Programs/UnrealBuildTool/Configuration/UEBuildTarget.cs b/Engine/Source/Programs/UnrealBuildTool/Configuration/UEBuildTarget.cs
--- a/Engine/Source/Programs/UnrealBuildTool/Configuration/UEBuildTarget.cs
+++ b/Engine/Source/Programs/UnrealBuildTool/Configuration/UEBuildTarget.cs
@@ -705,7 +705,7 @@ namespace UnrealBuildTool
 
 			// Construct the output path based on configuration, platform, game if not specified.
 			TargetRules.TargetType? TargetType = (Rules != null) ? Rules.Type : (TargetRules.TargetType?)null;
-			OutputPath = Path.GetFullPath(MakeBinaryPath("", AppName, UEBuildBinaryType.Executable, TargetType, bIsRocketGame, null, InAppName,Rules.ForceNameAsForDevelopment()));
+			OutputPath = Path.GetFullPath(MakeBinaryPath("", AppName, UEBuildBinaryType.Executable, TargetType, bIsRocketGame, null, InAppName, Configuration==UnrealTargetConfiguration.Shipping?Rules.ForceNameAsForDevelopment():false));
 
 			if (bCompileMonolithic && TargetRules.IsGameType(InRulesObject.Type))
 			{
