Fix scraping of environment variables in non-US locales. Include paths, libraries and programs may be installed to non-ANSI directory names (eg. C:\\Pliki Program≤w\\ in Polish), and CMD.EXE needs to be run with the /U switch to allow Unicode output.

diff --git a/Engine/Source/Programs/UnrealBuildTool/System/Utils.cs b/Engine/Source/Programs/UnrealBuildTool/System/Utils.cs
--- a/Engine/Source/Programs/UnrealBuildTool/System/Utils.cs
+++ b/Engine/Source/Programs/UnrealBuildTool/System/Utils.cs
@@ -281,8 +281,10 @@ namespace UnrealBuildTool
 				// process needs to be disposed when done
 				using(var BatchFileProcess = new Process())
 				{
+					// Run the batch file using cmd.exe with the /U option, to force Unicode output. Many locales have non-ANSI characters in system paths.
 					var StartInfo = BatchFileProcess.StartInfo;
-					StartInfo.FileName = EnvReaderBatchFileName;
+					StartInfo.FileName = Path.Combine(Environment.SystemDirectory, "cmd.exe");
+					StartInfo.Arguments = String.Format("/U /C \"{0}\"", EnvReaderBatchFileName);
 					StartInfo.CreateNoWindow = true;
 					StartInfo.UseShellExecute = false;
 					StartInfo.RedirectStandardOutput = true;
@@ -306,8 +308,8 @@ namespace UnrealBuildTool
 					Log.TraceVerbose( "Finished launching {0}.", StartInfo.FileName );
 				}
 
-				// Load environment variables
-				var EnvStringsFromFile = File.ReadAllLines( EnvOutputFileName );
+				// Load environment variables (the file will be encoded without a BOM, so we need to manually specify the encoding)
+				var EnvStringsFromFile = File.ReadAllLines( EnvOutputFileName, Encoding.Unicode );
 				foreach( var EnvString in EnvStringsFromFile )
 				{
 					// Parse the environment variable name and value from the string ("name=value")
