Fix intellisense generation for Rocket Xcode projects on Mac

diff --git a/Engine/Source/Programs/UnrealBuildTool/Mac/UEBuildMac.cs b/Engine/Source/Programs/UnrealBuildTool/Mac/UEBuildMac.cs
--- a/Engine/Source/Programs/UnrealBuildTool/Mac/UEBuildMac.cs
+++ b/Engine/Source/Programs/UnrealBuildTool/Mac/UEBuildMac.cs
@@ -166,8 +166,9 @@ namespace UnrealBuildTool
 
         public override void ValidateUEBuildConfiguration()
         {
-			if (ProjectFileGenerator.bGenerateProjectFiles || ProjectFileGenerator.bGeneratingRocketProjectFiles)
+			if (ProjectFileGenerator.bGenerateProjectFiles && !ProjectFileGenerator.bGeneratingRocketProjectFiles)
 			{
+				// When generating non-Rocket project files we need intellisense generator to include info from all modules, including editor-only third party libs
 				UEBuildConfiguration.bCompileLeanAndMeanUE = false;
 			}
         }
