Fixed a problem with UBT not having bRunningRocket set when called directly from UAT

diff --git a/Engine/Source/Programs/AutomationTool/Automation.cs b/Engine/Source/Programs/AutomationTool/Automation.cs
--- a/Engine/Source/Programs/AutomationTool/Automation.cs
+++ b/Engine/Source/Programs/AutomationTool/Automation.cs
@@ -254,6 +254,11 @@ AutomationTool.exe [-verbose] [-compileonly] [-p4] Command0 [-Arg0 -Arg1 -Arg2 
 			{
 				throw new AutomationException("{0} and {1} can't be set simultaneously.", GlobalCommandLine.NoSubmit.Name, GlobalCommandLine.Submit.Name);
 			}
+
+			if (GlobalCommandLine.Rocket)
+			{
+				UnrealBuildTool.UnrealBuildTool.bRunningRocket = true;
+			}
 		}
 
 		#endregion
diff --git a/Engine/Source/Programs/UnrealBuildTool/System/UnrealBuildTool.cs b/Engine/Source/Programs/UnrealBuildTool/System/UnrealBuildTool.cs
--- a/Engine/Source/Programs/UnrealBuildTool/System/UnrealBuildTool.cs
+++ b/Engine/Source/Programs/UnrealBuildTool/System/UnrealBuildTool.cs
@@ -51,7 +51,7 @@ namespace UnrealBuildTool
         static UnrealTargetPlatform OnlyPlatformSpecificFor = UnrealTargetPlatform.Unknown;
 
         /** Are we running for Rocket */
-		static bool bRunningRocket = false;
+		static public bool bRunningRocket = false;
 
 		/** Are we building Rocket */
 		static bool bBuildingRocket = false;
