Removed some Xcode project settings that we no longer need and that were causing warnings when opened in recent Xcode version

diff --git a/Engine/Source/Programs/UnrealBuildTool/System/XcodeProjectFileGenerator.cs b/Engine/Source/Programs/UnrealBuildTool/System/XcodeProjectFileGenerator.cs
--- a/Engine/Source/Programs/UnrealBuildTool/System/XcodeProjectFileGenerator.cs
+++ b/Engine/Source/Programs/UnrealBuildTool/System/XcodeProjectFileGenerator.cs
@@ -443,7 +443,7 @@ namespace UnrealBuildTool
 				"\t\t\t\tGCC_ENABLE_CPP_RTTI = NO;" + ProjectFileGenerator.NewLine +
 				"\t\t\t\tGCC_WARN_CHECK_SWITCH_STATEMENTS = NO;" + ProjectFileGenerator.NewLine +
 				"\t\t\t\tSUPPORTED_PLATFORMS = \"macosx\";" + ProjectFileGenerator.NewLine +
-				"\t\t\t\tONLY_ACTIVE_ARCH = YES;" + ProjectFileGenerator.NewLine +
+				"\t\t\t\tONLY_ACTIVE_ARCH = " + (ConfigName == "Debug" ? "YES;" : "NO;") + ProjectFileGenerator.NewLine +
 				"\t\t\t\tUSE_HEADERMAP = NO;" + ProjectFileGenerator.NewLine +
 				"\t\t\t\tSDKROOT = macosx;" + ProjectFileGenerator.NewLine +
 				"\t\t\t\tSYMROOT = " + EngineSubdir + "Intermediate/Build;" + ProjectFileGenerator.NewLine +
@@ -462,7 +462,6 @@ namespace UnrealBuildTool
 					"\t\t\tisa = XCBuildConfiguration;" + ProjectFileGenerator.NewLine +
 					"\t\t\tbuildSettings = {" + ProjectFileGenerator.NewLine +
 					"\t\t\t\tCOMBINE_HIDPI_IMAGES = YES;" + ProjectFileGenerator.NewLine +
-					"\t\t\t\tARCHS = \"x86_64\";" + ProjectFileGenerator.NewLine +
 					"\t\t\t\tVALID_ARCHS = \"x86_64\";" + ProjectFileGenerator.NewLine +
 					"\t\t\t\tPRODUCT_NAME = \"$(TARGET_NAME)\";" + ProjectFileGenerator.NewLine +
 					"\t\t\t\tMACOSX_DEPLOYMENT_TARGET = 10.9;" + ProjectFileGenerator.NewLine +
@@ -480,7 +479,6 @@ namespace UnrealBuildTool
 					"\t\t\tisa = XCBuildConfiguration;" + ProjectFileGenerator.NewLine +
 					"\t\t\tbuildSettings = {" + ProjectFileGenerator.NewLine +
 					"\t\t\t\tCOMBINE_HIDPI_IMAGES = YES;" + ProjectFileGenerator.NewLine +
-					"\t\t\t\tARCHS = \"x86_64 arm64 armv7 armv7s\";" + ProjectFileGenerator.NewLine +
 					"\t\t\t\tVALID_ARCHS = \"x86_64 arm64 armv7 armv7s\";" + ProjectFileGenerator.NewLine +
 					"\t\t\t\tSUPPORTED_PLATFORMS = \"macosx iphoneos iphonesimulator\";" + ProjectFileGenerator.NewLine +
 					"\t\t\t\t\"PRODUCT_NAME[sdk=macosx*]\" = \"$(TARGET_NAME)\";" + ProjectFileGenerator.NewLine +
@@ -613,7 +611,7 @@ namespace UnrealBuildTool
 				"\t\t\t\tCONFIGURATION_BUILD_DIR = \"" + EngineSubdir + "Binaries/IOS/Payload\";" + ProjectFileGenerator.NewLine +
 				"\t\t\t\tSUPPORTED_PLATFORMS = \"iphoneos iphonesimulator\";" + ProjectFileGenerator.NewLine +
 				"\t\t\t\tIPHONEOS_DEPLOYMENT_TARGET = 7.0;" + ProjectFileGenerator.NewLine +
-				"\t\t\t\tONLY_ACTIVE_ARCH = YES;" + ProjectFileGenerator.NewLine +
+				"\t\t\t\tONLY_ACTIVE_ARCH = " + (ConfigName == "Debug" ? "YES;" : "NO;") + ProjectFileGenerator.NewLine +
 				"\t\t\t\tPRODUCT_NAME = \"$(TARGET_NAME)\";" + ProjectFileGenerator.NewLine +
 				"\t\t\t\tSDKROOT = iphoneos;" + ProjectFileGenerator.NewLine +
 				"\t\t\t\tTEST_HOST = \"$(BUNDLE_LOADER)\";" + ProjectFileGenerator.NewLine +
