Kill some spam output from SDK switching.

diff --git a/Engine/Source/Programs/UnrealBuildTool/Configuration/UEBuildPlatform.cs b/Engine/Source/Programs/UnrealBuildTool/Configuration/UEBuildPlatform.cs
--- a/Engine/Source/Programs/UnrealBuildTool/Configuration/UEBuildPlatform.cs
+++ b/Engine/Source/Programs/UnrealBuildTool/Configuration/UEBuildPlatform.cs
@@ -774,7 +774,10 @@ namespace UnrealBuildTool
                     {
                         if (!ModifiedPathVars.Contains(PathAdd))
                         {
-                            Console.WriteLine("Adding Path: '{0}'", PathAdd);
+                            if (BuildConfiguration.bPrintDebugInfo)
+                            {
+                                Console.WriteLine("Adding Path: '{0}'", PathAdd);
+                            }
                             ModifiedPathVars.Add(PathAdd);
                         }
                     }
