#UE #HTML5

          - remove deprecated jcache, and mostly useless -g from the link step.
          - add -g2 to preserve  function names/ not miniify in development builds.

diff --git a/Engine/Source/Programs/UnrealBuildTool/HTML5/HTML5ToolChain.cs b/Engine/Source/Programs/UnrealBuildTool/HTML5/HTML5ToolChain.cs
--- a/Engine/Source/Programs/UnrealBuildTool/HTML5/HTML5ToolChain.cs
+++ b/Engine/Source/Programs/UnrealBuildTool/HTML5/HTML5ToolChain.cs
@@ -179,12 +179,6 @@ namespace UnrealBuildTool
 				// enable verbose mode
 				Result += " -v";
 
-				if (LinkEnvironment.Config.Target.Configuration != CPPTargetConfiguration.Shipping)
-				{
-					// enable caching of intermediate date
-					Result += " --jcache -g";
-				}
-
 				if (LinkEnvironment.Config.Target.Configuration == CPPTargetConfiguration.Debug)
 				{
 					// check for alignment/etc checking
@@ -206,7 +200,7 @@ namespace UnrealBuildTool
 				}
 				if (LinkEnvironment.Config.Target.Configuration == CPPTargetConfiguration.Development)
 				{
-					Result += " -O2 -s ASM_JS=1 -s OUTLINING_LIMIT=110000";
+					Result += " -O2 -s ASM_JS=1 -s OUTLINING_LIMIT=110000   -g2 ";
 				}
 				if (LinkEnvironment.Config.Target.Configuration == CPPTargetConfiguration.Shipping)
 				{
