Fix jumbo sized static libraries in Rocket builds. Engine doesn't use a PCH any more, which causes massive redundancy in object files.

diff --git a/Engine/Source/Programs/UnrealBuildTool/Configuration/UEBuildModule.cs b/Engine/Source/Programs/UnrealBuildTool/Configuration/UEBuildModule.cs
--- a/Engine/Source/Programs/UnrealBuildTool/Configuration/UEBuildModule.cs
+++ b/Engine/Source/Programs/UnrealBuildTool/Configuration/UEBuildModule.cs
@@ -1148,7 +1148,7 @@ namespace UnrealBuildTool
 					Log.TraceVerbose("Module '{0}' was not allowed to use SharedPCHs, because we're compiling to a library in monolithic mode", this.Name );
 				}
 
-				bool bUseSharedPCHFiles  = BuildConfiguration.bUseSharedPCHs && (bDisableSharedPCHFiles == false);
+				bool bUseSharedPCHFiles  = BuildConfiguration.bUseSharedPCHs && (bDisableSharedPCHFiles == false) && !UnrealBuildTool.BuildingRocket();
 				bool bIsASharedPCHModule = bUseSharedPCHFiles && GlobalCompileEnvironment.SharedPCHHeaderFiles.Any( PCH => PCH.Module == this );
 
 				// Map from pch header string to the source files that use that PCH
