fix for generating the stub when doing a local build on a mac
TTP341023
#ue4
#ios

diff --git a/Engine/Source/Programs/AutomationTool/UBTUtils.cs b/Engine/Source/Programs/AutomationTool/UBTUtils.cs
--- a/Engine/Source/Programs/AutomationTool/UBTUtils.cs
+++ b/Engine/Source/Programs/AutomationTool/UBTUtils.cs
@@ -26,6 +26,10 @@ namespace AutomationTool
 			{
 				CommandLine += " -rocket";
 			}
+			if (!IsBuildMachine && UnrealBuildTool.ExternalExecution.GetRuntimePlatform() == UnrealBuildTool.UnrealTargetPlatform.Mac)
+			{
+				CommandLine += " -nocreatestub";
+			}
 			CommandUtils.RunAndLog(Env, UBTExecutable, CommandLine, LogName);
 		}
 
diff --git a/Engine/Source/Programs/UnrealBuildTool/System/UnrealBuildTool.cs b/Engine/Source/Programs/UnrealBuildTool/System/UnrealBuildTool.cs
--- a/Engine/Source/Programs/UnrealBuildTool/System/UnrealBuildTool.cs
+++ b/Engine/Source/Programs/UnrealBuildTool/System/UnrealBuildTool.cs
@@ -1063,7 +1063,7 @@ namespace UnrealBuildTool
 				}
 
 				bool CreateStub = Utils.ParseCommandLineFlag(Arguments, "-nocreatestub", out ArgumentIndex);
-				if (CreateStub)
+				if (CreateStub || (String.IsNullOrEmpty(Environment.GetEnvironmentVariable("uebp_LOCAL_ROOT")) && ExternalExecution.GetRuntimePlatform() == UnrealTargetPlatform.Mac))
 				{
 					BuildConfiguration.bCreateStubIPA = false;
 				}
