Fixed one more Xcode problem with too long user header paths and predefined macros being too long for its build system.

diff --git a/Engine/Source/Programs/UnrealBuildTool/System/XcodeProjectFileGenerator.cs b/Engine/Source/Programs/UnrealBuildTool/System/XcodeProjectFileGenerator.cs
--- a/Engine/Source/Programs/UnrealBuildTool/System/XcodeProjectFileGenerator.cs
+++ b/Engine/Source/Programs/UnrealBuildTool/System/XcodeProjectFileGenerator.cs
@@ -567,6 +567,12 @@ namespace UnrealBuildTool
 					"\t\t\t\tCONFIGURATION_BUILD_DIR = \"" + EngineRelative + "Engine/Binaries/IOS/Payload\";" + ProjectFileGenerator.NewLine);
 			}
 			Contents.Append(
+				"\t\t\t\tGCC_PREPROCESSOR_DEFINITIONS = (" + ProjectFileGenerator.NewLine +
+				"\t\t\t\t);" + ProjectFileGenerator.NewLine +
+				"\t\t\t\tHEADER_SEARCH_PATHS = (" + ProjectFileGenerator.NewLine +
+				"\t\t\t\t);" + ProjectFileGenerator.NewLine +
+				"\t\t\t\tUSER_HEADER_SEARCH_PATHS = (" + ProjectFileGenerator.NewLine +
+				"\t\t\t\t);" + ProjectFileGenerator.NewLine +
 				"\t\t\t\tINFOPLIST_OUTPUT_FORMAT = xml;" + ProjectFileGenerator.NewLine +
 				"\t\t\t\t\"PROVISIONING_PROFILE[sdk=iphoneos*]\" = \"\";" + ProjectFileGenerator.NewLine +
 				"\t\t\t\tTARGETED_DEVICE_FAMILY = \"1,2\";" + ProjectFileGenerator.NewLine +
@@ -616,6 +622,12 @@ namespace UnrealBuildTool
 				"\t\t\t\tSDKROOT = iphoneos;" + ProjectFileGenerator.NewLine +
 				"\t\t\t\tTEST_HOST = \"$(BUNDLE_LOADER)\";" + ProjectFileGenerator.NewLine +
 				"\t\t\t\tWRAPPER_EXTENSION = xctest;" + ProjectFileGenerator.NewLine +
+				"\t\t\t\tGCC_PREPROCESSOR_DEFINITIONS = (" + ProjectFileGenerator.NewLine +
+				"\t\t\t\t);" + ProjectFileGenerator.NewLine +
+				"\t\t\t\tHEADER_SEARCH_PATHS = (" + ProjectFileGenerator.NewLine +
+				"\t\t\t\t);" + ProjectFileGenerator.NewLine +
+				"\t\t\t\tUSER_HEADER_SEARCH_PATHS = (" + ProjectFileGenerator.NewLine +
+				"\t\t\t\t);" + ProjectFileGenerator.NewLine +
 				"\t\t\t};" + ProjectFileGenerator.NewLine +
 				"\t\t\tname = " + ConfigName + ";" + ProjectFileGenerator.NewLine +
 				"\t\t};" + ProjectFileGenerator.NewLine);
