Disable XGE until Xoreax fixes Incredibuild (TTP #341174, IT ticket 65919).

(Merging CL 2192236)

diff --git a/Engine/Source/Programs/UnrealBuildTool/Linux/UEBuildLinux.cs b/Engine/Source/Programs/UnrealBuildTool/Linux/UEBuildLinux.cs
--- a/Engine/Source/Programs/UnrealBuildTool/Linux/UEBuildLinux.cs
+++ b/Engine/Source/Programs/UnrealBuildTool/Linux/UEBuildLinux.cs
@@ -223,6 +223,7 @@ namespace UnrealBuildTool.Linux
         public override void ValidateUEBuildConfiguration()
         {
             BuildConfiguration.bUseUnityBuild = true;
+            BuildConfiguration.bAllowXGE = false;   // disable XGE until Xoreax figures out the reason of and fixes Incredibuild crashes (e.g. TTP #341174).
 
             // Don't stop compilation at first error...
             BuildConfiguration.bStopXGECompilationAfterErrors = true;
