Windows XP build is now configurable from the BuildConfiguration.xml file.

diff --git a/Engine/Source/Programs/UnrealBuildTool/Windows/UEBuildWindows.cs b/Engine/Source/Programs/UnrealBuildTool/Windows/UEBuildWindows.cs
--- a/Engine/Source/Programs/UnrealBuildTool/Windows/UEBuildWindows.cs
+++ b/Engine/Source/Programs/UnrealBuildTool/Windows/UEBuildWindows.cs
@@ -65,6 +65,7 @@ namespace UnrealBuildTool
 
 		/// True if we're targeting Windows XP as a minimum spec.  In Visual Studio 2012 and higher, this may change how
 		/// we compile and link the application (http://blogs.msdn.com/b/vcblog/archive/2012/10/08/10357555.aspx)
+		[XmlConfig]
 		public static bool SupportWindowsXP = false;
 
 		/** True if VS EnvDTE is available (false when building using Visual Studio Express) */
@@ -373,8 +374,7 @@ namespace UnrealBuildTool
 			InBuildTarget.GlobalCompileEnvironment.Config.Definitions.Add("WIN32=1");
 
 			// Win32 XP is only supported at this time.
-			SupportWindowsXP = false;
-				//&& InBuildTarget.Platform == UnrealTargetPlatform.Win32;
+			SupportWindowsXP = SupportWindowsXP && InBuildTarget.Platform == UnrealTargetPlatform.Win32;
 
 			if( SupportWindowsXP )
 			{
