Changed UEBuildWinRT.RegisterBuildPlatform() to do nothing on Mono to solve a problem running UBT.exe built on Windows.

diff --git a/Engine/Source/Programs/UnrealBuildTool/WinRT/UEBuildWinRT.cs b/Engine/Source/Programs/UnrealBuildTool/WinRT/UEBuildWinRT.cs
--- a/Engine/Source/Programs/UnrealBuildTool/WinRT/UEBuildWinRT.cs
+++ b/Engine/Source/Programs/UnrealBuildTool/WinRT/UEBuildWinRT.cs
@@ -54,7 +54,7 @@ namespace UnrealBuildTool
 		public override void RegisterBuildPlatform()
 		{
 			[REDACTED_EMAIL]: Add platform support
-			if (UnrealBuildTool.RunningRocket() || UnrealBuildTool.BuildingRocket())
+			if (UnrealBuildTool.RunningRocket() || UnrealBuildTool.BuildingRocket() || Utils.IsRunningOnMono)
 			{
 				return;
 			}
