fixes to XcodeProjectFileGenerator to add workspace folder
TTP337507
#ue4
#ios

diff --git a/Engine/Source/Programs/UnrealBuildTool/System/XcodeProjectFileGenerator.cs b/Engine/Source/Programs/UnrealBuildTool/System/XcodeProjectFileGenerator.cs
--- a/Engine/Source/Programs/UnrealBuildTool/System/XcodeProjectFileGenerator.cs
+++ b/Engine/Source/Programs/UnrealBuildTool/System/XcodeProjectFileGenerator.cs
@@ -1384,6 +1384,27 @@ namespace UnrealBuildTool
 			WriteSchemes(XcodeProjectPath, ProjectTargets);
 		}
 
+		private void WriteWorkspace()
+		{
+			string PathBranch = (ExternalExecution.GetRuntimePlatform() != UnrealTargetPlatform.Mac) ? "/Engine/Intermediate/IOS" : "";
+			var XcodeProjectPath = MasterProjectRelativePath + PathBranch + "/" + MasterProjectName + ".xcodeproj";
+			if (bGeneratingGameProjectFiles && GameProjectName == "UE4Game")
+			{
+				XcodeProjectPath = UnrealBuildTool.GetUProjectPath() + "/" + MasterProjectName + ".xcodeproj";
+			}
+			if (!Directory.Exists(XcodeProjectPath))
+			{
+				Directory.CreateDirectory(XcodeProjectPath);
+			}
+
+			// create the workspace folder
+			XcodeProjectPath += "/project.xcworkspace";
+			if (!Directory.Exists(XcodeProjectPath))
+			{
+				Directory.CreateDirectory(XcodeProjectPath);
+			}
+		}
+
 		/// <summary>
 		/// Writes the project files to disk
 		/// </summary>
@@ -1638,6 +1659,9 @@ namespace UnrealBuildTool
 
 			WriteProject(XcodeProjectFileContent, ProjectTargets);
 
+			// TODO: add any content for any files in the workspace here
+			WriteWorkspace();
+
 			return true;
 		}
 
