Fix UBT not relinking targets after plugin static libraries are modified. Dependencies were ignored because their producing action was marked as creating an import library.

diff --git a/Engine/Source/Programs/UnrealBuildTool/Mac/MacToolChain.cs b/Engine/Source/Programs/UnrealBuildTool/Mac/MacToolChain.cs
--- a/Engine/Source/Programs/UnrealBuildTool/Mac/MacToolChain.cs
+++ b/Engine/Source/Programs/UnrealBuildTool/Mac/MacToolChain.cs
@@ -568,7 +568,7 @@ namespace UnrealBuildTool
 
 			// Tell the action that we're building an import library here and it should conditionally be
 			// ignored as a prerequisite for other actions
-			LinkAction.bProducesImportLibrary = !Utils.IsRunningOnMono && (bIsBuildingLibrary || LinkEnvironment.Config.bIsBuildingDLL);
+			LinkAction.bProducesImportLibrary = !Utils.IsRunningOnMono && (bBuildImportLibraryOnly || LinkEnvironment.Config.bIsBuildingDLL);
 
 			// Add the output file as a production of the link action.
 			FileItem OutputFile = FileItem.GetItemByPath(Path.GetFullPath(LinkEnvironment.Config.OutputFilePath));
diff --git a/Engine/Source/Programs/UnrealBuildTool/WinRT/WinRTToolChain.cs b/Engine/Source/Programs/UnrealBuildTool/WinRT/WinRTToolChain.cs
--- a/Engine/Source/Programs/UnrealBuildTool/WinRT/WinRTToolChain.cs
+++ b/Engine/Source/Programs/UnrealBuildTool/WinRT/WinRTToolChain.cs
@@ -689,7 +689,7 @@ namespace UnrealBuildTool
 
 			// Tell the action that we're building an import library here and it should conditionally be
 			// ignored as a prerequisite for other actions
-			LinkAction.bProducesImportLibrary = bIsBuildingLibrary || LinkEnvironment.Config.bIsBuildingDLL;
+			LinkAction.bProducesImportLibrary = bBuildImportLibraryOnly || LinkEnvironment.Config.bIsBuildingDLL;
 
 			// If we're only building an import library, add the '/DEF' option that tells the LIB utility
 			// to simply create a .LIB file and .EXP file, and don't bother validating imports
diff --git a/Engine/Source/Programs/UnrealBuildTool/Windows/VCToolChain.cs b/Engine/Source/Programs/UnrealBuildTool/Windows/VCToolChain.cs
--- a/Engine/Source/Programs/UnrealBuildTool/Windows/VCToolChain.cs
+++ b/Engine/Source/Programs/UnrealBuildTool/Windows/VCToolChain.cs
@@ -1052,7 +1052,7 @@ namespace UnrealBuildTool
 
 			// Tell the action that we're building an import library here and it should conditionally be
 			// ignored as a prerequisite for other actions
-			LinkAction.bProducesImportLibrary = bIsBuildingLibrary || LinkEnvironment.Config.bIsBuildingDLL;
+			LinkAction.bProducesImportLibrary = bBuildImportLibraryOnly || LinkEnvironment.Config.bIsBuildingDLL;
 
 
 			// If we're only building an import library, add the '/DEF' option that tells the LIB utility
