bump up the unity file size for compiling to prevent jump distance linker errors
TTP341310
#ue4

diff --git a/Engine/Source/Programs/UnrealBuildTool/Configuration/BuildConfiguration.cs b/Engine/Source/Programs/UnrealBuildTool/Configuration/BuildConfiguration.cs
--- a/Engine/Source/Programs/UnrealBuildTool/Configuration/BuildConfiguration.cs
+++ b/Engine/Source/Programs/UnrealBuildTool/Configuration/BuildConfiguration.cs
@@ -405,7 +405,7 @@ namespace UnrealBuildTool
 
 			// This controls how big unity files are.  Use smaller values for faster iterative iteration times at the cost of slower full rebuild times.
 			// This setting can greatly affect how effective parallelized compilation can be.
-			NumIncludedBytesPerUnityCPP = 256 * 1024;
+			NumIncludedBytesPerUnityCPP = 384 * 1024;
 
 			// This sets the number of source files (post-unity-combine) in a module before we'll bother to generate a precompiled header
 			// for that module.  If you want the fastest iterative builds for non-header changes (at the expense of slower full rebuild times),
