Proper disabling of XGE (for UAT builds too).

diff --git a/Engine/Source/Programs/UnrealBuildTool/Linux/UEBuildLinux.cs b/Engine/Source/Programs/UnrealBuildTool/Linux/UEBuildLinux.cs
--- a/Engine/Source/Programs/UnrealBuildTool/Linux/UEBuildLinux.cs
+++ b/Engine/Source/Programs/UnrealBuildTool/Linux/UEBuildLinux.cs
@@ -96,6 +96,12 @@ namespace UnrealBuildTool.Linux
             return SDKStatus.Invalid;
         }
 
+        public override bool CanUseXGE()
+        {
+            // disabled until XGE crash is fixed
+            return false;
+        }
+        
         /**
          *	Register the platform with the UEBuildPlatform class
          */
@@ -223,7 +229,6 @@ namespace UnrealBuildTool.Linux
         public override void ValidateUEBuildConfiguration()
         {
             BuildConfiguration.bUseUnityBuild = true;
-            BuildConfiguration.bAllowXGE = false;   // disable XGE until Xoreax figures out the reason of and fixes Incredibuild crashes (e.g. TTP #341174).
 
             // Don't stop compilation at first error...
             BuildConfiguration.bStopXGECompilationAfterErrors = true;
