UnrealBuildTool compile times fixes
- Fixed modules that supply shared PCHs not generating and using their own unique PCH in monolithic builds
- Fixed monolithic builds not detecting that a unique PCH was needed for modules that supply shared PCHs when header file names were different
- Added better verbose debug output to detect problems with shared PCH headers

diff --git a/Engine/Source/Programs/UnrealBuildTool/Configuration/UEBuildModule.cs b/Engine/Source/Programs/UnrealBuildTool/Configuration/UEBuildModule.cs
--- a/Engine/Source/Programs/UnrealBuildTool/Configuration/UEBuildModule.cs
+++ b/Engine/Source/Programs/UnrealBuildTool/Configuration/UEBuildModule.cs
@@ -1145,12 +1145,18 @@ namespace UnrealBuildTool
 				bool bDisableSharedPCHFiles = (Binary.Config.bCompileMonolithic && CompileEnvironment.Config.bIsBuildingLibrary);
 				if( BuildConfiguration.bUseSharedPCHs && bDisableSharedPCHFiles )
 				{
-					Log.TraceVerbose("Module '{0}' was not allowed to use SharedPCHs, because we're compiling to a library in monolithic mode", this.Name );
+					Log.TraceVerbose("Module '{0}' was not allowed to use Shared PCHs, because we're compiling to a library in monolithic mode", this.Name );
 				}
 
 				bool bUseSharedPCHFiles  = BuildConfiguration.bUseSharedPCHs && (bDisableSharedPCHFiles == false) && !UnrealBuildTool.BuildingRocket();
 				bool bIsASharedPCHModule = bUseSharedPCHFiles && GlobalCompileEnvironment.SharedPCHHeaderFiles.Any( PCH => PCH.Module == this );
 
+				string SharedPCHHeaderFilePath = null;
+				if( bIsASharedPCHModule )
+				{
+					SharedPCHHeaderFilePath = Path.GetFullPath( Path.Combine( ProjectFileGenerator.RootRelativePath, "Engine", "Source", this.SharedPCHHeaderFile ) );
+				}
+
 				// Map from pch header string to the source files that use that PCH
 				var UsageMapPCH = new Dictionary<string, List<FileItem>>( StringComparer.InvariantCultureIgnoreCase );
 
@@ -1163,6 +1169,15 @@ namespace UnrealBuildTool
 					{
 						Log.TraceVerbose("ActorConstruction.cpp");
 					}
+
+					if( CPPFile.PrecompiledHeaderIncludeFilename == null )
+					{
+						throw new BuildException( "No PCH usage for file \"{0}\" . Missing #include header?", CPPFile.AbsolutePath );
+					}
+
+					// Create a new entry if not in the pch usage map
+					UsageMapPCH.GetOrAddNew( CPPFile.PrecompiledHeaderIncludeFilename ).Add( CPPFile );
+
 					if (bUseSharedPCHFiles)
 					{
 						var SharedPCHStartTime = DateTime.UtcNow;
@@ -1172,7 +1187,8 @@ namespace UnrealBuildTool
 						// @todo SharedPCH: If we ever have SharedPCH headers that themselves belong to modules which never use DLL Exports, we can avoid
 						// generating TWO PCH files by checking for that here.  For now, we always assume that SharedPCH headers have exports when
 						// compiling in modular mode.
-						if( bAllowSharedPCH && ( !bIsASharedPCHModule || bCompileMonolithic ) )
+						bool bCanModuleUseOwnSharedPCH = bAllowSharedPCH && bIsASharedPCHModule && bCompileMonolithic && CPPFile.PrecompiledHeaderIncludeFilename.Equals( SharedPCHHeaderFilePath, StringComparison.InvariantCultureIgnoreCase );
+						if( bAllowSharedPCH && ( !bIsASharedPCHModule || bCanModuleUseOwnSharedPCH ) )
 						{
 							// Figure out which shared PCH tier we're in
 							int LargestSharedPCHHeaderFileIndex = -1;
@@ -1216,24 +1232,16 @@ namespace UnrealBuildTool
 							}
 							else
 							{
-								Log.TraceVerbose("File {0} doesn't use a SharedPCH!", CPPFile.AbsolutePath);
+								Log.TraceVerbose("File {0} doesn't use a Shared PCH!", CPPFile.AbsolutePath);
 							}
 
 							SharedPCHTotalTime += ( DateTime.UtcNow - SharedPCHStartTime ).TotalSeconds;
 						}
 						else
 						{
-							Log.TraceVerbose("File '{0}' cannot create or use SharedPCHs, because its module '{1}' needs its own private PCH", CPPFile.AbsolutePath, this.Name);
+							Log.TraceVerbose("File '{0}' cannot create or use Shared PCHs, because its module '{1}' needs its own private PCH", CPPFile.AbsolutePath, this.Name);
 						}
 					}
-
-					if( CPPFile.PrecompiledHeaderIncludeFilename == null )
-					{
-						throw new BuildException( "No PCH usage for file \"{0}\" . Missing #include header?", CPPFile.AbsolutePath );
-					}
-
-					// Create a new entry if not in the pch usage map
-					UsageMapPCH.GetOrAddNew( CPPFile.PrecompiledHeaderIncludeFilename ).Add( CPPFile );
 				}
 
 				if( BuildConfiguration.bPrintDebugInfo )
@@ -1298,7 +1306,7 @@ namespace UnrealBuildTool
 							//    Currently, it's possible for the shared PCH to be compiled differently depending on which module UBT happened to have
 							//    include it first during the build phase.  This could create problems with determinstic builds, or turn up compile
 							//    errors unexpectedly due to compile environment differences.
-							Log.TraceVerbose("Module " + Name + " uses existing SharedPCH '" + ModulePCHEnvironment.PrecompiledHeaderIncludeFilename + "' (from module " + ModulePCHEnvironment.ModuleName + ")");
+							Log.TraceVerbose("Module " + Name + " uses existing Shared PCH '" + ModulePCHEnvironment.PrecompiledHeaderIncludeFilename + "' (from module " + ModulePCHEnvironment.ModuleName + ")");
 						}
 
 						PCHToUse = SharedPCHHeaderFile;
@@ -1332,20 +1340,22 @@ namespace UnrealBuildTool
 
 						ModulePCHEnvironment = new PrecompileHeaderEnvironment( PCHModuleName, PCHHeaderNameInCode, PCHHeaderFile, ModuleCompileEnvironment.Config.CLRMode, ModuleCompileEnvironment.Config.OptimizeCode );
 
-						Log.TraceVerbose( "PCH file \"{0}\" generated for module \"{1}\" .", PCHHeaderFile.AbsolutePath, Name );
-
 						if( SharedPCHHeaderFile != null )
 						{
 							// Add to list of shared PCH environments
 							GlobalCompileEnvironment.SharedPCHEnvironments.Add( ModulePCHEnvironment );
-							Log.TraceVerbose( "Module " + Name + " uses new SharedPCH '" + ModulePCHEnvironment.PrecompiledHeaderIncludeFilename + "'" );
+							Log.TraceVerbose( "Module " + Name + " uses new Shared PCH '" + ModulePCHEnvironment.PrecompiledHeaderIncludeFilename + "'" );
 						}
 						else
 						{
-							Log.TraceVerbose( "Module " + Name + " uses a unique PCH '" + ModulePCHEnvironment.PrecompiledHeaderIncludeFilename + "'" );
+							Log.TraceVerbose( "Module " + Name + " uses a Unique PCH '" + ModulePCHEnvironment.PrecompiledHeaderIncludeFilename + "'" );
 						}
 					}
 				}
+				else
+				{
+					Log.TraceVerbose( "Module " + Name + " doesn't use a Shared PCH, and only has " + SourceFilesToBuild.CPPFiles.Count.ToString() + " source file(s).  No Unique PCH will be generated." );
+				}
 
 				// Compile the C++ source or the unity C++ files that use a PCH environment.
 				if( ModulePCHEnvironment != null )
@@ -1373,35 +1383,42 @@ namespace UnrealBuildTool
 					}
 
 					// Check if there are enough unity files to warrant pch generation (and we haven't already generated the shared one)
-					if( ModulePCHEnvironment.PrecompiledHeaderFile == null && ( SharedPCHHeaderFile != null || CPPFilesToBuild.Count >= MinFilesUsingPrecompiledHeader ) )
+					if( ModulePCHEnvironment.PrecompiledHeaderFile == null )
 					{
-						bool bAllowDLLExports = true;
-						var PCHOutputDirectory = ModuleCompileEnvironment.Config.OutputDirectory;
-						var PCHModuleName = this.Name;
-
-						if( SharedPCHHeaderFile != null )
+						if( SharedPCHHeaderFile != null || CPPFilesToBuild.Count >= MinFilesUsingPrecompiledHeader )
 						{
-							// Disallow DLLExports when generating shared PCHs.  These headers aren't able to export anything, because they're potentially shared between many modules.
-							bAllowDLLExports = false;
+							bool bAllowDLLExports = true;
+							var PCHOutputDirectory = ModuleCompileEnvironment.Config.OutputDirectory;
+							var PCHModuleName = this.Name;
+
+							if( SharedPCHHeaderFile != null )
+							{
+								// Disallow DLLExports when generating shared PCHs.  These headers aren't able to export anything, because they're potentially shared between many modules.
+								bAllowDLLExports = false;
 
-							// Save shared PCHs to a specific folder
-							PCHOutputDirectory = Path.Combine( CompileEnvironment.Config.OutputDirectory, "SharedPCHs" );
+								// Save shared PCHs to a specific folder
+								PCHOutputDirectory = Path.Combine( CompileEnvironment.Config.OutputDirectory, "SharedPCHs" );
 
-							// Use a fake module name for "shared" PCHs.  It may be used by many modules, so we don't want to use this module's name.
-							PCHModuleName = "Shared";
-						}
+								// Use a fake module name for "shared" PCHs.  It may be used by many modules, so we don't want to use this module's name.
+								PCHModuleName = "Shared";
+							}
 
-						var PCHOutput = PrecompileHeaderEnvironment.GeneratePCHCreationAction( 
-							CPPFilesToBuild[0].PCHHeaderNameInCode,
-							ModulePCHEnvironment.PrecompiledHeaderIncludeFilename,
-							ModuleCompileEnvironment, 
-							PCHOutputDirectory,
-							PCHModuleName, 
-							bAllowDLLExports );
-						ModulePCHEnvironment.PrecompiledHeaderFile = PCHOutput.PrecompiledHeaderFile;
+							var PCHOutput = PrecompileHeaderEnvironment.GeneratePCHCreationAction( 
+								CPPFilesToBuild[0].PCHHeaderNameInCode,
+								ModulePCHEnvironment.PrecompiledHeaderIncludeFilename,
+								ModuleCompileEnvironment, 
+								PCHOutputDirectory,
+								PCHModuleName, 
+								bAllowDLLExports );
+							ModulePCHEnvironment.PrecompiledHeaderFile = PCHOutput.PrecompiledHeaderFile;
 							
-						ModulePCHEnvironment.OutputObjectFiles.Clear();
-						ModulePCHEnvironment.OutputObjectFiles.AddRange( PCHOutput.ObjectFiles );
+							ModulePCHEnvironment.OutputObjectFiles.Clear();
+							ModulePCHEnvironment.OutputObjectFiles.AddRange( PCHOutput.ObjectFiles );
+						}
+						else if( CPPFilesToBuild.Count < MinFilesUsingPrecompiledHeader )
+						{
+							Log.TraceVerbose( "Module " + Name + " doesn't use a Shared PCH, and only has " + CPPFilesToBuild.Count.ToString() + " unity source file(s).  No Unique PCH will be generated." );
+						}
 					}
 
 					if( ModulePCHEnvironment.PrecompiledHeaderFile != null )
@@ -1487,7 +1504,7 @@ namespace UnrealBuildTool
 			// Don't mix CLR modes
 			if (SharedPCHEnvironment.CLRMode != ModuleCompileEnvironment.Config.CLRMode)
 			{
-				Log.TraceVerbose("Module {0} cannot use existing SharedPCH '{1}' (from module '{2}') because CLR modes don't match", Name, SharedPCHEnvironment.PrecompiledHeaderIncludeFilename.AbsolutePath, SharedPCHEnvironment.ModuleName);
+				Log.TraceVerbose("Module {0} cannot use existing Shared PCH '{1}' (from module '{2}') because CLR modes don't match", Name, SharedPCHEnvironment.PrecompiledHeaderIncludeFilename.AbsolutePath, SharedPCHEnvironment.ModuleName);
 				SharedPCHHeaderFile = null;
 				return null;
 			}
@@ -1511,7 +1528,7 @@ namespace UnrealBuildTool
 
 			if (SharedPCHCodeOptimization != ModuleCodeOptimization)
 			{
-				Log.TraceVerbose("Module {0} cannot use existing SharedPCH '{1}' (from module '{2}') because optimization levels don't match", Name, SharedPCHEnvironment.PrecompiledHeaderIncludeFilename.AbsolutePath, SharedPCHEnvironment.ModuleName);
+				Log.TraceVerbose("Module {0} cannot use existing Shared PCH '{1}' (from module '{2}') because optimization levels don't match", Name, SharedPCHEnvironment.PrecompiledHeaderIncludeFilename.AbsolutePath, SharedPCHEnvironment.ModuleName);
 				SharedPCHHeaderFile = null;
 				return null;
 			}
