Improved UnrealBuildTool performance diagnostics for time spent on PCHs

diff --git a/Engine/Source/Programs/UnrealBuildTool/Configuration/UEBuildModule.cs b/Engine/Source/Programs/UnrealBuildTool/Configuration/UEBuildModule.cs
--- a/Engine/Source/Programs/UnrealBuildTool/Configuration/UEBuildModule.cs
+++ b/Engine/Source/Programs/UnrealBuildTool/Configuration/UEBuildModule.cs
@@ -1161,7 +1161,7 @@ namespace UnrealBuildTool
 				var UsageMapPCH = new Dictionary<string, List<FileItem>>( StringComparer.InvariantCultureIgnoreCase );
 
 				// Determine what potential precompiled header is used by each source file.
-				double SharedPCHTotalTime = 0.0;
+				double SharedPCHTime = 0.0;
 				foreach( var CPPFile in SourceFilesToBuild.CPPFiles )
 				{
 					
@@ -1235,7 +1235,7 @@ namespace UnrealBuildTool
 								Log.TraceVerbose("File {0} doesn't use a Shared PCH!", CPPFile.AbsolutePath);
 							}
 
-							SharedPCHTotalTime += ( DateTime.UtcNow - SharedPCHStartTime ).TotalSeconds;
+							SharedPCHTime += ( DateTime.UtcNow - SharedPCHStartTime ).TotalSeconds;
 						}
 						else
 						{
@@ -1243,6 +1243,7 @@ namespace UnrealBuildTool
 						}
 					}
 				}
+				TotalSharedPCHTime += SharedPCHTime;
 
 				if( BuildConfiguration.bPrintDebugInfo )
 				{
@@ -1439,8 +1440,10 @@ namespace UnrealBuildTool
 
 				if( BuildConfiguration.bPrintPerformanceInfo )
 				{
-					var TotalPCHTime = DateTime.UtcNow - PCHTimerStart;
-					Trace.TraceInformation( "PCH time for " + Name + " is " + TotalPCHTime.TotalSeconds + "s (shared PCHs: " + SharedPCHTotalTime + "s)" );
+					var PCHTime = ( DateTime.UtcNow - PCHTimerStart ).TotalSeconds;
+					Trace.TraceInformation( "PCH time for " + Name + " is " + PCHTime + "s (shared PCHs: " + SharedPCHTime + "s)" );
+
+					TotalPCHTime += PCHTime;
 				}
 			}
 
@@ -1715,6 +1718,13 @@ namespace UnrealBuildTool
 
 		private UHTModuleInfo? CachedModuleUHTInfo = null;
 
+		/// Total time spent figuring out which PCH to use for each module and source file
+		public static double TotalPCHTime = 0.0;
+
+		/// Total time spent processing which shared PCH headers are needed for each module and source file
+		public static double TotalSharedPCHTime = 0.0;
+
+
 		/// <summary>
 		/// If any of this module's source files contain UObject definitions, this will return those header files back to the caller
 		/// </summary>
diff --git a/Engine/Source/Programs/UnrealBuildTool/System/UnrealBuildTool.cs b/Engine/Source/Programs/UnrealBuildTool/System/UnrealBuildTool.cs
--- a/Engine/Source/Programs/UnrealBuildTool/System/UnrealBuildTool.cs
+++ b/Engine/Source/Programs/UnrealBuildTool/System/UnrealBuildTool.cs
@@ -952,11 +952,11 @@ namespace UnrealBuildTool
 						Log.TraceInformation("GetIncludes time: " + CPPEnvironment.TotalTimeSpentGettingIncludes + "s (" + CPPEnvironment.TotalIncludesRequested + " includes)" );
 						Log.TraceInformation("DirectIncludes cache miss time: " + CPPEnvironment.DirectIncludeCacheMissesTotalTime + "s (" + CPPEnvironment.TotalDirectIncludeCacheMisses + " misses)" );
 						Log.TraceInformation("FindIncludePaths calls: " + CPPEnvironment.TotalFindIncludedFileCalls + " (" + CPPEnvironment.IncludePathSearchAttempts + " searches)" );
-						Log.TraceInformation("Total FileItems: {0} ({1} missing)", FileItem.TotalFileItemCount, FileItem.MissingFileItemCount);
+						Log.TraceInformation("PCH time: " + UEBuildModuleCPP.TotalPCHTime + "s (Shared PCH time: " + UEBuildModuleCPP.TotalSharedPCHTime + "s)" );
 						Log.TraceInformation("Include Resolves: {0} ({1} misses, {2:0.00}%)", CPPEnvironment.TotalDirectIncludeResolves, CPPEnvironment.TotalDirectIncludeResolveCacheMisses, (float)CPPEnvironment.TotalDirectIncludeResolveCacheMisses / (float)CPPEnvironment.TotalDirectIncludeResolves * 100);
 						Log.TraceInformation("Total FileItems: {0} ({1} missing)", FileItem.TotalFileItemCount, FileItem.MissingFileItemCount);
 
-						Log.TraceInformation("Execution time: {0}", BuildDuration);
+						Log.TraceInformation("Execution time: {0}s", BuildDuration);
 					}
 
 					Telemetry.SendEvent("PerformanceInfo.2",
