Add /FS flag to VCToolChain in VS2013 only to ensure parallel builds that use PDB work properly.

diff --git a/Engine/Source/Programs/UnrealBuildTool/Windows/VCToolChain.cs b/Engine/Source/Programs/UnrealBuildTool/Windows/VCToolChain.cs
--- a/Engine/Source/Programs/UnrealBuildTool/Windows/VCToolChain.cs
+++ b/Engine/Source/Programs/UnrealBuildTool/Windows/VCToolChain.cs
@@ -295,6 +295,12 @@ namespace UnrealBuildTool
 						{
 							Result += " /Zi";
 						}
+                        // We need to add this so VS won't lock the PDB file and prevent synchronous updates. This forces serialization through MSPDBSRV.exe.
+                        // See http://msdn.microsoft.com/en-us/library/dn502518.aspx for deeper discussion of /FS switch.
+                        if (BuildConfiguration.bUseIncrementalLinking)
+                        {
+                            Result += " /FS";
+                        }
 					}
 					// Store C7-format debug info in the .obj files, which is faster.
 					else
