Somewhat better info about loops in action graph.

diff --git a/Engine/Source/Programs/UnrealBuildTool/System/ActionGraph.cs b/Engine/Source/Programs/UnrealBuildTool/System/ActionGraph.cs
--- a/Engine/Source/Programs/UnrealBuildTool/System/ActionGraph.cs
+++ b/Engine/Source/Programs/UnrealBuildTool/System/ActionGraph.cs
@@ -27,6 +27,9 @@ namespace UnrealBuildTool
 	/** A build action. */
 	public class Action
 	{
+        /** Unique action identifier */
+        public int UniqueId;
+
 		public List<FileItem> PrerequisiteItems = new List<FileItem>();
 		public List<FileItem> ProducedItems = new List<FileItem>();
 		public delegate void EventHandler(Action A);
@@ -75,12 +78,15 @@ namespace UnrealBuildTool
 		/** True if any libraries produced by this action should be considered 'import libraries' */
 		public bool bProducesImportLibrary = false;
 
+        /** Always-incremented unique id */
+        private static int NextUniqueId = 0;
 
 		public Action(ActionType InActionType)
 		{
 			ActionType = InActionType;
 
 			UnrealBuildTool.AllActions.Add(this);
+            UniqueId = ++NextUniqueId;
 		}
 
 		/**
@@ -485,6 +491,7 @@ namespace UnrealBuildTool
 			// Starting with actions that only depend on non-produced items, iteratively expand a set of actions that are only dependent on
 			// non-cyclical actions.
 			Dictionary<Action, bool> ActionIsNonCyclical = new Dictionary<Action, bool>();
+            Dictionary<Action, List<Action>> CyclicActions = new Dictionary<Action, List<Action>>();
 			while (true)
 			{
 				bool bFoundNewNonCyclicalAction = false;
@@ -502,6 +509,16 @@ namespace UnrealBuildTool
 								if (!ActionIsNonCyclical.ContainsKey(PrerequisiteItem.ProducingAction))
 								{
 									bActionOnlyDependsOnNonCyclicalActions = false;
+                                    if (!CyclicActions.ContainsKey(Action))
+                                    {
+                                        CyclicActions.Add(Action, new List<Action>());
+                                    }
+
+                                    List<Action> CyclicPrereq = CyclicActions[Action];
+                                    if (!CyclicPrereq.Contains(PrerequisiteItem.ProducingAction))
+                                    {
+                                        CyclicPrereq.Add(PrerequisiteItem.ProducingAction);
+                                    }
 								}
 							}
 						}
@@ -532,21 +549,45 @@ namespace UnrealBuildTool
 				{
 					if (!ActionIsNonCyclical.ContainsKey(Action))
 					{
-						CycleDescription += string.Format("Action: {0}\r\n", Action.CommandPath);
-						CycleDescription += string.Format("\twith arguments: {0}\r\n", Action.CommandArguments);
+						CycleDescription += string.Format("Action #{0}: {1}\n", Action.UniqueId, Action.CommandPath);
+						CycleDescription += string.Format("\twith arguments: {0}\n", Action.CommandArguments);
 						foreach (FileItem PrerequisiteItem in Action.PrerequisiteItems)
 						{
-							CycleDescription += string.Format("\tdepends on: {0}\r\n", PrerequisiteItem.AbsolutePath);
+							CycleDescription += string.Format("\tdepends on: {0}\n", PrerequisiteItem.AbsolutePath);
 						}
 						foreach (FileItem ProducedItem in Action.ProducedItems)
 						{
-							CycleDescription += string.Format("\tproduces:   {0}\r\n", ProducedItem.AbsolutePath);
+							CycleDescription += string.Format("\tproduces:   {0}\n", ProducedItem.AbsolutePath);
 						}
-						CycleDescription += "\r\n\r\n";
+                        CycleDescription += string.Format("\tDepends on cyclic actions:\n");
+                        if (CyclicActions.ContainsKey(Action))
+                        {
+                            foreach (Action CyclicPrerequisiteAction in CyclicActions[Action])
+                            {
+                                if (CyclicPrerequisiteAction.ProducedItems.Count == 1)
+                                {
+                                    CycleDescription += string.Format("\t\t{0} (produces: {1})\n", CyclicPrerequisiteAction.UniqueId, CyclicPrerequisiteAction.ProducedItems[0].AbsolutePath);
+                                }
+                                else
+                                {
+                                    CycleDescription += string.Format("\t\t{0}\n", CyclicPrerequisiteAction.UniqueId);
+                                    foreach (FileItem CyclicProducedItem in CyclicPrerequisiteAction.ProducedItems)
+                                    {
+                                        CycleDescription += string.Format("\t\t\tproduces:   {0}\n", CyclicProducedItem.AbsolutePath);
+                                    }
+                                }
+                            }
+                            CycleDescription += "\n";
+                        }
+                        else
+                        {
+                            CycleDescription += string.Format("\t\tNone?? Coding error!\n");
+                        }
+                        CycleDescription += "\n\n";
 					}
 				}
 
-				throw new BuildException("Action graph contains cycle!\r\n\r\n{0}", CycleDescription);
+				throw new BuildException("Action graph contains cycle!\n\n{0}", CycleDescription);
 			}
 		}
 
@@ -601,7 +642,7 @@ namespace UnrealBuildTool
 						|| OldProducingCommandLine != NewProducingCommandLine)
 						{
 							Log.TraceVerbose(
-								"{0}: Produced item \"{1}\" was produced by outdated command-line.\r\nOld command-line: {2}\r\nNew command-line: {3}",
+								"{0}: Produced item \"{1}\" was produced by outdated command-line.\nOld command-line: {2}\nNew command-line: {3}",
 								RootAction.StatusDescription,
 								Path.GetFileName(ProducedItem.AbsolutePath),
 								OldProducingCommandLine,
