Use separate subfolder per architecture.

(contributed as part of PR 242).

diff --git a/Engine/Source/Programs/UnrealBuildTool/System/UnrealBuildTool.cs b/Engine/Source/Programs/UnrealBuildTool/System/UnrealBuildTool.cs
--- a/Engine/Source/Programs/UnrealBuildTool/System/UnrealBuildTool.cs
+++ b/Engine/Source/Programs/UnrealBuildTool/System/UnrealBuildTool.cs
@@ -918,7 +918,7 @@ namespace UnrealBuildTool
 								{
 									// Setup environment wasn't called, so set the flag
 									BuildConfiguration.bDeployAfterCompile = BuildPlatform.RequiresDeployPrepAfterCompile();
-									BuildConfiguration.PlatformIntermediateFolder = Path.Combine(BuildConfiguration.BaseIntermediateFolder, CheckPlatform.ToString() + BuildPlatform.GetActiveArchitecture());
+									BuildConfiguration.PlatformIntermediateFolder = Path.Combine(BuildConfiguration.BaseIntermediateFolder, CheckPlatform.ToString(), BuildPlatform.GetActiveArchitecture());
 								}
 							}
 
@@ -1040,7 +1040,7 @@ namespace UnrealBuildTool
 			var BuildPlatform = UEBuildPlatform.GetBuildPlatform(ResetPlatform);
 
 			// now that we have the platform, we can set the intermediate path to include the platform/architecture name
-			BuildConfiguration.PlatformIntermediateFolder = Path.Combine(BuildConfiguration.BaseIntermediateFolder, ResetPlatform.ToString() + BuildPlatform.GetActiveArchitecture());
+			BuildConfiguration.PlatformIntermediateFolder = Path.Combine(BuildConfiguration.BaseIntermediateFolder, ResetPlatform.ToString(), BuildPlatform.GetActiveArchitecture());
 
 			BuildPlatform.ResetBuildConfiguration(ResetPlatform, ResetConfiguration);
 
