Minor mono fixes (unused vars).

- Partially (except for PS4/XB1) in github PR #242.

diff --git a/Engine/Source/Programs/UnrealBuildTool/Android/AndroidToolChain.cs b/Engine/Source/Programs/UnrealBuildTool/Android/AndroidToolChain.cs
--- a/Engine/Source/Programs/UnrealBuildTool/Android/AndroidToolChain.cs
+++ b/Engine/Source/Programs/UnrealBuildTool/Android/AndroidToolChain.cs
@@ -338,9 +338,7 @@ namespace UnrealBuildTool
 						string[] Tokens = Line.Data.Split("(".ToCharArray());
 						if (Tokens.Length > 1)
 						{
-							// make sure what follows the parens is what we expdect
-							string[] ErrorPos = Tokens[1].Split(")".ToCharArray());
-
+							// make sure what follows the parens is what we expect
 							string Filename = Path.GetFullPath(Tokens[0]);
 							// build up the final string
 							string Output = string.Format("{0}({1}", Filename, Tokens[1], Line.Data[0]);
diff --git a/Engine/Source/Programs/UnrealBuildTool/HTML5/UEBuildHTML5.cs b/Engine/Source/Programs/UnrealBuildTool/HTML5/UEBuildHTML5.cs
--- a/Engine/Source/Programs/UnrealBuildTool/HTML5/UEBuildHTML5.cs
+++ b/Engine/Source/Programs/UnrealBuildTool/HTML5/UEBuildHTML5.cs
@@ -117,9 +117,9 @@ namespace UnrealBuildTool
 		/**
 		 *	Get the extension to use for the given binary type
 		 *	
-		 *	@param	InBinaryType		The binrary type being built
+		 *	@param	InBinaryType		The binary type being built
 		 *	
-		 *	@return	string				The binary extenstion (ie 'exe' or 'dll')
+		 *	@return	string				The binary extension (ie 'exe' or 'dll')
 		 */
 		public override string GetBinaryExtension(UEBuildBinaryType InBinaryType)
 		{
diff --git a/Engine/Source/Programs/UnrealBuildTool/System/MakefileGenerator.cs b/Engine/Source/Programs/UnrealBuildTool/System/MakefileGenerator.cs
--- a/Engine/Source/Programs/UnrealBuildTool/System/MakefileGenerator.cs
+++ b/Engine/Source/Programs/UnrealBuildTool/System/MakefileGenerator.cs
@@ -163,7 +163,7 @@ namespace UnrealBuildTool
         //protected override bool WriteMasterProjectFile( ProjectFile UBTProject )
         protected override bool WriteProjectFiles ()
         {
-            return WriteMakefile() & WriteCMakeLists();
+            return WriteMakefile() && WriteCMakeLists();
         }
 
         /// ProjectFileGenerator interface
diff --git a/Engine/Source/Programs/UnrealBuildTool/WinRT/WinRTDeploy.cs b/Engine/Source/Programs/UnrealBuildTool/WinRT/WinRTDeploy.cs
--- a/Engine/Source/Programs/UnrealBuildTool/WinRT/WinRTDeploy.cs
+++ b/Engine/Source/Programs/UnrealBuildTool/WinRT/WinRTDeploy.cs
@@ -106,7 +106,6 @@ namespace UnrealBuildTool
 					if (FilesInSource.Contains(DestFilename.ToUpperInvariant()) == false)
 					{
 						Log.TraceInformation("Destination file does not exist in Source - DELETING: {0}", DestFile);
-						FileAttributes attributes = File.GetAttributes(DestFile);
 						try
 						{
 							DeployHelper_DeleteFile(DestFile);
diff --git a/Engine/Source/Runtime/Core/Public/Containers/BitArray.h b/Engine/Source/Runtime/Core/Public/Containers/BitArray.h
--- a/Engine/Source/Runtime/Core/Public/Containers/BitArray.h
+++ b/Engine/Source/Runtime/Core/Public/Containers/BitArray.h
@@ -345,7 +345,7 @@ public:
 			// Copy bits from the end to the region we are removing
 			for( int32 Index=0;Index<NumBitsToRemove;Index++ )
 			{
-#if PLATFORM_MAC
+#if PLATFORM_MAC || PLATFORM_LINUX
 				// Clang compiler doesn't understand the short syntax, so let's be explicit
 				int32 FromIndex = NumBits - NumBitsToRemove + Index;
 				FConstBitReference From(GetData()[FromIndex / NumBitsPerDWORD],1 << (FromIndex & (NumBitsPerDWORD - 1)));
