UBT: Fixing crash when build exception message contains any of the C# formatting escape sequences

diff --git a/Engine/Source/Programs/UnrealBuildTool/System/CPPHeaders.cs b/Engine/Source/Programs/UnrealBuildTool/System/CPPHeaders.cs
--- a/Engine/Source/Programs/UnrealBuildTool/System/CPPHeaders.cs
+++ b/Engine/Source/Programs/UnrealBuildTool/System/CPPHeaders.cs
@@ -208,7 +208,7 @@ namespace UnrealBuildTool
 							BuildExceptionMessage += "Exception:\n";
 							BuildExceptionMessage += ex.ToString();
 							BuildExceptionMessage += "\n";
-							throw new BuildException(ex, BuildExceptionMessage);
+							throw new BuildException(ex, "{0}", BuildExceptionMessage);
 						}
 					}
 					else
