Support for modular builds (cross and native).

- Parts of PR #242 (with changes).

diff --git a/Engine/Source/Developer/DirectoryWatcher/Private/DirectoryWatcherPrivatePCH.h b/Engine/Source/Developer/DirectoryWatcher/Private/DirectoryWatcherPrivatePCH.h
--- a/Engine/Source/Developer/DirectoryWatcher/Private/DirectoryWatcherPrivatePCH.h
+++ b/Engine/Source/Developer/DirectoryWatcher/Private/DirectoryWatcherPrivatePCH.h
@@ -13,4 +13,9 @@ DECLARE_LOG_CATEGORY_EXTERN(LogDirectoryWatcher, Log, All);
 #elif PLATFORM_MAC
 #include "Mac/DirectoryWatchRequestMac.h"
 #include "Mac/DirectoryWatcherMac.h"
-#endif
\ No newline at end of file
+#elif PLATFORM_LINUX
+#include "Linux/DirectoryWatchRequestLinux.h"
+#include "Linux/DirectoryWatcherLinux.h"
+#else
+#error "Unknown platform"
+#endif
diff --git a/Engine/Source/Developer/DirectoryWatcher/Private/Linux/DirectoryWatchRequestLinux.cpp b/Engine/Source/Developer/DirectoryWatcher/Private/Linux/DirectoryWatchRequestLinux.cpp
new file mode 100644
--- /dev/null
+++ b/Engine/Source/Developer/DirectoryWatcher/Private/Linux/DirectoryWatchRequestLinux.cpp
@@ -0,0 +1,172 @@
+// Copyright 1998-2014 Epic Games, Inc. All Rights Reserved.
+
+#include "DirectoryWatcherPrivatePCH.h"
+
+FDirectoryWatchRequestLinux::FDirectoryWatchRequestLinux()
+:	bRunning(false)
+,	bEndWatchRequestInvoked(false)
+{
+    NotifyFilter = IN_CREATE | IN_MOVE | IN_MODIFY | IN_DELETE;
+}
+
+FDirectoryWatchRequestLinux::~FDirectoryWatchRequestLinux()
+{
+    Shutdown();
+}
+
+void FDirectoryWatchRequestLinux::Shutdown()
+{
+    free(WatchDescriptor);
+    close(FileDescriptor);
+}
+
+bool FDirectoryWatchRequestLinux::Init(const FString& InDirectory)
+{
+	if (InDirectory.Len() == 0)
+	{
+		// Verify input
+		return false;
+	}
+
+	if (bRunning)
+	{
+		Shutdown();
+	}
+
+	bEndWatchRequestInvoked = false;
+
+	// Make sure the path is absolute
+	const FString FullPath = FPaths::ConvertRelativePathToFull(InDirectory);
+
+    FileDescriptor = inotify_init1(IN_NONBLOCK | IN_CLOEXEC);
+
+    if (FileDescriptor == -1)
+    {
+		UE_LOG(LogDirectoryWatcher, Error, TEXT("Failed to init inotify"));
+    	return false;
+    }
+
+    IFileManager::Get().FindFilesRecursive(AllFiles, *FullPath, TEXT("*"), false, true);
+
+    WatchDescriptor = (int*)FMemory::Alloc((AllFiles.Num()+1) * sizeof(int));
+    if (WatchDescriptor == NULL) 
+    {
+		UE_LOG(LogDirectoryWatcher, Error, TEXT("Failed to allocate memory for WatchDescriptor"));
+    	return false;
+    }
+
+    for (int32 FileIdx = 0; FileIdx < AllFiles.Num(); ++FileIdx)
+    {
+		const FString& FolderName = AllFiles[FileIdx];
+		WatchDescriptor[FileIdx] = inotify_add_watch(FileDescriptor, TCHAR_TO_UTF8(*FolderName), NotifyFilter); //FileIdx+1
+		if (WatchDescriptor[FileIdx] == -1) 
+		{
+			UE_LOG(LogDirectoryWatcher, Error, TEXT("inotify_add_watch cannot watch folder %s"), *FolderName);
+			return false;
+		}
+    }
+
+  	bRunning = true;
+
+  	return true;
+}
+
+void FDirectoryWatchRequestLinux::AddDelegate( const IDirectoryWatcher::FDirectoryChanged& InDelegate )
+{
+	Delegates.Add(InDelegate);
+}
+
+bool FDirectoryWatchRequestLinux::RemoveDelegate( const IDirectoryWatcher::FDirectoryChanged& InDelegate )
+{
+	if ( Delegates.Contains(InDelegate) )
+	{
+		Delegates.Remove(InDelegate);
+		return true;
+	}
+
+	return false;
+}
+
+bool FDirectoryWatchRequestLinux::HasDelegates() const
+{
+	return Delegates.Num() > 0;
+}
+
+void FDirectoryWatchRequestLinux::EndWatchRequest()
+{
+	bEndWatchRequestInvoked = true;
+}
+
+
+void FDirectoryWatchRequestLinux::ProcessPendingNotifications()
+{
+    ProcessChanges();
+
+	// Trigger all listening delegates with the files that have changed
+	if ( FileChanges.Num() > 0 )
+	{
+		for (int32 DelegateIdx = 0; DelegateIdx < Delegates.Num(); ++DelegateIdx)
+		{
+			Delegates[DelegateIdx].Execute(FileChanges);
+		}
+
+		FileChanges.Empty();
+	}
+}
+
+void FDirectoryWatchRequestLinux::ProcessChanges()
+{
+    uint8 Buf[BUF_LEN] __attribute__ ((aligned(__alignof__(inotify_event))));
+
+	for (;;)
+	{
+		ssize_t Len = read(FileDescriptor, Buf, sizeof(Buf));
+		if (Len == -1 && errno != EAGAIN)
+		{
+			// TODO: handle errors
+			return;
+		}
+
+		if (Len <= 0)
+		{
+			break;
+		}
+		else
+		{
+			for (const uint8 * Ptr = Buf; Ptr < Buf + Len;)
+			{
+				const inotify_event* Event = reinterpret_cast<const inotify_event *>(Ptr);
+				FFileChangeData::EFileChangeAction Action = FFileChangeData::FCA_Unknown;
+
+				// not sure if "else" is correct here, but we cannot OR actions that way
+				if ((Event->mask & IN_CREATE) || (Event->mask & IN_MOVED_TO))
+				{
+					Action = FFileChangeData::FCA_Added;
+				}
+				else if (Event->mask & IN_MODIFY)
+				{
+					Action = FFileChangeData::FCA_Modified;
+				}
+				else if ((Event->mask & IN_DELETE) || (Event->mask & IN_MOVED_FROM))
+				{
+					Action = FFileChangeData::FCA_Removed;
+				}
+
+				Ptr += sizeof(inotify_event)+Event->len;
+
+				// find the file
+				FString Filename;
+				for (int32 i = 0; i < AllFiles.Num(); ++i)
+				{
+					if (WatchDescriptor[i] == Event->wd)
+					{
+						Filename = AllFiles[i];
+						break;
+					}
+				}
+
+				new (FileChanges)FFileChangeData(Filename, Action);
+			}
+		}
+	}
+}
diff --git a/Engine/Source/Developer/DirectoryWatcher/Private/Linux/DirectoryWatchRequestLinux.h b/Engine/Source/Developer/DirectoryWatcher/Private/Linux/DirectoryWatchRequestLinux.h
new file mode 100644
--- /dev/null
+++ b/Engine/Source/Developer/DirectoryWatcher/Private/Linux/DirectoryWatchRequestLinux.h
@@ -0,0 +1,48 @@
+// Copyright 1998-2014 Epic Games, Inc. All Rights Reserved.
+
+#pragma once
+
+#include <errno.h>
+#include <sys/stat.h>
+#include <sys/types.h>
+#include <sys/inotify.h>
+
+#define EVENT_SIZE  ( sizeof (struct inotify_event) )
+#define BUF_LEN     ( 1024 * ( EVENT_SIZE + 16 ) )
+
+class FDirectoryWatchRequestLinux
+{
+public:
+	FDirectoryWatchRequestLinux();
+	virtual ~FDirectoryWatchRequestLinux();
+
+	/** Sets up the directory handle and request information */
+	bool Init(const FString& InDirectory);
+
+	/** Adds a delegate to get fired when the directory changes */
+	void AddDelegate( const IDirectoryWatcher::FDirectoryChanged& InDelegate );
+	/** Removes a delegate to get fired when the directory changes */
+	bool RemoveDelegate( const IDirectoryWatcher::FDirectoryChanged& InDelegate );
+	/** Returns true if this request has any delegates listening to directory changes */
+	bool HasDelegates() const;
+	/** Prepares the request for deletion */
+	void EndWatchRequest();
+	/** Triggers all pending file change notifications */
+	void ProcessPendingNotifications();
+
+private:
+
+    bool bRunning;
+	bool bEndWatchRequestInvoked;
+
+    int FileDescriptor;
+    int * WatchDescriptor;
+    int NotifyFilter;
+
+    TArray<FString> AllFiles;
+	TArray<IDirectoryWatcher::FDirectoryChanged> Delegates;
+	TArray<FFileChangeData> FileChanges;
+
+    void Shutdown();
+    void ProcessChanges();
+};
diff --git a/Engine/Source/Developer/DirectoryWatcher/Private/Linux/DirectoryWatcherLinux.cpp b/Engine/Source/Developer/DirectoryWatcher/Private/Linux/DirectoryWatcherLinux.cpp
new file mode 100644
--- /dev/null
+++ b/Engine/Source/Developer/DirectoryWatcher/Private/Linux/DirectoryWatcherLinux.cpp
@@ -0,0 +1,120 @@
+// Copyright 1998-2014 Epic Games, Inc. All Rights Reserved.
+
+#include "DirectoryWatcherPrivatePCH.h"
+
+FDirectoryWatcherLinux::FDirectoryWatcherLinux()
+{
+	NumRequests = 0;
+}
+
+FDirectoryWatcherLinux::~FDirectoryWatcherLinux()
+{
+	if (RequestMap.Num() != 0)
+	{
+		// Delete any remaining requests here. These requests are likely from modules which are still loaded at the time that this module unloads.
+		for (TMap<FString, FDirectoryWatchRequestLinux*>::TConstIterator RequestIt(RequestMap); RequestIt; ++RequestIt)
+		{
+			if (ensure(RequestIt.Value()))
+			{
+				delete RequestIt.Value();
+				NumRequests--;
+			}
+		}
+
+		RequestMap.Empty();
+	}
+
+	if (RequestsPendingDelete.Num() != 0)
+	{
+		for (int32 RequestIdx = 0; RequestIdx < RequestsPendingDelete.Num(); ++RequestIdx)
+		{
+			delete RequestsPendingDelete[RequestIdx];
+			NumRequests--;
+		}
+	}
+
+	// Make sure every request that was created is destroyed
+	ensure(NumRequests == 0);
+}
+
+bool FDirectoryWatcherLinux::RegisterDirectoryChangedCallback( const FString& Directory, const FDirectoryChanged& InDelegate )
+{
+	FDirectoryWatchRequestLinux** RequestPtr = RequestMap.Find(Directory);
+	FDirectoryWatchRequestLinux* Request = NULL;
+
+	if (RequestPtr)
+	{
+		// There should be no NULL entries in the map
+		check (*RequestPtr);
+
+		Request = *RequestPtr;
+	}
+	else
+	{
+		Request = new FDirectoryWatchRequestLinux();
+		NumRequests++;
+
+		// Begin reading directory changes
+		if (!Request->Init(Directory))
+		{
+			UE_LOG(LogDirectoryWatcher, Warning, TEXT("Failed to begin reading directory changes for %s."), *Directory);
+			delete Request;
+			NumRequests--;
+			return false;
+		}
+
+		RequestMap.Add(Directory, Request);
+	}
+
+	Request->AddDelegate(InDelegate);
+
+	return true;
+}
+
+bool FDirectoryWatcherLinux::UnregisterDirectoryChangedCallback(const FString& Directory, const FDirectoryChanged& InDelegate)
+{
+	FDirectoryWatchRequestLinux** RequestPtr = RequestMap.Find(Directory);
+
+	if (RequestPtr)
+	{
+		// There should be no NULL entries in the map
+		check (*RequestPtr);
+
+		FDirectoryWatchRequestLinux* Request = *RequestPtr;
+
+		if (Request->RemoveDelegate(InDelegate))
+		{
+			if (!Request->HasDelegates())
+			{
+				// Remove from the active map and add to the pending delete list
+				RequestMap.Remove(Directory);
+				RequestsPendingDelete.AddUnique(Request);
+
+				// Signal to end the watch which will mark this request for deletion
+				Request->EndWatchRequest();
+			}
+
+			return true;
+		}
+
+	}
+	return false;
+}
+
+void FDirectoryWatcherLinux::Tick(float DeltaSeconds)
+{
+	// Delete unregistered requests
+	for (int32 RequestIdx = RequestsPendingDelete.Num() - 1; RequestIdx >= 0; --RequestIdx)
+	{
+		FDirectoryWatchRequestLinux* Request = RequestsPendingDelete[RequestIdx];
+		delete Request;
+		NumRequests--;
+		RequestsPendingDelete.RemoveAt(RequestIdx);
+	}
+
+	// Trigger any file change notification delegates
+	for (TMap<FString, FDirectoryWatchRequestLinux*>::TConstIterator RequestIt(RequestMap); RequestIt; ++RequestIt)
+	{
+		RequestIt.Value()->ProcessPendingNotifications();
+	}
+}
diff --git a/Engine/Source/Developer/DirectoryWatcher/Private/Linux/DirectoryWatcherLinux.h b/Engine/Source/Developer/DirectoryWatcher/Private/Linux/DirectoryWatcherLinux.h
new file mode 100644
--- /dev/null
+++ b/Engine/Source/Developer/DirectoryWatcher/Private/Linux/DirectoryWatcherLinux.h
@@ -0,0 +1,23 @@
+// Copyright 1998-2014 Epic Games, Inc. All Rights Reserved.
+
+#pragma once
+
+class FDirectoryWatcherLinux : public IDirectoryWatcher
+{
+public:
+	FDirectoryWatcherLinux();
+	virtual ~FDirectoryWatcherLinux();
+
+	virtual bool RegisterDirectoryChangedCallback(const FString& Directory, const FDirectoryChanged& InDelegate) override;
+	virtual bool UnregisterDirectoryChangedCallback(const FString& Directory, const FDirectoryChanged& InDelegate) override;
+	virtual void Tick (float DeltaSeconds) override;
+
+	/** Map of directory paths to requests */
+	TMap<FString, FDirectoryWatchRequestLinux*> RequestMap;
+	TArray<FDirectoryWatchRequestLinux*> RequestsPendingDelete;
+
+	/** A count of FDirectoryWatchRequestLinux created to ensure they are cleaned up on shutdown */
+	int32 NumRequests;
+};
+
+typedef FDirectoryWatcherLinux FDirectoryWatcher;
diff --git a/Engine/Source/Developer/ShaderFormatOpenGL/Private/OpenGLShaderCompiler.cpp b/Engine/Source/Developer/ShaderFormatOpenGL/Private/OpenGLShaderCompiler.cpp
--- a/Engine/Source/Developer/ShaderFormatOpenGL/Private/OpenGLShaderCompiler.cpp
+++ b/Engine/Source/Developer/ShaderFormatOpenGL/Private/OpenGLShaderCompiler.cpp
@@ -26,7 +26,7 @@
 #include "HideWindowsPlatformTypes.h"
 #elif PLATFORM_LINUX
 	#define GL_GLEXT_PROTOTYPES
-	#include <GL/gl.h>
+	#include <GL/glcorearb.h>
 	#include <GL/glext.h>
 #elif PLATFORM_MAC
 	#include <OpenGL/OpenGL.h>
diff --git a/Engine/Source/Editor/UnrealEd/Private/UnrealEdMisc.cpp b/Engine/Source/Editor/UnrealEd/Private/UnrealEdMisc.cpp
--- a/Engine/Source/Editor/UnrealEd/Private/UnrealEdMisc.cpp
+++ b/Engine/Source/Editor/UnrealEd/Private/UnrealEdMisc.cpp
@@ -483,6 +483,12 @@ FString CreateProjectPath()
 	{
 		return UTF8_TO_TCHAR([[[NSBundle mainBundle] executablePath] fileSystemRepresentation]);
 	}
+#elif PLATFORM_LINUX
+	const TCHAR* PlatformConfig = FPlatformMisc::GetUBTPlatform();
+	FString ExeFileName = FPaths::EngineDir() / TEXT("Binaries") / PlatformConfig / FString(FPlatformProcess::ExecutableName());
+	return ExeFileName;
+#else
+#error "Unknown platform"
 #endif
 }
 
diff --git a/Engine/Source/Programs/UnrealBuildTool/Configuration/UEBuildTarget.cs b/Engine/Source/Programs/UnrealBuildTool/Configuration/UEBuildTarget.cs
--- a/Engine/Source/Programs/UnrealBuildTool/Configuration/UEBuildTarget.cs
+++ b/Engine/Source/Programs/UnrealBuildTool/Configuration/UEBuildTarget.cs
@@ -1069,7 +1069,7 @@ namespace UnrealBuildTool
 						string UnixLibraryFileName = Path.Combine(LibraryPath, "lib" + LibraryName + LibraryExtension);
 						if(File.Exists(UnixLibraryFileName))
 						{
-							FileNames.Add(UnixLibraryFileName);
+                            FileNames.Add(UnixLibraryFileName);
 						}
 					}
 				}
@@ -2031,14 +2031,19 @@ namespace UnrealBuildTool
 
 			string OutBinaryPath = "";
 			// Append binary file name
+			string Prefix = "";
+			if (Platform == UnrealTargetPlatform.Linux && (BinaryType == UEBuildBinaryType.DynamicLinkLibrary || BinaryType == UEBuildBinaryType.StaticLibrary))
+			{
+				Prefix = "lib";
+			}
 			if (LocalConfig == UnrealTargetConfiguration.Development || bForceNameAsForDevelopment)
 			{
-				OutBinaryPath = Path.Combine(BaseDirectory, String.Format("{0}{1}{2}", BinaryName, BinarySuffix, BinaryExtension));
+				OutBinaryPath = Path.Combine(BaseDirectory, String.Format("{3}{0}{1}{2}", BinaryName, BinarySuffix, BinaryExtension, Prefix));
 			}
 			else
 			{
-				OutBinaryPath = Path.Combine(BaseDirectory, String.Format("{0}-{1}-{2}{3}{4}", 
-					BinaryName, Platform.ToString(), LocalConfig.ToString(), BinarySuffix, BinaryExtension));
+				OutBinaryPath = Path.Combine(BaseDirectory, String.Format("{5}{0}-{1}-{2}{3}{4}", 
+					BinaryName, Platform.ToString(), LocalConfig.ToString(), BinarySuffix, BinaryExtension, Prefix));
 			}
 			return OutBinaryPath;
 		}
diff --git a/Engine/Source/Programs/UnrealBuildTool/Linux/LinuxProjectGenerator.cs b/Engine/Source/Programs/UnrealBuildTool/Linux/LinuxProjectGenerator.cs
--- a/Engine/Source/Programs/UnrealBuildTool/Linux/LinuxProjectGenerator.cs
+++ b/Engine/Source/Programs/UnrealBuildTool/Linux/LinuxProjectGenerator.cs
@@ -6,7 +6,7 @@ using System.Text;
 using System.Diagnostics;
 using System.IO;
 
-namespace UnrealBuildTool.Linux
+namespace UnrealBuildTool
 {
     /**
      *	Base class for platform-specific project generators 
diff --git a/Engine/Source/Programs/UnrealBuildTool/Linux/LinuxToolChain.cs b/Engine/Source/Programs/UnrealBuildTool/Linux/LinuxToolChain.cs
--- a/Engine/Source/Programs/UnrealBuildTool/Linux/LinuxToolChain.cs
+++ b/Engine/Source/Programs/UnrealBuildTool/Linux/LinuxToolChain.cs
@@ -7,7 +7,7 @@ using System.Diagnostics;
 using System.IO;
 using Microsoft.Win32;
 
-namespace UnrealBuildTool.Linux
+namespace UnrealBuildTool
 {
     class LinuxToolChain : UEToolChain
     {
@@ -35,7 +35,7 @@ namespace UnrealBuildTool.Linux
             proc.WaitForExit();
 
             string path = proc.StandardOutput.ReadLine();
-            Console.WriteLine(String.Format("which {0} result: ({1}) {2}", name, proc.ExitCode, path));
+            Log.TraceVerbose(String.Format("which {0} result: ({1}) {2}", name, proc.ExitCode, path));
 
             if (proc.ExitCode == 0 && String.IsNullOrEmpty(proc.StandardError.ReadToEnd()))
             {
@@ -96,6 +96,8 @@ namespace UnrealBuildTool.Linux
             }
 
             Result += " -Wall -Werror";
+            // test without this next line?
+            Result += " -funwind-tables";               // generate unwind tables as they seem to be needed for stack tracing (why??)
             Result += " -Wsequence-point";              // additional warning not normally included in Wall: warns if order of operations is ambigious
             //Result += " -Wunreachable-code";            // additional warning not normally included in Wall: warns if there is code that will never be executed - not helpful due to bIsGCC and similar 
             //Result += " -Wshadow";                      // additional warning not normally included in Wall: warns if there variable/typedef shadows some other variable - not helpful because we have gobs of code that shadows variables
@@ -103,6 +105,14 @@ namespace UnrealBuildTool.Linux
             Result += " -fno-math-errno";               // do not assume that math ops have side effects
             Result += " -fno-rtti";                     // no run-time type info
 
+            // these needs to be removed ASAP
+            Result += " -Wno-delete-non-virtual-dtor";
+            Result += " -Wno-reorder";
+            Result += " -Wno-logical-op-parentheses";
+            Result += " -Wno-ignored-attributes";
+            Result += " -Wno-overloaded-virtual";
+            Result += " -Wno-unused-value";
+
             if (String.IsNullOrEmpty(ClangPath))
             {
                 // GCC only option
@@ -159,7 +169,7 @@ namespace UnrealBuildTool.Linux
             {
                 Result += " -g3";
                 Result += " -fno-omit-frame-pointer";
-                Result += " -funwind-tables";               // generate unwind tables as they seem to be needed for stack tracing (why??)
+                Result += " -funwind-tables";               // generate unwind tables as they seem to be needed for stack tracing
                 Result += " -fstack-protector";
                 //Result += " -fsanitize=address";  // Preferred clang tool for detecting address based errors but unusable for some reason with Module.Engine.7_of_42.cpp
             }
@@ -257,24 +267,19 @@ namespace UnrealBuildTool.Linux
             {
                 Result += " -s"; // Strip binaries in Shipping
             }
-
             if (LinkEnvironment.Config.bIsBuildingDLL)
             {
                 Result += " -shared";
             }
-            //Result += " -v";
-
-            if (CrossCompiling())
-            {
-                // ignore unresolved symbols in shared libs
-                Result += " -Wl,--unresolved-symbols=ignore-in-shared-libs";
-            }
             else
             {
-                Result += " -Wl,--no-undefined";
+                // ignore unresolved symbols in shared libs
+                Result += string.Format(" -Wl,--unresolved-symbols=ignore-in-shared-libs");
             }
 
             // RPATH for third party libs
+            Result += " -Wl,-rpath=${{ORIGIN}}";
+            Result += " -Wl,-rpath-link=${{ORIGIN}}";
             Result += " -Wl,-rpath=${ORIGIN}/../../../Engine/Binaries/Linux";
             // FIXME: really ugly temp solution. Modules need to be able to specify this
             Result += " -Wl,-rpath=${ORIGIN}/../../../Engine/Binaries/ThirdParty/jemalloc/Linux/x86_64-unknown-linux-gnu";
@@ -348,11 +353,21 @@ namespace UnrealBuildTool.Linux
         static string ArPath;
         static string RanlibPath;
 
+		/** Track which scripts need to be deleted before appending to */
+		private bool bHasWipedFixDepsScript = false;
+
+		private static List<FileItem> BundleDependencies = new List<FileItem>();
+
         public override CPPOutput CompileCPPFiles(CPPEnvironment CompileEnvironment, List<FileItem> SourceFiles, string ModuleName)
         {
             string Arguments = GetCLArguments_Global(CompileEnvironment);
             string PCHArguments = "";
 
+            if (CompileEnvironment.Config.bIsBuildingDLL)
+            {
+                Arguments += " -fPIC";
+            }
+
             if (CompileEnvironment.Config.PrecompiledHeaderAction == PrecompiledHeaderAction.Include)
             {
                 // Add the precompiled header file's path to the include path so Clang can find it.
@@ -549,8 +564,69 @@ namespace UnrealBuildTool.Linux
             return OutputFile;
         }
 
+		public FileItem FixDependencies(LinkEnvironment LinkEnvironment, FileItem Executable)
+		{
+			if (!LinkEnvironment.Config.bIsCrossReferenced)
+			{
+				return null;
+			}
+
+			Log.TraceVerbose("Adding postlink step");
+
+            bool bUseCmdExe = ExternalExecution.GetRuntimePlatform() == UnrealTargetPlatform.Win64 || ExternalExecution.GetRuntimePlatform() == UnrealTargetPlatform.Win32;
+            string ShellBinary = bUseCmdExe ? "cmd.exe" : "/bin/sh";
+            string ExecuteSwitch = bUseCmdExe ? " /C" : " -c";
+            string ScriptName = bUseCmdExe ? "FixDependencies.bat" : "FixDependencies.sh";
+
+            FileItem FixDepsScript = FileItem.GetItemByFullPath(Path.Combine(LinkEnvironment.Config.LocalShadowDirectory, ScriptName));
+
+			Action PostLinkAction = new Action(ActionType.Link);
+			PostLinkAction.WorkingDirectory = Path.GetFullPath(".");
+            PostLinkAction.CommandPath = ShellBinary;
+			PostLinkAction.StatusDescription = string.Format("{0}", Path.GetFileName(Executable.AbsolutePath));
+			PostLinkAction.CommandDescription = "FixDeps";
+			PostLinkAction.bCanExecuteRemotely = false;
+            PostLinkAction.CommandArguments = ExecuteSwitch;
+            
+            if (bUseCmdExe)
+            {
+                PostLinkAction.CommandArguments += " \"";
+            }
+            else
+            {
+                PostLinkAction.CommandArguments += String.Format(" 'chmod +x {0}' && {1} {2} '", FixDepsScript.AbsolutePath, ShellBinary, ExecuteSwitch);
+            }
+
+			FileItem OutputFile = FileItem.GetItemByPath(Path.Combine(LinkEnvironment.Config.LocalShadowDirectory, Path.GetFileNameWithoutExtension(Executable.AbsolutePath) + ".link"));
+
+			// Make sure we don't run this script until the all executables and shared libraries
+			// have been built.
+			PostLinkAction.PrerequisiteItems.Add(Executable);
+			foreach (FileItem Dependency in BundleDependencies)
+			{
+				PostLinkAction.PrerequisiteItems.Add(Dependency);
+			}
+
+			PostLinkAction.CommandArguments += ShellBinary + ExecuteSwitch + " \"" + FixDepsScript.AbsolutePath + "\" && ";
+
+            string Touch = bUseCmdExe ? "echo \"\" >> \"{0}\" && copy /b \"{0}\" +,," : "echo \"After script\" && touch \"{0}\"";
+
+            PostLinkAction.CommandArguments += String.Format(Touch, OutputFile.AbsolutePath);
+			PostLinkAction.CommandArguments += bUseCmdExe ? "\"" : "'";
+
+			System.Console.WriteLine("{0} {1}", PostLinkAction.CommandPath, PostLinkAction.CommandArguments);
+			
+			PostLinkAction.ProducedItems.Add(OutputFile);
+			return OutputFile;
+		}
+
+
         public override FileItem LinkFiles(LinkEnvironment LinkEnvironment, bool bBuildImportLibraryOnly)
         {
+            Debug.Assert(!bBuildImportLibraryOnly);
+
+            List<string> RPaths = new List<string>();
+
             if (LinkEnvironment.Config.bIsBuildingLibrary || bBuildImportLibraryOnly)
             {
                 return CreateArchiveAndIndex(LinkEnvironment);
@@ -595,27 +671,97 @@ namespace UnrealBuildTool.Linux
             string ResponseFileName = GetResponseFileName(LinkEnvironment, OutputFile);
             LinkAction.CommandArguments += string.Format(" -Wl,@\"{0}\"", ResponseFile.Create(ResponseFileName, InputFileNames));
 
+            if (LinkEnvironment.Config.bIsBuildingDLL)
+            {
+	            LinkAction.CommandArguments += string.Format(" -Wl,-soname={0}", OutputFile);
+            }
+
+            // Start with the configured LibraryPaths and also add paths to any libraries that
+            // we depend on (libraries that we've build ourselves).
+            List<string> AllLibraryPaths = LinkEnvironment.Config.LibraryPaths;
+            foreach (string AdditionalLibrary in LinkEnvironment.Config.AdditionalLibraries)
+            {
+                string PathToLib = Path.GetDirectoryName(AdditionalLibrary);
+                if (!String.IsNullOrEmpty(PathToLib)) 
+                {
+                    // make path absolute, because FixDependencies script may be executed in a different directory
+                    string AbsolutePathToLib = Path.GetFullPath(PathToLib);
+                    if (!AllLibraryPaths.Contains(AbsolutePathToLib))
+                    {
+                        AllLibraryPaths.Add(AbsolutePathToLib);
+                    }
+                }
+
+                if ((AdditionalLibrary.Contains("Plugins") || AdditionalLibrary.Contains("Binaries/ThirdParty") || AdditionalLibrary.Contains("Binaries\\ThirdParty")) && Path.GetDirectoryName(AdditionalLibrary) != Path.GetDirectoryName(OutputFile.AbsolutePath))
+                {
+                    string RelativePath = Utils.MakePathRelativeTo(Path.GetDirectoryName(AdditionalLibrary), Path.GetDirectoryName(OutputFile.AbsolutePath));
+                    if (!RPaths.Contains(RelativePath))
+                    {
+                        RPaths.Add(RelativePath);
+                        LinkAction.CommandArguments += string.Format(" -Wl,-rpath=\"${{ORIGIN}}/{0}\"", RelativePath);
+                    }
+                }
+            }
+
+            LinkAction.CommandArguments += string.Format(" -Wl,-rpath-link=\"{0}\"", Path.GetDirectoryName(OutputFile.AbsolutePath));
+
 			// Add the library paths to the argument list.
-			foreach (string LibraryPath in LinkEnvironment.Config.LibraryPaths)
+			foreach (string LibraryPath in AllLibraryPaths)
 			{
-				LinkAction.CommandArguments += string.Format(" -Wl,-L\"{0}\"", LibraryPath);
+                // use absolute paths because of FixDependencies script again
+				LinkAction.CommandArguments += string.Format(" -L\"{0}\"", Path.GetFullPath(LibraryPath));
 			}
 
 			// add libraries in a library group
             LinkAction.CommandArguments += string.Format(" -Wl,--start-group");
+
+            List<string> EngineAndGameLibraries = new List<string>();
             foreach (string AdditionalLibrary in LinkEnvironment.Config.AdditionalLibraries)
-			{
-				if (String.IsNullOrEmpty(Path.GetDirectoryName(AdditionalLibrary)))
-				{
-					LinkAction.CommandArguments += string.Format(" -l{0}", AdditionalLibrary);
-				}
-				else
-				{
-					// full pathed libs are compiled by us, so we depend on linking them
-					LinkAction.CommandArguments += string.Format(" \"{0}\"", AdditionalLibrary);
-					LinkAction.PrerequisiteItems.Add(FileItem.GetItemByFullPath(AdditionalLibrary));
-				}
-			}
+            {
+                if (String.IsNullOrEmpty(Path.GetDirectoryName(AdditionalLibrary)))
+                {
+                    // library was passed just like "jemalloc", turn it into -ljemalloc
+                    LinkAction.CommandArguments += string.Format(" -l{0}", AdditionalLibrary);
+                }
+                else if (Path.GetExtension(AdditionalLibrary) == ".a")
+                {
+                    // static library passed in, pass it along but make path absolute, because FixDependencies script may be executed in a different directory
+                    string AbsoluteAdditionalLibrary = Path.GetFullPath(AdditionalLibrary);
+                    LinkAction.CommandArguments += (" " + AbsoluteAdditionalLibrary);
+                    LinkAction.PrerequisiteItems.Add(FileItem.GetItemByPath(AdditionalLibrary));
+                }
+                else
+                {
+                    // Skip over full-pathed library dependencies when building DLLs to avoid circular
+                    // dependencies.
+                    FileItem LibraryDependency = FileItem.GetItemByPath(AdditionalLibrary);
+
+                    var LibName = Path.GetFileNameWithoutExtension(AdditionalLibrary);
+                    if (LibName.StartsWith("lib"))
+                    {
+                        // Remove lib prefix
+                        LibName = LibName.Remove(0, 3);
+                    }
+                    string LibLinkFlag = string.Format(" -l{0}", LibName);
+
+                    if (LinkEnvironment.Config.bIsBuildingDLL && LinkEnvironment.Config.bIsCrossReferenced)
+                    {
+                        // We are building a cross referenced DLL so we can't actually include
+                        // dependencies at this point. Instead we add it to the list of
+                        // libraries to be used in the FixDependencies step.
+                        EngineAndGameLibraries.Add(LibLinkFlag);
+                        if (!LinkAction.CommandArguments.Contains("--allow-shlib-undefined"))
+                        {
+                            LinkAction.CommandArguments += string.Format(" -Wl,--allow-shlib-undefined");
+                        }
+                    }
+                    else
+                    {
+                        LinkAction.PrerequisiteItems.Add(LibraryDependency);
+                        LinkAction.CommandArguments += LibLinkFlag;
+                    }
+                }
+            }
             LinkAction.CommandArguments += " -lrt"; // needed for clock_gettime()
             LinkAction.CommandArguments += " -lm"; // math
             LinkAction.CommandArguments += string.Format(" -Wl,--end-group");
@@ -651,6 +797,69 @@ namespace UnrealBuildTool.Linux
             // Only execute linking on the local PC.
             LinkAction.bCanExecuteRemotely = false;
 
+			// Prepare a script that will run later, once all shared libraries and the executable
+			// are created. This script will be called by action created in FixDependencies()
+			if (LinkEnvironment.Config.bIsCrossReferenced && LinkEnvironment.Config.bIsBuildingDLL)
+			{
+                bool bUseCmdExe = ExternalExecution.GetRuntimePlatform() == UnrealTargetPlatform.Win64 || ExternalExecution.GetRuntimePlatform() == UnrealTargetPlatform.Win32;
+                string ScriptName = bUseCmdExe ? "FixDependencies.bat" : "FixDependencies.sh";
+
+				string FixDepsScriptPath = Path.Combine(LinkEnvironment.Config.LocalShadowDirectory, ScriptName);
+				if (!bHasWipedFixDepsScript)
+				{
+					bHasWipedFixDepsScript = true;
+					Log.TraceVerbose("Creating script: {0}", FixDepsScriptPath);
+					StreamWriter Writer = File.CreateText(FixDepsScriptPath);
+
+                    if (bUseCmdExe)
+                    {
+					    Writer.Write("rem Automatically generated by UnrealBuildTool\n");
+					    Writer.Write("rem *DO NOT EDIT*\n\n");
+                    }
+                    else
+                    {
+					    Writer.Write("#!/bin/sh\n");
+					    Writer.Write("# Automatically generated by UnrealBuildTool\n");
+					    Writer.Write("# *DO NOT EDIT*\n\n");
+					    Writer.Write("set -o errexit\n");
+                    }
+				    Writer.Close();
+				}
+
+				StreamWriter FixDepsScript = File.AppendText(FixDepsScriptPath);
+
+				string EngineAndGameLibrariesString = "";
+				foreach (string Library in EngineAndGameLibraries)
+				{
+					EngineAndGameLibrariesString += Library;
+				}
+
+				FixDepsScript.Write(string.Format("echo Fixing {0}\n", Path.GetFileName(OutputFile.AbsolutePath)));
+                if (!bUseCmdExe)
+                {
+				    FixDepsScript.Write(string.Format("TIMESTAMP=`stat --format %y \"{0}\"`\n", OutputFile.AbsolutePath));
+                }
+				string FixDepsLine = LinkAction.CommandPath + " " + LinkAction.CommandArguments;
+				string Replace = "-Wl,--allow-shlib-undefined";
+					
+				FixDepsLine = FixDepsLine.Replace(Replace, EngineAndGameLibrariesString);
+				FixDepsLine = FixDepsLine.Replace(OutputFile.AbsolutePath, OutputFile.AbsolutePath + ".fixed");
+				FixDepsLine = FixDepsLine.Replace("$", "\\$");
+				FixDepsScript.Write(FixDepsLine + "\n");
+                if (bUseCmdExe)
+                {
+                    FixDepsScript.Write(string.Format("move /Y \"{0}.fixed\" \"{0}\"\n", OutputFile.AbsolutePath));
+                }
+                else
+                {
+                    FixDepsScript.Write(string.Format("mv \"{0}.fixed\" \"{0}\"\n", OutputFile.AbsolutePath));
+                    FixDepsScript.Write(string.Format("touch -d \"$TIMESTAMP\" \"{0}\"\n\n", OutputFile.AbsolutePath));
+                }
+				FixDepsScript.Close();
+			}
+
+            //LinkAction.CommandArguments += " -v";
+
             return OutputFile;
         }
 
@@ -659,6 +868,14 @@ namespace UnrealBuildTool.Linux
             throw new BuildException("Linux cannot compile C# files");
         }
 
+        static public void SetupBundleDependencies(List<UEBuildBinary> Binaries, string GameName)
+        {
+            foreach (UEBuildBinary Binary in Binaries)
+            {
+                BundleDependencies.Add(FileItem.GetItemByPath(Binary.ToString()));
+            }
+        }
+
         /** Converts the passed in path from UBT host to compiler native format. */
         public override String ConvertPath(String OriginalPath)
         {
@@ -671,5 +888,23 @@ namespace UnrealBuildTool.Linux
                 return OriginalPath;
             }
         }
+
+        public override ICollection<FileItem> PostBuild(FileItem Executable, LinkEnvironment BinaryLinkEnvironment)
+        {
+            var OutputFiles = base.PostBuild(Executable, BinaryLinkEnvironment);
+
+            if (BinaryLinkEnvironment.Config.bIsBuildingDLL || BinaryLinkEnvironment.Config.bIsBuildingLibrary)
+            {
+                return OutputFiles;
+            }
+
+            FileItem FixDepsOutputFile = FixDependencies(BinaryLinkEnvironment, Executable);
+            if (FixDepsOutputFile != null)
+            {
+                OutputFiles.Add(FixDepsOutputFile);
+            }
+
+            return OutputFiles;
+        }
     }
 }
diff --git a/Engine/Source/Programs/UnrealBuildTool/Linux/UEBuildLinux.cs b/Engine/Source/Programs/UnrealBuildTool/Linux/UEBuildLinux.cs
--- a/Engine/Source/Programs/UnrealBuildTool/Linux/UEBuildLinux.cs
+++ b/Engine/Source/Programs/UnrealBuildTool/Linux/UEBuildLinux.cs
@@ -6,7 +6,7 @@ using System.Text;
 using System.Diagnostics;
 using System.IO;
 
-namespace UnrealBuildTool.Linux
+namespace UnrealBuildTool
 {
     class LinuxPlatform : UEBuildPlatform
     {
@@ -120,9 +120,7 @@ namespace UnrealBuildTool.Linux
                 string EngineSourcePath = Path.Combine(ProjectFileGenerator.RootRelativePath, "Engine", "Source");
                 string LinuxTargetPlatformFile = Path.Combine(EngineSourcePath, "Developer", "Linux", "LinuxTargetPlatform", "LinuxTargetPlatform.Build.cs");
 
-                if (
-                    (File.Exists(LinuxTargetPlatformFile) == false)
-                    )
+                if (File.Exists(LinuxTargetPlatformFile) == false)
                 {
                     bRegisterBuildPlatform = false;
                 }
@@ -216,8 +214,7 @@ namespace UnrealBuildTool.Linux
          */
         public override void ValidateBuildConfiguration(CPPTargetConfiguration Configuration, CPPTargetPlatform Platform, bool bCreateDebugInfo)
         {
-            // increase Unity size to avoid too long command lines
-            BuildConfiguration.NumIncludedBytesPerUnityCPP = 1024 * 1024;
+            UEBuildConfiguration.bCompileSimplygon = false;
             UEBuildConfiguration.bCompileICU = true;
         }
 
@@ -228,6 +225,13 @@ namespace UnrealBuildTool.Linux
          */
         public override void ValidateUEBuildConfiguration()
         {
+            if (ProjectFileGenerator.bGenerateProjectFiles && !ProjectFileGenerator.bGeneratingRocketProjectFiles)
+            {
+                // When generating non-Rocket project files we need intellisense generator to include info from all modules,
+                // including editor-only third party libs
+                UEBuildConfiguration.bCompileLeanAndMeanUE = false;
+            }
+
             BuildConfiguration.bUseUnityBuild = true;
 
             // Don't stop compilation at first error...
@@ -295,6 +299,32 @@ namespace UnrealBuildTool.Linux
                     InModule.AddPlatformSpecificDynamicallyLoadedModule("LinuxServerTargetPlatform");
                 }
             }
+            else if (Target.Platform == UnrealTargetPlatform.Linux)
+            {
+                bool bBuildShaderFormats = UEBuildConfiguration.bForceBuildShaderFormats;
+
+                if (!UEBuildConfiguration.bBuildRequiresCookedData)
+                {
+                    if (InModule.ToString() == "TargetPlatform")
+                    {
+                        bBuildShaderFormats = true;
+                    }
+                }
+
+                // allow standalone tools to use target platform modules, without needing Engine
+                if (UEBuildConfiguration.bForceBuildTargetPlatforms)
+                {
+                    InModule.AddPlatformSpecificDynamicallyLoadedModule("HTML5TargetPlatform");
+                    InModule.AddDynamicallyLoadedModule("LinuxTargetPlatform");
+                    InModule.AddDynamicallyLoadedModule("LinuxNoEditorTargetPlatform");
+                    InModule.AddDynamicallyLoadedModule("LinuxServerTargetPlatform");
+                }
+
+                if (bBuildShaderFormats)
+                {
+                    InModule.AddDynamicallyLoadedModule("ShaderFormatOpenGL");
+                }
+            }
         }
 
         /**
diff --git a/Engine/Source/Programs/UnrealBuildTool/System/ExternalExecution.cs b/Engine/Source/Programs/UnrealBuildTool/System/ExternalExecution.cs
--- a/Engine/Source/Programs/UnrealBuildTool/System/ExternalExecution.cs
+++ b/Engine/Source/Programs/UnrealBuildTool/System/ExternalExecution.cs
@@ -215,17 +215,12 @@ namespace UnrealBuildTool
 			{
 			case PlatformID.Win32NT:
 				return UnrealTargetPlatform.Win64;
-            case PlatformID.Unix:
-                // Mono returns Unix when running on Mac OSX, so we need some kind of addtional
-                // check check to detect when we are running on Linux vs MacOSX.
-                // TODO(sbc): Find a better way to do this.  Shelling out to uname would seem
-                // to make the most sense but I'm not sure there is an elegant way to do that from
-                // C#.
-                if (File.Exists("/etc/lsb-release") || File.Exists("/etc/debian_version") || (Directory.Exists("/etc/portage") && Directory.Exists("/usr/portage")))
+			case PlatformID.Unix:
+                if (File.Exists("/System/Library/CoreServices/SystemVersion.plist"))
                 {
-                    return UnrealTargetPlatform.Linux;
+                    return UnrealTargetPlatform.Mac;
                 }
-                return UnrealTargetPlatform.Mac;
+                return UnrealTargetPlatform.Linux;
 			default:
 				throw new BuildException("Unhandled runtime platform " + Platform);
 			}
@@ -428,6 +423,7 @@ namespace UnrealBuildTool
 		public static int RunExternalExecutable(string ExePath, string Commandline)
 		{
 			var ExeInfo = new ProcessStartInfo(ExePath, Commandline);
+			Log.TraceVerbose( "RunExternalExecutable {0} {1}", ExePath, Commandline );
 			ExeInfo.UseShellExecute = false;
 			ExeInfo.RedirectStandardOutput = true;
 			using (var GameProcess = Process.Start(ExeInfo))
diff --git a/Engine/Source/Programs/UnrealBuildTool/System/Project.cs b/Engine/Source/Programs/UnrealBuildTool/System/Project.cs
--- a/Engine/Source/Programs/UnrealBuildTool/System/Project.cs
+++ b/Engine/Source/Programs/UnrealBuildTool/System/Project.cs
@@ -402,13 +402,13 @@ namespace UnrealBuildTool
 			}
 
 			// Otherwise make sure it's absolute
-			string FullInputPath = Utils.CleanDirectorySeparators(Path.GetFullPath(InputPath), '\\');
+			string FullInputPath = Utils.CleanDirectorySeparators(Path.GetFullPath(InputPath));
 
 			// Try to make it relative to the solution directory.
-			string FullSolutionPath = Utils.CleanDirectorySeparators(Path.GetFullPath(ProjectFileGenerator.MasterProjectRelativePath), '\\');
-			if (!FullSolutionPath.EndsWith("\\"))
+			string FullSolutionPath = Utils.CleanDirectorySeparators(Path.GetFullPath(ProjectFileGenerator.MasterProjectRelativePath));
+			if (FullSolutionPath.Last() != Path.DirectorySeparatorChar)
 			{
-				FullSolutionPath += "\\";
+				FullSolutionPath += Path.DirectorySeparatorChar;
 			}
 			if (FullInputPath.StartsWith(FullSolutionPath))
 			{
diff --git a/Engine/Source/Programs/UnrealBuildTool/UnrealBuildTool_Mono.csproj b/Engine/Source/Programs/UnrealBuildTool/UnrealBuildTool_Mono.csproj
--- a/Engine/Source/Programs/UnrealBuildTool/UnrealBuildTool_Mono.csproj
+++ b/Engine/Source/Programs/UnrealBuildTool/UnrealBuildTool_Mono.csproj
@@ -12,7 +12,7 @@
     <AssemblyName>UnrealBuildTool</AssemblyName>
     <RunPostBuildEvent>OnOutputUpdated</RunPostBuildEvent>
     <IsWebBootstrapper>false</IsWebBootstrapper>
-    <TargetFrameworkVersion>v4.5</TargetFrameworkVersion>
+    <TargetFrameworkVersion>v4.0</TargetFrameworkVersion>
     <FileUpgradeFlags>
     </FileUpgradeFlags>
     <OldToolsVersion>3.5</OldToolsVersion>
@@ -88,10 +88,110 @@
     <Reference Include="System.Xml" />
     <Reference Include="System.Xml.Linq" />
   </ItemGroup>
+  <Choose>
+    <When Condition="Exists('XboxOne\UEBuildXboxOne.cs')">
+      <ItemGroup>
+        <Compile Include="XboxOne\UEBuildXboxOne.cs" />
+        <Compile Include="XboxOne\XboxOneToolChain.cs" />
+        <Compile Include="XboxOne\XboxOneDeploy.cs" />
+        <Compile Include="XboxOne\XboxOneProjectGenerator.cs" />
+      </ItemGroup>
+    </When>
+  </Choose>
+  <Choose>
+    <When Condition="Exists('PS4\UEBuildPS4.cs')">
+      <ItemGroup>
+        <Compile Include="PS4\PS4ProjectGenerator.cs" />
+        <Compile Include="PS4\PS4ToolChain.cs" />
+        <Compile Include="PS4\UEBuildPS4.cs" />
+      </ItemGroup>
+    </When>
+  </Choose>
+  <Choose>
+    <When Condition="Exists('Cygwin\UEBuildCygwin.cs')">
+      <ItemGroup>
+        <Compile Include="Cygwin\CygwinToolChain.cs" />
+        <Compile Include="Cygwin\UEBuildCygwin.cs" />
+        <Compile Include="Cygwin\CygwinProjectGenerator.cs" />
+      </ItemGroup>
+    </When>
+  </Choose>
+  <Choose>
+    <When Condition="Exists('Mac\UEBuildMac.cs')">
+      <ItemGroup>
+        <Compile Include="Mac\MacToolChain.cs" />
+        <Compile Include="Mac\UEBuildMac.cs" />
+        <Compile Include="Mac\MacProjectGenerator.cs" />
+        <Compile Include="Mac\UEDeployMac.cs" />
+      </ItemGroup>
+    </When>
+  </Choose>
+  <Choose>
+    <When Condition="Exists('Android\UEBuildAndroid.cs')">
+      <ItemGroup>
+        <Compile Include="Android\UEBuildAndroid.cs" />
+        <Compile Include="Android\AndroidToolChain.cs" />
+        <Compile Include="Android\AndroidProjectGenerator.cs" />
+      </ItemGroup>
+    </When>
+  </Choose>
+  <Choose>
+    <When Condition="Exists('WinRT\UEBuildWinRT.cs')">
+      <ItemGroup>
+        <Compile Include="WinRT\UEBuildWinRT.cs" />
+        <Compile Include="WinRT\WinRTToolChain.cs" />
+        <Compile Include="WinRT\WinRTDeploy.cs" />
+        <Compile Include="WinRT\WinRTProjectGenerator.cs" />
+      </ItemGroup>
+    </When>
+  </Choose>
+  <Choose>
+    <When Condition="Exists('IOS\UEBuildIOS.cs')">
+      <ItemGroup>
+        <Compile Include="IOS\IOSProjectGenerator.cs" />
+        <Compile Include="IOS\IOSToolChain.cs" />
+        <Compile Include="IOS\UEBuildIOS.cs" />
+        <Compile Include="IOS\UEDeployIOS.cs" />
+      </ItemGroup>
+    </When>
+  </Choose>
+  <Choose>
+    <When Condition="Exists('HTML5\UEBuildHTML5.cs')">
+      <ItemGroup>
+        <Compile Include="HTML5\HTML5ProjectGenerator.cs" />
+        <Compile Include="HTML5\HTML5ToolChain.cs" />
+        <Compile Include="HTML5\UEBuildHTML5.cs" />
+      </ItemGroup>
+    </When>
+  </Choose>
+  <Choose>
+    <When Condition="Exists('Linux\UEBuildLinux.cs')">
+      <ItemGroup>
+        <Compile Include="Linux\LinuxProjectGenerator.cs" />
+        <Compile Include="Linux\LinuxToolChain.cs" />
+        <Compile Include="Linux\UEBuildLinux.cs" />
+      </ItemGroup>
+    </When>
+  </Choose>
+  <Choose>
+    <When Condition="Exists('System\NotForLicensees\TelemetryProviderET.cs')">
+      <ItemGroup>
+        <Compile Include="System\NotForLicensees\TelemetryProviderET.cs" />
+      </ItemGroup>
+    </When>
+  </Choose>
+  <Choose>
+    <When Condition="Exists('ToolChain\NoRedist\EpicMachineNames.cs')">
+      <ItemGroup>
+        <Compile Include="ToolChain\NoRedist\EpicMachineNames.cs" />
+      </ItemGroup>
+    </When>
+  </Choose>
   <ItemGroup>
     <Compile Include="..\DotNETCommon\MetaData.cs">
       <Link>Properties\MetaData.cs</Link>
     </Compile>
+    <Compile Include="Android\UEDeployAndroid.cs" />
     <Compile Include="Configuration\UEBuildClient.cs" />
     <Compile Include="Configuration\UEBuildDeploy.cs" />
     <Compile Include="Configuration\UEBuildBinary.cs" />
@@ -107,15 +207,6 @@
     <Compile Include="fastJSON\JsonSerializer.cs" />
     <Compile Include="fastJSON\Reflection.cs" />
     <Compile Include="fastJSON\SafeDictionary.cs" />
-    <Compile Include="HTML5\HTML5ProjectGenerator.cs" />
-    <Compile Include="HTML5\HTML5ToolChain.cs" />
-    <Compile Include="HTML5\UEBuildHTML5.cs" />
-    <Compile Include="IOS\IOSProjectGenerator.cs" />
-    <Compile Include="IOS\IOSToolChain.cs" />
-    <Compile Include="IOS\UEBuildIOS.cs" />
-    <Compile Include="Linux\LinuxProjectGenerator.cs" />
-    <Compile Include="Linux\LinuxToolChain.cs" />
-    <Compile Include="Linux\UEBuildLinux.cs" />
     <Compile Include="Properties\AssemblyInfo.cs" />
     <Compile Include="Configuration\UEBuildConfiguration.cs" />
     <Compile Include="Configuration\UEBuildTarget.cs" />
@@ -123,6 +214,9 @@
     <Compile Include="System\ActionHistory.cs" />
     <Compile Include="System\BuildException.cs" />
     <Compile Include="Configuration\BuildConfiguration.cs" />
+    <Compile Include="System\ModuleProcessingException.cs" />
+    <Compile Include="System\MakefileGenerator.cs" />
+    <Compile Include="System\VCSolutionOptions.cs" />
     <Compile Include="System\DirectoryLookupCache.cs" />
     <Compile Include="System\CPPEnvironment.cs" />
     <Compile Include="System\CPPHeaders.cs" />
@@ -133,7 +227,6 @@
     <Compile Include="System\JunkDeleter.cs" />
     <Compile Include="System\LinkEnvironment.cs" />
     <Compile Include="System\LocalExecutor.cs" />
-    <Compile Include="System\MakefileGenerator.cs" />
     <Compile Include="System\Plugins.cs" />
     <Compile Include="System\Project.cs" />
     <Compile Include="System\ProjectFileGenerator.cs" />
@@ -157,20 +250,14 @@
     <Compile Include="System\XGE.cs" />
     <Compile Include="Utilities\CopyrightVerify.cs" />
     <Compile Include="Utilities\DictionaryExtensions.cs" />
-    <Compile Include="Utilities\XmlConfigLoader.cs" />
     <Compile Include="Utilities\GraphVisualization.cs">
       <SubType>Code</SubType>
     </Compile>
+    <Compile Include="Configuration\EngineConfiguration.cs" />
+    <Compile Include="Utilities\XmlConfigLoader.cs" />
     <Compile Include="Windows\UEBuildWindows.cs" />
     <Compile Include="Windows\VCToolChain.cs" />
     <Compile Include="Windows\WindowsProjectGenerator.cs" />
-    <Compile Include="Mac\MacToolChain.cs" />
-    <Compile Include="Mac\UEBuildMac.cs" />
-    <Compile Include="Mac\MacProjectGenerator.cs" />
-    <Compile Include="Configuration\EngineConfiguration.cs" />
-    <Compile Include="IOS\UEDeployIOS.cs" />
-    <Compile Include="System\ModuleProcessingException.cs" />
-    <Compile Include="System\VCSolutionOptions.cs" />
   </ItemGroup>
   <ItemGroup>
     <BootstrapperPackage Include="Microsoft.Net.Client.3.5">
@@ -216,4 +303,4 @@
     <PreBuildEvent>
     </PreBuildEvent>
   </PropertyGroup>
-</Project>
+</Project>
\ No newline at end of file
diff --git a/Engine/Source/Programs/UnrealFrontend/Private/Linux/LinuxUnrealFrontendMain.cpp b/Engine/Source/Programs/UnrealFrontend/Private/Linux/LinuxUnrealFrontendMain.cpp
new file mode 100644
--- /dev/null
+++ b/Engine/Source/Programs/UnrealFrontend/Private/Linux/LinuxUnrealFrontendMain.cpp
@@ -0,0 +1,63 @@
+// Copyright 1998-2014 Epic Games, Inc. All Rights Reserved.
+
+/*=============================================================================
+	LinuxUnrealFrontendMain.cpp: Implements the main entry point for Linux.
+=============================================================================*/
+
+#include "UnrealFrontendMain.h"
+
+static FString GSavedCommandLine;
+
+int main( int argc, char *argv[] )
+{
+	for (int32 Option = 1; Option < argc; ++Option)
+	{
+		GSavedCommandLine += TEXT(" ");
+		FString Argument(ANSI_TO_TCHAR(argv[Option]));
+		if (Argument.Contains(TEXT(" ")))
+		{
+			if (Argument.Contains(TEXT("=")))
+			{
+				FString ArgName;
+				FString ArgValue;
+				Argument.Split( TEXT("="), &ArgName, &ArgValue );
+				Argument = FString::Printf( TEXT("%s=\"%s\""), *ArgName, *ArgValue );
+			}
+			else
+			{
+				Argument = FString::Printf(TEXT("\"%s\""), *Argument);
+			}
+		}
+		GSavedCommandLine += Argument;
+	}
+	
+	FPlatformMisc::SetGracefulTerminationHandler();
+
+#if !UE_BUILD_SHIPPING
+	if (FParse::Param(*GSavedCommandLine,TEXT("crashreports")))
+	{
+		GAlwaysReportCrash = true;
+	}
+#endif
+
+	GUseCrashReportClient = true;
+	
+#if UE_BUILD_DEBUG
+	if (!GAlwaysReportCrash)
+#else
+	if (FPlatformMisc::IsDebuggerPresent() && !GAlwaysReportCrash)
+#endif
+	{
+		UnrealFrontendMain(*GSavedCommandLine);
+	}
+	else
+	{
+		GIsGuarded = 1;
+		UnrealFrontendMain(*GSavedCommandLine);
+		GIsGuarded = 0;
+	}
+
+	FEngineLoop::AppExit();
+	
+	return 0;
+}
diff --git a/Engine/Source/Runtime/Core/Core.Build.cs b/Engine/Source/Runtime/Core/Core.Build.cs
--- a/Engine/Source/Runtime/Core/Core.Build.cs
+++ b/Engine/Source/Runtime/Core/Core.Build.cs
@@ -119,7 +119,10 @@ public class Core : ModuleRules
 				"SDL2"
                 );
 
-            if (UEBuildConfiguration.bCompileAgainstEngine == true && Target.Type != TargetRules.TargetType.Server && UEBuildConfiguration.bCompileSteamOSS)
+            // add steam controller dependency for game and client only
+            if (UEBuildConfiguration.bCompileAgainstEngine == true && 
+                (Target.Type == TargetRules.TargetType.Game || Target.Type == TargetRules.TargetType.Client) &&
+                UEBuildConfiguration.bCompileSteamOSS)
             {
                 AddThirdPartyPrivateStaticDependencies(Target, "SteamController");
             }
