Linux toolchain: replace slashes properly (also account for escaped ones).

diff --git a/Engine/Source/Programs/UnrealBuildTool/Linux/LinuxToolChain.cs b/Engine/Source/Programs/UnrealBuildTool/Linux/LinuxToolChain.cs
--- a/Engine/Source/Programs/UnrealBuildTool/Linux/LinuxToolChain.cs
+++ b/Engine/Source/Programs/UnrealBuildTool/Linux/LinuxToolChain.cs
@@ -768,7 +768,8 @@ namespace UnrealBuildTool
 
             // Add the additional arguments specified by the environment.
             LinkAction.CommandArguments += LinkEnvironment.Config.AdditionalArguments;
-            LinkAction.CommandArguments.Replace("\\", "/");
+            LinkAction.CommandArguments = LinkAction.CommandArguments.Replace("\\\\", "/");
+            LinkAction.CommandArguments = LinkAction.CommandArguments.Replace("\\", "/");
 
             // prepare a linker script
             string LinkerScriptPath = Path.Combine(LinkEnvironment.Config.LocalShadowDirectory, "remove-sym.ldscript");
