PS4 - Runtime conversion of callstack addresses to function names.
* Linker outputs .map
* External PS4MapFile.exe parses the .map file and generates a compressed symbol info file
* At runtime the stack address is used to find the function name from the symbol info file
* Works for all configs except shipping, including packaged builds

diff --git a/Engine/Source/Programs/UnrealBuildTool/Configuration/UEBuildTarget.cs b/Engine/Source/Programs/UnrealBuildTool/Configuration/UEBuildTarget.cs
--- a/Engine/Source/Programs/UnrealBuildTool/Configuration/UEBuildTarget.cs
+++ b/Engine/Source/Programs/UnrealBuildTool/Configuration/UEBuildTarget.cs
@@ -705,7 +705,7 @@ namespace UnrealBuildTool
 
 			// Construct the output path based on configuration, platform, game if not specified.
 			TargetRules.TargetType? TargetType = (Rules != null) ? Rules.Type : (TargetRules.TargetType?)null;
-			OutputPath = Path.GetFullPath(MakeBinaryPath("", AppName, UEBuildBinaryType.Executable, TargetType, bIsRocketGame, null, InAppName, Configuration==UnrealTargetConfiguration.Shipping?Rules.ForceNameAsForDevelopment():false));
+            OutputPath = Path.GetFullPath(MakeBinaryPath("", AppName, UEBuildBinaryType.Executable, TargetType, bIsRocketGame, null, InAppName, Configuration == UnrealTargetConfiguration.Shipping ? Rules.ForceNameAsForDevelopment() : false, Rules.ExeBinariesSubFolder));
 
 			if (bCompileMonolithic && TargetRules.IsGameType(InRulesObject.Type))
 			{
@@ -1974,7 +1974,7 @@ namespace UnrealBuildTool
 
 		/** Given a UBT-built binary name (e.g. "Core"), returns a relative path to the binary for the current build configuration (e.g. "../Binaries/Win64/Core-Win64-Debug.lib") */
 		public static string MakeBinaryPath(string ModuleName, string BinaryName, UnrealTargetPlatform Platform, 
-			UnrealTargetConfiguration Configuration, UEBuildBinaryType BinaryType, TargetRules.TargetType? TargetType, bool bIsRocketModule, PluginInfo PluginInfo, string AppName, bool bForceNameAsForDevelopment = false)
+			UnrealTargetConfiguration Configuration, UEBuildBinaryType BinaryType, TargetRules.TargetType? TargetType, bool bIsRocketModule, PluginInfo PluginInfo, string AppName, bool bForceNameAsForDevelopment = false, string ExeBinariesSubFolder = null)
 		{
 			// Determine the binary extension for the platform and binary type.
 			var BuildPlatform = UEBuildPlatform.GetBuildPlatform(Platform);
@@ -1987,7 +1987,7 @@ namespace UnrealBuildTool
 			}
 			
 			string ModuleBinariesSubDir = "";
-			if ((BinaryType == UEBuildBinaryType.DynamicLinkLibrary) && (string.IsNullOrEmpty(ModuleName) == false))
+            if (BinaryType == UEBuildBinaryType.DynamicLinkLibrary && (string.IsNullOrEmpty(ModuleName) == false))
 			{
 				// Allow for modules to specify sub-folders in the Binaries folder
 				var RulesFilename = RulesCompiler.GetModuleFilename(ModuleName);
@@ -2001,6 +2001,10 @@ namespace UnrealBuildTool
 					}
 				}
 			}
+            else if ( BinaryType == UEBuildBinaryType.Executable && string.IsNullOrEmpty(ExeBinariesSubFolder) == false )
+            {
+                ModuleBinariesSubDir = ExeBinariesSubFolder;
+            }
 
 			[REDACTED_EMAIL]: This just happens to work since exp and lib files go into intermediate...
 
@@ -2049,7 +2053,7 @@ namespace UnrealBuildTool
 		}
 
 		/** Given a UBT-built binary name (e.g. "Core"), returns a relative path to the binary for the current build configuration (e.g. "../../Binaries/Win64/Core-Win64-Debug.lib") */
-		public string MakeBinaryPath(string ModuleName, string BinaryName, UEBuildBinaryType BinaryType, TargetRules.TargetType? TargetType, bool bIsRocketModule, PluginInfo PluginInfo, string AppName, bool bForceNameAsForDevelopment = false)
+		public string MakeBinaryPath(string ModuleName, string BinaryName, UEBuildBinaryType BinaryType, TargetRules.TargetType? TargetType, bool bIsRocketModule, PluginInfo PluginInfo, string AppName, bool bForceNameAsForDevelopment = false, string ExeBinariesSubFolder = null)
 		{
 			if (String.IsNullOrEmpty(ModuleName) && Configuration == UnrealTargetConfiguration.DebugGame && !bCompileMonolithic)
 			{
@@ -2057,7 +2061,7 @@ namespace UnrealBuildTool
 			}
 			else
 			{
-				return MakeBinaryPath(ModuleName, BinaryName, Platform, Configuration, BinaryType, TargetType, bIsRocketModule, PluginInfo, AppName, bForceNameAsForDevelopment);
+				return MakeBinaryPath(ModuleName, BinaryName, Platform, Configuration, BinaryType, TargetType, bIsRocketModule, PluginInfo, AppName, bForceNameAsForDevelopment, ExeBinariesSubFolder);
 			}
 		}
 
diff --git a/Engine/Source/Programs/UnrealBuildTool/System/RulesCompiler.cs b/Engine/Source/Programs/UnrealBuildTool/System/RulesCompiler.cs
--- a/Engine/Source/Programs/UnrealBuildTool/System/RulesCompiler.cs
+++ b/Engine/Source/Programs/UnrealBuildTool/System/RulesCompiler.cs
@@ -441,6 +441,11 @@ namespace UnrealBuildTool
 		/// </summary>
 		public bool bOutputToEngineBinaries = false;
 
+        /// <summary>
+        /// Sub folder where the built target goes: Engine/Binaries/<PLATFORM>/<SUBDIR>
+        /// </summary>
+        public string ExeBinariesSubFolder = String.Empty;
+
 		/// <summary>
 		/// Whether this target should be compiled in monolithic mode
 		/// </summary>
