Added UnrealBuildTool performance diagnostic for UObject file gather time

diff --git a/Engine/Source/Programs/UnrealBuildTool/Configuration/UEBuildTarget.cs b/Engine/Source/Programs/UnrealBuildTool/Configuration/UEBuildTarget.cs
--- a/Engine/Source/Programs/UnrealBuildTool/Configuration/UEBuildTarget.cs
+++ b/Engine/Source/Programs/UnrealBuildTool/Configuration/UEBuildTarget.cs
@@ -1499,6 +1499,8 @@ namespace UnrealBuildTool
 
 			if( (BuildConfiguration.bXGEExport && UEBuildConfiguration.bGenerateManifest) || (!UEBuildConfiguration.bGenerateManifest && !UEBuildConfiguration.bCleanProject && !ProjectFileGenerator.bGenerateProjectFiles) )
 			{
+				var UObjectDiscoveryStartTime = DateTime.UtcNow;
+
 				var UObjectModules = new List<UHTModuleInfo>();
 
 				// Figure out which modules have UObjects that we may need to generate headers for
@@ -1536,6 +1538,12 @@ namespace UnrealBuildTool
 						Log.TraceVerbose( "Detected UObject module: " + DependencyModuleCPP.Name );
 					}
 				}
+				
+				if( BuildConfiguration.bPrintPerformanceInfo )
+				{
+					double UObjectDiscoveryTime = ( DateTime.UtcNow - UObjectDiscoveryStartTime ).TotalSeconds;
+					Trace.TraceInformation( "UObject discovery time: " + UObjectDiscoveryTime + "s" );
+				}
 
 				if( UObjectModules.Count > 0 )
 				{
