Removed obsoleted UObject source file checking code
- This hasn't been needed for a few months, since the new GetUHTModuleInfo() implementation

diff --git a/Engine/Source/Programs/UnrealBuildTool/Configuration/UEBuildModule.cs b/Engine/Source/Programs/UnrealBuildTool/Configuration/UEBuildModule.cs
--- a/Engine/Source/Programs/UnrealBuildTool/Configuration/UEBuildModule.cs
+++ b/Engine/Source/Programs/UnrealBuildTool/Configuration/UEBuildModule.cs
@@ -1748,7 +1748,6 @@ namespace UnrealBuildTool
 				if (!Regex.IsMatch(FileContents, "^\\s*U(CLASS|STRUCT|ENUM|INTERFACE|DELEGATE)\\b", RegexOptions.Multiline))
 					continue;
 
-				UObjectHeaderFileItem.HasUObjects = true;
 				if (UObjectHeaderFileItem.AbsolutePath.StartsWith(ClassesFolder))
 				{
 					_AllClassesHeaders.Add(UObjectHeaderFileItem);
diff --git a/Engine/Source/Programs/UnrealBuildTool/System/CPPHeaders.cs b/Engine/Source/Programs/UnrealBuildTool/System/CPPHeaders.cs
--- a/Engine/Source/Programs/UnrealBuildTool/System/CPPHeaders.cs
+++ b/Engine/Source/Programs/UnrealBuildTool/System/CPPHeaders.cs
@@ -327,15 +327,11 @@ namespace UnrealBuildTool
 		{
 			// Try to fulfill request from cache first.
 			List<DependencyInclude> Result = null;
-			bool HasUObjects;
-			if( IncludeDependencyCache.GetDirectDependencies( CPPFile, out Result, out HasUObjects ) )
+			if( IncludeDependencyCache.GetDirectDependencies( CPPFile, out Result ) )
 			{
-				CPPFile.HasUObjects = HasUObjects;
 				return Result;
 			}
 
-			HasUObjects = false;
-
 			var TimerStartTime = DateTime.UtcNow;
 			++CPPEnvironment.TotalDirectIncludeCacheMisses;
 
@@ -379,21 +375,18 @@ namespace UnrealBuildTool
 						EndIndex = FileContents.Length;
 					}
 
-					CollectHeaders(CPPFile, Result, ref HasUObjects, FileToRead, FileContents, InstalledFolder, StartIndex, EndIndex);
+					CollectHeaders(CPPFile, Result, FileToRead, FileContents, InstalledFolder, StartIndex, EndIndex);
 
 					StartIndex = EndIndex + 1;
 				}
 			}
 			else
 			{
-				CollectHeaders(CPPFile, Result, ref HasUObjects, FileToRead, FileContents, InstalledFolder, 0, FileContents.Length);
+				CollectHeaders(CPPFile, Result, FileToRead, FileContents, InstalledFolder, 0, FileContents.Length);
 			}
 
-			// Store whether this C++ file has UObjects right inside the FileItem for this file
-			CPPFile.HasUObjects = HasUObjects;
-
 			// Populate cache with results.
-			IncludeDependencyCache.SetDependencyInfo( CPPFile, Result, CPPFile.HasUObjects );
+			IncludeDependencyCache.SetDependencyInfo( CPPFile, Result );
 
 			CPPEnvironment.DirectIncludeCacheMissesTotalTime += ( DateTime.UtcNow - TimerStartTime ).TotalSeconds;
 
@@ -405,13 +398,12 @@ namespace UnrealBuildTool
 		/// </summary>
 		/// <param name="CPPFile"></param>
 		/// <param name="Result"></param>
-		/// <param name="HasUObjects"></param>
 		/// <param name="FileToRead"></param>
 		/// <param name="FileContents"></param>
 		/// <param name="InstalledFolder"></param>
 		/// <param name="StartIndex"></param>
 		/// <param name="EndIndex"></param>
-		private static void CollectHeaders(FileItem CPPFile, List<DependencyInclude> Result, ref bool HasUObjects, string FileToRead, string FileContents, string InstalledFolder, int StartIndex, int EndIndex)
+		private static void CollectHeaders(FileItem CPPFile, List<DependencyInclude> Result, string FileToRead, string FileContents, string InstalledFolder, int StartIndex, int EndIndex)
 		{
 			Match M = CPPHeaderRegex.Match(FileContents, StartIndex, EndIndex - StartIndex);
 			CaptureCollection CC = M.Groups["HeaderFile"].Captures;
@@ -420,14 +412,6 @@ namespace UnrealBuildTool
 			{
 				string HeaderValue = C.Value;
 
-				// Is this include statement including a generated UObject header?  If so, then we know we'll need to pass this file to UnrealHeaderTool for code generation
-				if (!HasUObjects && HeaderValue.EndsWith(".generated.h", StringComparison.InvariantCultureIgnoreCase))
-				{
-					// @todo uht: Ideally we would also check for other clues, like USTRUCT or UCLASS.
-					// @todo uht: We need to be careful to skip comments or string text that may have these clues
-					HasUObjects = true;
-				}
-
 				//@TODO: The intermediate exclusion is to work around autogenerated absolute paths in Module.SomeGame.cpp style files
 				bool bIsIntermediateOrThirdParty = FileToRead.Contains("Intermediate") || FileToRead.Contains("ThirdParty");
 				bool bCheckForBackwardSlashes = FileToRead.StartsWith(InstalledFolder);
diff --git a/Engine/Source/Programs/UnrealBuildTool/System/DependencyCache.cs b/Engine/Source/Programs/UnrealBuildTool/System/DependencyCache.cs
--- a/Engine/Source/Programs/UnrealBuildTool/System/DependencyCache.cs
+++ b/Engine/Source/Programs/UnrealBuildTool/System/DependencyCache.cs
@@ -42,9 +42,6 @@ namespace UnrealBuildTool
 		/** List of files this file includes.  These are direct include paths and cannot be resolved to an actual file on 
 		    disk without using the proper list of include directories for this file's module */
 		public List<DependencyInclude> Includes;
-
-		/** True if this file contains UObjects */
-		public bool HasUObjects;
 	}
 
 
@@ -213,12 +210,10 @@ namespace UnrealBuildTool
 		 * 
 		 * @param	File				File to try to find dependencies in cache
 		 * @param	Result	[out]		List of dependencies if successful, null otherwise
-		 * @param	HasUObjects [out]	True if this file contains UObjects that may need to be passed along to UnrealHeaderTool to generate code for
 		 */
-		public bool GetDirectDependencies(FileItem File, out List<DependencyInclude> Result, out bool HasUObjects)
+		public bool GetDirectDependencies(FileItem File, out List<DependencyInclude> Result)
 		{
 			Result = null;
-			HasUObjects = false;
 
 			// Check whether File is in cache.
 			DependencyInfo DependencyInfo;
@@ -250,7 +245,6 @@ namespace UnrealBuildTool
 					}
 					// Cached version is up to date, return it.
 					Result = DependencyInfo.Includes;
-					HasUObjects = DependencyInfo.HasUObjects;
 					return true;
 				}
 				// File has been modified since the last time it was cached.
@@ -283,11 +277,10 @@ namespace UnrealBuildTool
 		 * 
 		 * @param	File			File to update dependencies for
 		 * @param	Dependencies	List of dependencies to cache for passed in file
-		 * @param	HasUObjects		True if this file contains UObjects
 		 */
-		public void SetDependencyInfo(FileItem File, List<DependencyInclude> DirectDependencies, bool HasUObjects)
+		public void SetDependencyInfo(FileItem File, List<DependencyInclude> DirectDependencies)
 		{
-			DependencyMap[File.AbsolutePath] = new DependencyInfo() { Includes = DirectDependencies, HasUObjects = HasUObjects };
+			DependencyMap[File.AbsolutePath] = new DependencyInfo() { Includes = DirectDependencies };
 			bIsDirty = true;
 		}
 
diff --git a/Engine/Source/Programs/UnrealBuildTool/System/FileItem.cs b/Engine/Source/Programs/UnrealBuildTool/System/FileItem.cs
--- a/Engine/Source/Programs/UnrealBuildTool/System/FileItem.cs
+++ b/Engine/Source/Programs/UnrealBuildTool/System/FileItem.cs
@@ -53,10 +53,6 @@ namespace UnrealBuildTool
 		/** Whether or not this is a remote file, in which case we can't access it directly */
 		public bool bIsRemoteFile = false;
 
-		/** True if this is a C++ file that was found to contains UObject class or struct declarations, and may need to be processed by UnrealHeaderTool to 
-		    generate code before compiling.  This is only valid after GetIncludeDependencies() has been called for this file */
-		public bool HasUObjects = false;
-
 		/** A list of remote file items that have been created but haven't needed the remote info yet, so we can gang up many into one request */
 		static List<FileItem> DelayedRemoteLookupFiles = new List<FileItem>();
 
