UnrealBuildTool: No longer scans all includes before running UHT, unless UObject modules are out of date
- This can reduce the time it takes for UBT to run when no UObject files were changed

diff --git a/Engine/Source/Programs/UnrealBuildTool/Configuration/UEBuildModule.cs b/Engine/Source/Programs/UnrealBuildTool/Configuration/UEBuildModule.cs
--- a/Engine/Source/Programs/UnrealBuildTool/Configuration/UEBuildModule.cs
+++ b/Engine/Source/Programs/UnrealBuildTool/Configuration/UEBuildModule.cs
@@ -1062,6 +1062,9 @@ namespace UnrealBuildTool
 				return LinkInputFiles;
 			}
 
+			// Process all of the header file dependencies for this module
+			this.ProcessAllCppDependencies( ModuleCompileEnvironment );
+
 			// Throw an error if the module's source file list referenced any non-existent files.
 			if (SourceFilesToBuild.MissingFiles.Count > 0)
 			{
@@ -1541,18 +1544,18 @@ namespace UnrealBuildTool
 
 		public void ProcessAllCppDependencies(CPPEnvironment ModuleCompileEnvironment)
 		{
-			if (ProcessedDependencies != null)
-				return;
-
-			FileItem UniquePCH = null;
-			foreach( var CPPFile in SourceFilesFound.CPPFiles )
+			if( ProcessedDependencies == null )
 			{
-				// Find headers used by the source file.
-				var PCH = ProcessDependencies(CPPFile, ModuleCompileEnvironment);
-				UniquePCH = UniquePCH ?? PCH;
-			}
+				FileItem UniquePCH = null;
+				foreach( var CPPFile in SourceFilesFound.CPPFiles )
+				{
+					// Find headers used by the source file.
+					var PCH = ProcessDependencies(CPPFile, ModuleCompileEnvironment);
+					UniquePCH = UniquePCH ?? PCH;
+				}
 
-			ProcessedDependencies = new ProcessedDependenciesClass{ UniquePCHHeaderFile = UniquePCH };
+				ProcessedDependencies = new ProcessedDependenciesClass{ UniquePCHHeaderFile = UniquePCH };
+			}
 		}
 
 		private FileItem ProcessDependencies(FileItem CPPFile, CPPEnvironment ModuleCompileEnvironment)
@@ -1564,7 +1567,9 @@ namespace UnrealBuildTool
 			}
 
 			if (DirectIncludeFilenames.Count == 0)
+			{ 
 				return null;
+			}
 
 			var FirstInclude = DirectIncludeFilenames[0];
 
@@ -1716,7 +1721,7 @@ namespace UnrealBuildTool
 		public HashSet<FileItem> _PublicUObjectHeaders  = new HashSet<FileItem>();
 		public HashSet<FileItem> _PrivateUObjectHeaders = new HashSet<FileItem>();
 
-		private UHTModuleInfo? CachedModuleUHTInfo = null;
+		private UHTModuleInfo CachedModuleUHTInfo = null;
 
 		/// Total time spent figuring out which PCH to use for each module and source file
 		public static double TotalPCHTime = 0.0;
@@ -1732,10 +1737,12 @@ namespace UnrealBuildTool
 		/// <param name="PublicUObjectHeaders">Dependent UObject headers in Public source folders</param>
 		/// <param name="PrivateUObjectHeaders">Dependent UObject headers not in Public source folders</param>
 		/// <returns>
-		public UHTModuleInfo GetUHTModuleInfo(CPPEnvironment CompileEnvironment)
+		public UHTModuleInfo GetUHTModuleInfo()
 		{
-			if (CachedModuleUHTInfo.HasValue)
-				return CachedModuleUHTInfo.Value;
+			if (CachedModuleUHTInfo != null)
+			{
+				return CachedModuleUHTInfo;
+			}
 
 			var ClassesFolder = Path.Combine(this.ModuleDirectory, "Classes");
 			var PublicFolder  = Path.Combine(this.ModuleDirectory, "Public");
@@ -1746,7 +1753,9 @@ namespace UnrealBuildTool
 				var UObjectHeaderFileItem = FileItem.GetExistingItemByPath( ClassHeader );
 				var FileContents = Utils.ReadAllText(UObjectHeaderFileItem.AbsolutePath);
 				if (!Regex.IsMatch(FileContents, "^\\s*U(CLASS|STRUCT|ENUM|INTERFACE|DELEGATE)\\b", RegexOptions.Multiline))
+				{ 
 					continue;
+				}
 
 				if (UObjectHeaderFileItem.AbsolutePath.StartsWith(ClassesFolder))
 				{
@@ -1771,7 +1780,7 @@ namespace UnrealBuildTool
 				PrivateUObjectHeaders       = _PrivateUObjectHeaders.ToList()
 			};
 
-			return CachedModuleUHTInfo.Value;
+			return CachedModuleUHTInfo;
 		}
 
 		public override void GetAllDependencyModules( ref Dictionary<string, UEBuildModule> ReferencedModules, ref List<UEBuildModule> OrderedModules, bool bIncludeDynamicallyLoaded, bool bForceCircular )
diff --git a/Engine/Source/Programs/UnrealBuildTool/Configuration/UEBuildTarget.cs b/Engine/Source/Programs/UnrealBuildTool/Configuration/UEBuildTarget.cs
--- a/Engine/Source/Programs/UnrealBuildTool/Configuration/UEBuildTarget.cs
+++ b/Engine/Source/Programs/UnrealBuildTool/Configuration/UEBuildTarget.cs
@@ -1510,32 +1510,16 @@ namespace UnrealBuildTool
 					foreach( var DependencyModuleCPP in DependencyModules.OfType<UEBuildModuleCPP>().Where( CPPModule => !UObjectModules.Any( Module => Module.ModuleName == CPPModule.Name ) ) )
 					{
 						if( !DependencyModuleCPP.bIncludedInTarget )
-							throw new BuildException( "Expecting module {0} to have bIncludeInTarget set", DependencyModuleCPP.Name );
-
-						var ModuleCompileEnvironment = DependencyModuleCPP.CreateModuleCompileEnvironment(GlobalCompileEnvironment);
-						DependencyModuleCPP.ProcessAllCppDependencies(ModuleCompileEnvironment);
-
-						var UHTModuleInfo = DependencyModuleCPP.GetUHTModuleInfo(ModuleCompileEnvironment);
-						if( UHTModuleInfo.PublicUObjectClassesHeaders.Count == 0 && UHTModuleInfo.PrivateUObjectHeaders.Count == 0 && UHTModuleInfo.PublicUObjectHeaders.Count == 0 )
-							continue;
-
-						UHTModuleInfo.PCH                      = "";
-						UHTModuleInfo.GeneratedCPPFilenameBase = Path.Combine( UEBuildModuleCPP.GetGeneratedCodeDirectoryForModule(this, UHTModuleInfo.ModuleDirectory, UHTModuleInfo.ModuleName), UHTModuleInfo.ModuleName ) + ".generated";
-
-						if (DependencyModuleCPP.ProcessedDependencies.UniquePCHHeaderFile != null)
 						{
-							UHTModuleInfo.PCH = DependencyModuleCPP.ProcessedDependencies.UniquePCHHeaderFile.AbsolutePath;
+							throw new BuildException( "Expecting module {0} to have bIncludeInTarget set", DependencyModuleCPP.Name );
 						}
 
-						// If we've got this far and there are no source files then it's likely we're running Rocket and ignoring
-						// engine files, so we don't need a .generated.cpp either
-						if (DependencyModuleCPP.SourceFilesToBuild.Count != 0)
+						var UHTModuleInfo = DependencyModuleCPP.GetUHTModuleInfo();
+						if( UHTModuleInfo.PublicUObjectClassesHeaders.Count > 0 || UHTModuleInfo.PrivateUObjectHeaders.Count > 0 || UHTModuleInfo.PublicUObjectHeaders.Count > 0 )
 						{
-							DependencyModuleCPP.AutoGenerateCppInfo = new UEBuildModuleCPP.AutoGenerateCppInfoClass( UHTModuleInfo.GeneratedCPPFilenameBase + ".cpp" );
+							UObjectModules.Add( UHTModuleInfo );
+							Log.TraceVerbose( "Detected UObject module: " + UHTModuleInfo.ModuleName );
 						}
-
-						UObjectModules.Add( UHTModuleInfo );
-						Log.TraceVerbose( "Detected UObject module: " + DependencyModuleCPP.Name );
 					}
 				}
 				
@@ -1552,7 +1536,7 @@ namespace UnrealBuildTool
 					string ModuleInfoFileName = Path.GetFullPath( Path.Combine( ProjectIntermediateDirectory, "UnrealHeaderTool.manifest" ) );
 
 					ECompilationResult UHTResult = ECompilationResult.OtherCompilationError;
-					if (!ExternalExecution.ExecuteHeaderToolIfNecessary(this, UObjectModules, ModuleInfoFileName, ref UHTResult))
+					if (!ExternalExecution.ExecuteHeaderToolIfNecessary(this, GlobalCompileEnvironment, UObjectModules, ModuleInfoFileName, ref UHTResult))
 					{
 						Log.TraceInformation("UnrealHeaderTool failed for target '" + TargetName + "' (platform: " + Platform.ToString() + ", module info: " + ModuleInfoFileName + ").");
 						return UHTResult;
diff --git a/Engine/Source/Programs/UnrealBuildTool/System/ExternalExecution.cs b/Engine/Source/Programs/UnrealBuildTool/System/ExternalExecution.cs
--- a/Engine/Source/Programs/UnrealBuildTool/System/ExternalExecution.cs
+++ b/Engine/Source/Programs/UnrealBuildTool/System/ExternalExecution.cs
@@ -23,7 +23,7 @@ namespace UnrealBuildTool
 	}
 
 	/** Information about a module that needs to be passed to UnrealHeaderTool for code generation */
-	public struct UHTModuleInfo
+	public class UHTModuleInfo
 	{
 		/** Module name */
 		public string ModuleName;
@@ -453,21 +453,51 @@ namespace UnrealBuildTool
 		 * Builds and runs the header tool and touches the header directories.
 		 * Performs any early outs if headers need no changes, given the UObject modules, tool path, game name, and configuration
 		 */
-		public static bool ExecuteHeaderToolIfNecessary( UEBuildTarget Target, List<UHTModuleInfo> UObjectModules, string ModuleInfoFileName, ref ECompilationResult UHTResult )
+		public static bool ExecuteHeaderToolIfNecessary( UEBuildTarget Target, CPPEnvironment GlobalCompileEnvironment, List<UHTModuleInfo> UObjectModules, string ModuleInfoFileName, ref ECompilationResult UHTResult )
 		{
-			// We never want to try to execute the header tool when we're already trying to build it!
-			var bIsBuildingUHT = Target.GetTargetName().Equals( "UnrealHeaderTool", StringComparison.InvariantCultureIgnoreCase );
-			var BuildPlatform  = UEBuildPlatform.GetBuildPlatform(Target.Platform);
-			var CppPlatform    = BuildPlatform.GetCPPTargetPlatform(Target.Platform);
-			var ToolChain      = UEToolChain.GetPlatformToolChain(CppPlatform);
-			var RootLocalPath  = Path.GetFullPath(ProjectFileGenerator.RootRelativePath);
-			var Manifest       = new UHTManifest(Target, RootLocalPath, ToolChain.ConvertPath(RootLocalPath + '\\'), UObjectModules);
-
 			using (ProgressWriter Progress = new ProgressWriter("Generating headers...", false))
 			{
+				// We never want to try to execute the header tool when we're already trying to build it!
+				var bIsBuildingUHT = Target.GetTargetName().Equals( "UnrealHeaderTool", StringComparison.InvariantCultureIgnoreCase );
+
+				var BuildPlatform = UEBuildPlatform.GetBuildPlatform(Target.Platform);
+				var CppPlatform = BuildPlatform.GetCPPTargetPlatform(Target.Platform);
+				var ToolChain = UEToolChain.GetPlatformToolChain(CppPlatform);
+				var RootLocalPath  = Path.GetFullPath(ProjectFileGenerator.RootRelativePath);
+
+
 				// ensure the headers are up to date
-				if (!bIsBuildingUHT && 
-					(UEBuildConfiguration.bForceHeaderGeneration == true || AreGeneratedCodeFilesOutOfDate(Target, UObjectModules)))
+				bool bUHTNeedsToRun = (UEBuildConfiguration.bForceHeaderGeneration == true || AreGeneratedCodeFilesOutOfDate(Target, UObjectModules));
+				if( bUHTNeedsToRun )
+				{
+					// Since code files are definitely out of date, we'll now finish computing information about the UObject modules for UHT.  We
+					// want to save this work until we know that UHT actually needs to be run to speed up best-case iteration times.
+					foreach( var UHTModuleInfo in UObjectModules )
+					{
+						UHTModuleInfo.PCH = "";
+						UHTModuleInfo.GeneratedCPPFilenameBase = Path.Combine( UEBuildModuleCPP.GetGeneratedCodeDirectoryForModule(Target, UHTModuleInfo.ModuleDirectory, UHTModuleInfo.ModuleName), UHTModuleInfo.ModuleName ) + ".generated";
+
+						// We need to figure out which PCH header this module is including, so that UHT can inject an include statement for it into any .cpp files it is synthesizing
+						var DependencyModuleCPP = (UEBuildModuleCPP)Target.GetModuleByName( UHTModuleInfo.ModuleName );
+						var ModuleCompileEnvironment = DependencyModuleCPP.CreateModuleCompileEnvironment(GlobalCompileEnvironment);
+						DependencyModuleCPP.ProcessAllCppDependencies(ModuleCompileEnvironment);
+						if (DependencyModuleCPP.ProcessedDependencies.UniquePCHHeaderFile != null)
+						{
+							UHTModuleInfo.PCH = DependencyModuleCPP.ProcessedDependencies.UniquePCHHeaderFile.AbsolutePath;
+						}
+
+						// If we've got this far and there are no source files then it's likely we're running Rocket and ignoring
+						// engine files, so we don't need a .generated.cpp either
+						if (DependencyModuleCPP.SourceFilesToBuild.Count != 0)
+						{
+							DependencyModuleCPP.AutoGenerateCppInfo = new UEBuildModuleCPP.AutoGenerateCppInfoClass( UHTModuleInfo.GeneratedCPPFilenameBase + ".cpp" );
+						}
+					}
+				}
+
+				UHTManifest Manifest = new UHTManifest(Target, RootLocalPath, ToolChain.ConvertPath(RootLocalPath + '\\'), UObjectModules);
+
+				if( !bIsBuildingUHT && bUHTNeedsToRun )
 				{
 					// Always build UnrealHeaderTool if header regeneration is required, unless we're running within a Rocket ecosystem
 					if (UnrealBuildTool.RunningRocket() == false && UEBuildConfiguration.bDoNotBuildUHT == false)
@@ -494,7 +524,9 @@ namespace UnrealBuildTool
 						}
 
 						if ( RunExternalExecutable( UnrealBuildTool.GetUBTPath(), UBTArguments.ToString() ) != 0 )
+						{ 
 							return false;
+						}
 					}
 
 					Progress.Write(1, 3);
