Refactor AutoSDKs in preparation for doing doing AutoSDK setup for the editor.

diff --git a/Engine/Source/Programs/UnrealBuildTool/Android/UEBuildAndroid.cs b/Engine/Source/Programs/UnrealBuildTool/Android/UEBuildAndroid.cs
--- a/Engine/Source/Programs/UnrealBuildTool/Android/UEBuildAndroid.cs
+++ b/Engine/Source/Programs/UnrealBuildTool/Android/UEBuildAndroid.cs
@@ -21,10 +21,6 @@ namespace UnrealBuildTool
 		[XmlConfig]
 		public static string AndroidSdkApiTarget = "latest";
 
-        // this is to make sure we don't check the sdk all the time
-        static int HasSDK = -1;
-
-
 		// The current architecture - affects everything about how UBT operates on Android
 		public override string GetActiveArchitecture()
 		{
@@ -37,7 +33,7 @@ namespace UnrealBuildTool
             return false;
         }
 
-        public override bool PlatformSupportsSDKSwitching()
+        protected override bool PlatformSupportsAutoSDKs()
         {
             return true;
         }
@@ -47,14 +43,14 @@ namespace UnrealBuildTool
             return "Android";
         }
 
-        public override string GetRequiredSDKString()
+        protected override string GetRequiredSDKString()
         {
             return "-19";
         }
 
-        public override String GetRequiredScriptVersionString()
+        protected override String GetRequiredScriptVersionString()
         {
-            return "1.1";
+            return "2.0";
         }
 
         /// <summary>
@@ -63,11 +59,6 @@ namespace UnrealBuildTool
         /// <returns></returns>
         private bool HasAnySDK()
         {
-            if (base.HasRequiredSDKsInstalled() == SDKStatus.Valid)
-            {
-                return true;
-            }
-
             string NDKPath = Environment.GetEnvironmentVariable("NDKROOT");
 
             // we don't have an NDKROOT specified
@@ -87,23 +78,23 @@ namespace UnrealBuildTool
         }
 
 
-		public override SDKStatus HasRequiredSDKsInstalled()
+        protected override SDKStatus HasRequiredManualSDKInternal()
 		{
-            if (HasSDK == -1)
+            // if any autosdk setup has been done then the local process environment is suspect
+            if (HasSetupAutoSDK())
             {
-                if (HasAnySDK())
-                {
-                    HasSDK = 1;
-                }
-                else
-                {
-                    HasSDK = 0;
-                }
+                return SDKStatus.Invalid;
             }
-            return HasSDK == 1 ? SDKStatus.Valid : SDKStatus.Invalid;
+
+            if (HasAnySDK())
+            {
+                return SDKStatus.Valid;
+            }
+
+            return SDKStatus.Invalid;            
 		}
 
-		public override void RegisterBuildPlatform()
+		protected override void RegisterBuildPlatformInternal()
 		{
 			if ((ProjectFileGenerator.bGenerateProjectFiles == true) || (HasRequiredSDKsInstalled() == SDKStatus.Valid))
 			{
diff --git a/Engine/Source/Programs/UnrealBuildTool/Configuration/UEBuildPlatform.cs b/Engine/Source/Programs/UnrealBuildTool/Configuration/UEBuildPlatform.cs
--- a/Engine/Source/Programs/UnrealBuildTool/Configuration/UEBuildPlatform.cs
+++ b/Engine/Source/Programs/UnrealBuildTool/Configuration/UEBuildPlatform.cs
@@ -11,8 +11,7 @@
 	public enum SDKStatus
 	{
 		Valid,			// Desired SDK is installed and set up.
-		Invalid,		// Could not find the desired SDK, SDK setup failed, etc.
-		Unknown,		// Not enough info to determine SDK state. 
+		Invalid,		// Could not find the desired SDK, SDK setup failed, etc.		
 	};
 
 	public interface IUEBuildPlatform
@@ -243,31 +242,13 @@
 		void SetupBinaries(UEBuildTarget InBuildTarget);
 	}
 
-	public abstract class UEBuildPlatform : IUEBuildPlatform
+	public abstract partial class UEBuildPlatform : IUEBuildPlatform
 	{
 		public static Dictionary<UnrealTargetPlatform, IUEBuildPlatform> BuildPlatformDictionary = new Dictionary<UnrealTargetPlatform, IUEBuildPlatform>();
 
 		// a mapping of a group to the platforms in the group (ie, Microsoft contains Win32 and Win64)
-		static Dictionary<UnrealPlatformGroup, List<UnrealTargetPlatform>> PlatformGroupDictionary = new Dictionary<UnrealPlatformGroup, List<UnrealTargetPlatform>>();
-
-        /** Name of the file that holds currently install SDK version string */
-        static string CurrentlyInstalledSDKStringManifest = "CurrentlyInstalled.txt";
-
-        /** name of the file that holds the last succesfully run SDK setup script version */
-        static string LastRunScriptVersionManifest = "CurrentlyInstalled.Version.txt";
-
-        /** Name of the file that holds environment variables of current SDK */
-        static string SDKEnvironmentVarsFile = "OutputEnvVars.txt";
-
-        /** Enum describing types of hooks a platform SDK can have */
-        public enum SDKHookType
-        {
-            Install,
-            Uninstall
-        };
-
-		public static readonly string SDKRootEnvVar = "UE_SDKS_ROOT";
-
+		static Dictionary<UnrealPlatformGroup, List<UnrealTargetPlatform>> PlatformGroupDictionary = new Dictionary<UnrealPlatformGroup, List<UnrealTargetPlatform>>();        
+        
         /**
          * Attempt to convert a string to an UnrealTargetPlatform enum entry
          * 
@@ -406,241 +387,7 @@
                     PlatformEntry.Value.ModifyNewlyLoadedModule(InModule, Target);
                 }
 			}
-		}
-
-        /** 
-         * Whether platform supports switching SDKs during runtime
-         * 
-         * @return true if supports
-         */
-        public virtual bool PlatformSupportsSDKSwitching()
-        {
-            return false;
-        }
-
-        /** 
-         * Returns platform-specific name used in SDK repository
-         * 
-         * @return path to SDK Repository
-         */
-        public virtual string GetSDKTargetPlatformName()
-        {
-            return "";
-        }
-
-        /** 
-         * Returns SDK string as required by the platform 
-         * 
-         * @return Valid SDK string
-         */
-        public virtual string GetRequiredSDKString()
-        {
-            return "";
-        }
+		}                      
 
         /**
-        * Gets the version number of the SDK setup script itself.  The version in the base should ALWAYS be 1.0.  If you need to force a rebuild for a given platform, override this for the given platform.
-        * @return Setup script version
-        */
-        public virtual String GetRequiredScriptVersionString()
-        {
-            return "1.0";
-        }
-
-        /** 
-         * Returns path to platform SDKs
-         * 
-         * @return Valid SDK string
-         */
-        public virtual string GetPathToPlatformSDKs(out bool SDKRootSet)
-        {
-            string SDKPath = "";
-            SDKRootSet = false;
-
-            string SDKRoot = Environment.GetEnvironmentVariable(SDKRootEnvVar);
-            if (SDKRoot != null)
-            {
-                SDKRootSet = true;
-                if (SDKRoot != "")
-                {
-                    SDKPath = Path.Combine(SDKRoot, "Host" + ExternalExecution.GetRuntimePlatform(), GetSDKTargetPlatformName());
-                }
-            }
-            return SDKPath;
-        }
-
-        /**
-         * Gets currently installed version
-         * 
-         * @param PlatformSDKRoot absolute path to platform SDK root
-         * @param OutInstalledSDKVersionString version string as currently installed
-         * 
-         * @return true if was able to read it
-         */
-        public virtual bool GetCurrentlyInstalledSDKString(string PlatformSDKRoot, out string OutInstalledSDKVersionString)
-        {
-            if (Directory.Exists(PlatformSDKRoot))
-            {
-                string VersionFilename = Path.Combine(PlatformSDKRoot, CurrentlyInstalledSDKStringManifest);
-                if (File.Exists(VersionFilename))
-                {
-                    using (StreamReader Reader = new StreamReader(VersionFilename))
-                    {
-                        string Version = Reader.ReadLine();
-                        if (Version != null)
-                        {
-                            OutInstalledSDKVersionString = Version;
-                            return true;
-                        }
-                    }
-                }
-            }
-
-            OutInstalledSDKVersionString = "";
-            return false;
-        }
-
-        /**
-         * Gets the version of the last successfully run setup script.
-         * 
-         * @param PlatformSDKRoot absolute path to platform SDK root
-         * @param OutLastRunScriptVersion version string
-         * 
-         * @return true if was able to read it
-         */
-        public virtual bool GetLastRunScriptVersionString(string PlatformSDKRoot, out string OutLastRunScriptVersion)
-        {
-            if (Directory.Exists(PlatformSDKRoot))
-            {
-                string VersionFilename = Path.Combine(PlatformSDKRoot, LastRunScriptVersionManifest);
-                if (File.Exists(VersionFilename))
-                {
-                    using (StreamReader Reader = new StreamReader(VersionFilename))
-                    {
-                        string Version = Reader.ReadLine();
-                        if (Version != null)
-                        {
-                            OutLastRunScriptVersion = Version;
-                            return true;
-                        }
-                    }
-                }
-            }
-
-            OutLastRunScriptVersion = "";
-            return false;
-        }
-       
-
-        /**
-         * Sets currently installed version
-         * 
-         * @param PlatformSDKRoot absolute path to platform SDK root
-         * @param InstalledSDKVersionString SDK version string to set
-         * 
-         * @return true if was able to set it
-         */
-        public virtual bool SetCurrentlyInstalledSDKString(string PlatformSDKRoot, string InstalledSDKVersionString)
-        {
-            if (Directory.Exists(PlatformSDKRoot))
-            {
-                string VersionFilename = Path.Combine(PlatformSDKRoot, CurrentlyInstalledSDKStringManifest);
-                if (File.Exists(VersionFilename))
-                {
-                    File.Delete(VersionFilename);
-                }
-
-                using (StreamWriter Writer = File.CreateText(VersionFilename))
-                {
-                    Writer.WriteLine(InstalledSDKVersionString);
-                    return true;
-                }
-            }
-
-            return false;
-        }
-
-        public virtual bool SetLastRunScriptVersion(string PlatformSDKRoot, string LastRunScriptVersion)
-        {
-            if (Directory.Exists(PlatformSDKRoot))
-            {
-                string VersionFilename = Path.Combine(PlatformSDKRoot, LastRunScriptVersionManifest);
-                if (File.Exists(VersionFilename))
-                {
-                    File.Delete(VersionFilename);
-                }
-
-                using (StreamWriter Writer = File.CreateText(VersionFilename))
-                {
-                    Writer.WriteLine(LastRunScriptVersion);
-                    return true;
-                }
-            }
-
-            return false;
-        }
-
-        /**
-         * Returns Hook names as needed by the platform
-         * (e.g. can be overriden with custom executables or scripts)
-         *
-         * @param Hook Hook type
-         */
-        public virtual string GetHookExecutableName(SDKHookType Hook)
-        {
-            if (Hook == SDKHookType.Uninstall)
-            {
-                return "unsetup.bat";
-            }
-
-            return "setup.bat";
-        }
-
-        /**
-         * Runs install/uninstall hooks for SDK
-         * 
-         * @param PlatformSDKRoot absolute path to platform SDK root
-         * @param SDKVersionString version string to run for (can be empty!)
-         * @param Hook which one of hooks to run
-         * @param bHookCanBeNonExistent whether a non-existing hook means failure
-         * 
-         * @return true if succeeded
-         */
-        public virtual bool RunSDKHooks(string PlatformSDKRoot, string SDKVersionString, SDKHookType Hook, bool bHookCanBeNonExistent = true)
-        {
-            if (SDKVersionString != "")
-            {
-                string SDKDirectory = Path.Combine(PlatformSDKRoot, SDKVersionString);
-                string HookExe = Path.Combine(SDKDirectory, GetHookExecutableName(Hook));
-
-                if (File.Exists(HookExe))
-                {
-                    Console.WriteLine("Running {0} hook {1}", Hook, HookExe);
-
-                    // run it
-                    Process HookProcess = new Process();
-                    HookProcess.StartInfo.WorkingDirectory = SDKDirectory;
-                    HookProcess.StartInfo.FileName = HookExe;
-                    HookProcess.StartInfo.Arguments = "";
-                    HookProcess.StartInfo.WindowStyle = ProcessWindowStyle.Hidden;
-
-                    //installers may require administrator access to succeed. so run as an admmin.
-					HookProcess.StartInfo.Verb = "runas";
-                    HookProcess.Start();
-                    HookProcess.WaitForExit();
-
-                    if (HookProcess.ExitCode != 0)
-                    {
-                        Console.WriteLine("Hook exited uncleanly (returned {0}), considering it failed.", HookProcess.ExitCode);
-                        return false;
-                    }
-
-                    return true;
-                }
-            }
-
-            return bHookCanBeNonExistent;
-        }
-
-        /**
          * Returns the delimiter used to separate paths in the PATH environment variable for the platform we are executing on.
@@ -661,210 +408,9 @@
             }
         }
 
-        /**
-         * Loads environment variables from SDK
-         * 
-         * @param PlatformSDKRoot absolute path to platform SDK
-         * @param SDKVersionString SDK version string (cannot be empty!)
-         * 
-         * @return true if succeeded
-         */
-        public virtual bool SetupEnvironmentFromSDK(string PlatformSDKRoot, string SDKVersionString)
-        {
-            string EnvVarFile = Path.Combine(PlatformSDKRoot, SDKVersionString, SDKEnvironmentVarsFile);
-            if (File.Exists(EnvVarFile))
-            {
-                using (StreamReader Reader = new StreamReader(EnvVarFile))
-                {
-                    List<string> PathAdds = new List<string>();
-                    List<string> PathRemoves = new List<string>();
-
-                    List<string> EnvVarNames = new List<string>();
-                    List<string> EnvVarValues = new List<string>();
-                    for (; ;)
-                    {
-                        string VariableString = Reader.ReadLine();
-                        if (VariableString == null)
-                        {
-                            break;
-                        }
-						
-                        string[] Parts = VariableString.Split('=');
-                        if (Parts.Length != 2)
-                        {
-                            Console.WriteLine("Incorrect environment variable declaration:");
-                            Console.WriteLine(VariableString);
-                            return false;
-                        }
-
-                        if (String.Compare(Parts[0], "strippath", true) == 0)
-                        {
-                            PathRemoves.Add(Parts[1]);
-                        }
-                        else if (String.Compare(Parts[0], "addpath", true) == 0)
-                        {
-                            PathAdds.Add(Parts[1]);
-                        }
-                        else
-                        {
-                            // convenience for setup.bat writers.  Trim any accidental whitespace from var names/values.
-                            EnvVarNames.Add(Parts[0].Trim());
-                            EnvVarValues.Add(Parts[1].Trim());
-                        }
-                    }
-
-                    // don't actually set anything until we successfully validate and read all values in.
-                    // we don't want to set a few vars, return a failure, and then have a platform try to
-                    // build against a manually installed SDK with half-set env vars.
-                    for (int i = 0; i < EnvVarNames.Count; ++i)
-                    {
-                        string EnvVarName = EnvVarNames[i];
-                        string EnvVarValue = EnvVarValues[i];                        
-                        if (BuildConfiguration.bPrintDebugInfo)
-                        {
-                            Console.WriteLine("Setting variable '{0}' to '{1}'", EnvVarName, EnvVarValue);
-                        }
-                        Environment.SetEnvironmentVariable(EnvVarName, EnvVarValue);
-                    }
-                       
-                    
-                    // actually perform the PATH stripping / adding.
-                    String OrigPathVar = Environment.GetEnvironmentVariable("PATH");
-                    String PathDelimiter = GetPathVarDelimiter();
-                    String[] PathVars = OrigPathVar.Split(PathDelimiter.ToCharArray());
-
-                    List<String> ModifiedPathVars = new List<string>();
-                    ModifiedPathVars.AddRange(PathVars);
-
-                    // perform removes first, in case they overlap with any adds.
-                    foreach (String PathRemove in PathRemoves)
-                    {
-                        foreach (String PathVar in PathVars)
-                        {
-                            if (PathVar.IndexOf(PathRemove, StringComparison.OrdinalIgnoreCase) >= 0)
-                            {
-                                if (BuildConfiguration.bPrintDebugInfo)
-                                {
-                                    Console.WriteLine("Removing Path: '{0}'", PathVar);
-                                }
-                                ModifiedPathVars.Remove(PathVar);
-                            }
-                        }
-                    }
-
-                    // remove all the of ADDs so that if this function is executed multiple times, the paths will be guarateed to be in the same order after each run.
-                    // If we did not do this, a 'remove' that matched some, but not all, of our 'adds' would cause the order to change.
-                    foreach (String PathAdd in PathAdds)
-                    {
-                        foreach (String PathVar in PathVars)
-                        {
-                            if (String.Compare(PathAdd, PathVar, true) == 0)
-                            {
-                                if (BuildConfiguration.bPrintDebugInfo)
-                                {
-                                    Console.WriteLine("Removing Path: '{0}'", PathVar);
-                                }
-                                ModifiedPathVars.Remove(PathVar);
-                            }
-                        }
-                    }
-
-                    // perform adds, but don't add duplicates
-                    foreach (String PathAdd in PathAdds)
-                    {
-                        if (!ModifiedPathVars.Contains(PathAdd))
-                        {
-                            if (BuildConfiguration.bPrintDebugInfo)
-                            {
-                                Console.WriteLine("Adding Path: '{0}'", PathAdd);
-                            }
-                            ModifiedPathVars.Add(PathAdd);
-                        }
-                    }
-
-                    String ModifiedPath = String.Join(PathDelimiter, ModifiedPathVars);
-                    Environment.SetEnvironmentVariable("PATH", ModifiedPath);
-
-                    return true;
-                }
-            }
-
-            return false;
-        }
-
-        public void InvalidateCurrentlyInstalled(String PlatformSDKRoot)
-        {
-            SetCurrentlyInstalledSDKString(PlatformSDKRoot, "");
-            SetLastRunScriptVersion(PlatformSDKRoot, "");
-        }
+        
 
         /**
-         *	Whether the required external SDKs are installed for this platform
-         */
-		public virtual SDKStatus HasRequiredSDKsInstalled()
-		{			
-            if (PlatformSupportsSDKSwitching())
-            {				
-                // if we don't have a path to platform SDK's, take no action
-                bool bSDKPathSet;
-                string PlatformSDKRoot = GetPathToPlatformSDKs(out bSDKPathSet);				
-                if (bSDKPathSet)
-                {					
-                    if (PlatformSDKRoot != "")
-                    {
-                        // check script version so script fixes can be propagated without touching every build machine's CurrentlyInstalled file manually.
-                        bool bScriptVersionMatches = false;
-                        string CurrentScriptVersionString;
-                        if (GetLastRunScriptVersionString(PlatformSDKRoot, out CurrentScriptVersionString) && CurrentScriptVersionString == GetRequiredScriptVersionString())
-                        {
-                            bScriptVersionMatches = true;
-                        }
-
-                        string CurrentSDKString;
-                        if (!GetCurrentlyInstalledSDKString(PlatformSDKRoot, out CurrentSDKString) || CurrentSDKString != GetRequiredSDKString() || !bScriptVersionMatches)
-                        {
-                        
-                            // switch over (note that version string can be empty)
-                            if (!RunSDKHooks(PlatformSDKRoot, CurrentSDKString, SDKHookType.Uninstall))
-                            {
-                                Console.WriteLine("Failed to uninstall currently installed SDK {0}", CurrentSDKString);
-                                InvalidateCurrentlyInstalled(PlatformSDKRoot);
-                                return SDKStatus.Invalid;
-                            }
-                            // delete Manifest file to avoid multiple uninstalls
-                            InvalidateCurrentlyInstalled(PlatformSDKRoot);
-
-                            if (!RunSDKHooks(PlatformSDKRoot, GetRequiredSDKString(), SDKHookType.Install, false))
-                            {								
-                                Console.WriteLine("Failed to install required SDK {0}.  Attemping to uninstall", GetRequiredSDKString());
-                                RunSDKHooks(PlatformSDKRoot, GetRequiredSDKString(), SDKHookType.Uninstall, false);
-
-                                return SDKStatus.Invalid;
-                            }
-                            SetCurrentlyInstalledSDKString(PlatformSDKRoot, GetRequiredSDKString());
-                            SetLastRunScriptVersion(PlatformSDKRoot, GetRequiredScriptVersionString());
-                        }
-
-                        // load environment variables from current SDK
-                        if (!SetupEnvironmentFromSDK(PlatformSDKRoot, GetRequiredSDKString()))
-                        {
-                            Console.WriteLine("Failed to load environment from required SDK {0}", GetRequiredSDKString());
-                            InvalidateCurrentlyInstalled(PlatformSDKRoot);
-                            return SDKStatus.Invalid;
-                        }                        
-                        return SDKStatus.Valid;
-                    }
-                    else
-                    {
-                        Console.WriteLine("Failed to install required SDK {0}. {1} is blank.", GetRequiredSDKString(), SDKRootEnvVar);
-                        return SDKStatus.Invalid;
-                    }
-                }
-            }
-            return SDKStatus.Unknown;
-        }
-
-        /**
          *	If this platform can be compiled with XGE
          */
         public virtual bool CanUseXGE()
@@ -875,7 +421,13 @@
 		/**
 		 *	Register the platform with the UEBuildPlatform class
 		 */
-		public abstract void RegisterBuildPlatform();
+        public void RegisterBuildPlatform()
+        {
+            ManageAndValidateSDK();
+            RegisterBuildPlatformInternal();
+        }
+
+        protected abstract void RegisterBuildPlatformInternal();
 
 		/**
 		 *	Retrieve the CPPTargetPlatform for the given UnrealTargetPlatform
@@ -1253,4 +805,636 @@
 		{
 		}
 	}
+
+    // AutoSDKs handling portion
+    public abstract partial class UEBuildPlatform : IUEBuildPlatform
+    {
+
+        #region protected AutoSDKs Utility
+
+        /** Name of the file that holds currently install SDK version string */
+        protected static string CurrentlyInstalledSDKStringManifest = "CurrentlyInstalled.txt";
+
+        /** name of the file that holds the last succesfully run SDK setup script version */
+        protected static string LastRunScriptVersionManifest = "CurrentlyInstalled.Version.txt";
+
+        /** Name of the file that holds environment variables of current SDK */
+        protected static string SDKEnvironmentVarsFile = "OutputEnvVars.txt";
+
+        protected static readonly string SDKRootEnvVar = "UE_SDKS_ROOT";
+
+        /** 
+         * Whether platform supports switching SDKs during runtime
+         * 
+         * @return true if supports
+         */
+        protected virtual bool PlatformSupportsAutoSDKs()
+        {
+            return false;
+        }
+
+        static private bool bCheckedAutoSDKRootEnvVar = false;
+        static private bool bAutoSDKSystemEnabled = false;
+        static private bool HasAutoSDKSystemEnabled()
+        {
+            if (!bCheckedAutoSDKRootEnvVar)
+            {
+                string SDKRoot = Environment.GetEnvironmentVariable(SDKRootEnvVar);
+                if (SDKRoot != null)
+                {
+                    bAutoSDKSystemEnabled = true;
+                }
+                bCheckedAutoSDKRootEnvVar = true;
+            }
+            return bAutoSDKSystemEnabled;
+        }
+
+        // Whether AutoSDK setup is safe. AutoSDKs will damage manual installs on some platforms.
+        protected bool IsAutoSDKSafe()
+        {
+            return !IsAutoSDKDestructive() || !HasAnyManualInstall();
+        }
+
+        /** 
+         * Returns SDK string as required by the platform 
+         * 
+         * @return Valid SDK string
+         */
+        protected virtual string GetRequiredSDKString()
+        {
+            return "";
+        }
+
+        /**
+        * Gets the version number of the SDK setup script itself.  The version in the base should ALWAYS be 1.0.  If you need to force a rebuild for a given platform, override this for the given platform.
+        * @return Setup script version
+        */
+        protected virtual String GetRequiredScriptVersionString()
+        {
+            return "1.0";
+        }
+
+        /** 
+         * Returns path to platform SDKs
+         * 
+         * @return Valid SDK string
+         */
+        protected string GetPathToPlatformAutoSDKs()
+        {
+            string SDKPath = "";            
+            string SDKRoot = Environment.GetEnvironmentVariable(SDKRootEnvVar);
+            if (SDKRoot != null)
+            {                
+                if (SDKRoot != "")
+                {
+                    SDKPath = Path.Combine(SDKRoot, "Host" + ExternalExecution.GetRuntimePlatform(), GetSDKTargetPlatformName());
+                }
+            }
+            return SDKPath;
+        }
+
+        /**
+         * Gets currently installed version
+         * 
+         * @param PlatformSDKRoot absolute path to platform SDK root
+         * @param OutInstalledSDKVersionString version string as currently installed
+         * 
+         * @return true if was able to read it
+         */
+        protected bool GetCurrentlyInstalledSDKString(string PlatformSDKRoot, out string OutInstalledSDKVersionString)
+        {
+            if (Directory.Exists(PlatformSDKRoot))
+            {
+                string VersionFilename = Path.Combine(PlatformSDKRoot, CurrentlyInstalledSDKStringManifest);
+                if (File.Exists(VersionFilename))
+                {
+                    using (StreamReader Reader = new StreamReader(VersionFilename))
+                    {
+                        string Version = Reader.ReadLine();
+                        if (Version != null)
+                        {
+                            OutInstalledSDKVersionString = Version;
+                            return true;
+                        }
+                    }
+                }
+            }
+
+            OutInstalledSDKVersionString = "";
+            return false;
+        }
+
+        /**
+         * Gets the version of the last successfully run setup script.
+         * 
+         * @param PlatformSDKRoot absolute path to platform SDK root
+         * @param OutLastRunScriptVersion version string
+         * 
+         * @return true if was able to read it
+         */
+        protected bool GetLastRunScriptVersionString(string PlatformSDKRoot, out string OutLastRunScriptVersion)
+        {
+            if (Directory.Exists(PlatformSDKRoot))
+            {
+                string VersionFilename = Path.Combine(PlatformSDKRoot, LastRunScriptVersionManifest);
+                if (File.Exists(VersionFilename))
+                {
+                    using (StreamReader Reader = new StreamReader(VersionFilename))
+                    {
+                        string Version = Reader.ReadLine();
+                        if (Version != null)
+                        {
+                            OutLastRunScriptVersion = Version;
+                            return true;
+                        }
+                    }
+                }
+            }
+
+            OutLastRunScriptVersion = "";
+            return false;
+        }
+
+        /**
+         * Sets currently installed version
+         * 
+         * @param PlatformSDKRoot absolute path to platform SDK root
+         * @param InstalledSDKVersionString SDK version string to set
+         * 
+         * @return true if was able to set it
+         */
+        protected bool SetCurrentlyInstalledAutoSDKString(string PlatformSDKRoot, string InstalledSDKVersionString)
+        {
+            if (Directory.Exists(PlatformSDKRoot))
+            {
+                string VersionFilename = Path.Combine(PlatformSDKRoot, CurrentlyInstalledSDKStringManifest);
+                if (File.Exists(VersionFilename))
+                {
+                    File.Delete(VersionFilename);
+                }
+
+                using (StreamWriter Writer = File.CreateText(VersionFilename))
+                {
+                    Writer.WriteLine(InstalledSDKVersionString);
+                    return true;
+                }
+            }
+
+            return false;
+        }
+
+        protected bool SetLastRunAutoSDKScriptVersion(string PlatformSDKRoot, string LastRunScriptVersion)
+        {
+            if (Directory.Exists(PlatformSDKRoot))
+            {
+                string VersionFilename = Path.Combine(PlatformSDKRoot, LastRunScriptVersionManifest);
+                if (File.Exists(VersionFilename))
+                {
+                    File.Delete(VersionFilename);
+                }
+
+                using (StreamWriter Writer = File.CreateText(VersionFilename))
+                {
+                    Writer.WriteLine(LastRunScriptVersion);
+                    return true;
+                }
+            }
+
+            return false;
+        }
+
+        /**
+        * Returns Hook names as needed by the platform
+        * (e.g. can be overriden with custom executables or scripts)
+        *
+        * @param Hook Hook type
+        */
+        protected virtual string GetHookExecutableName(SDKHookType Hook)
+        {
+            if (Hook == SDKHookType.Uninstall)
+            {
+                return "unsetup.bat";
+            }
+
+            return "setup.bat";
+        }
+
+        /**
+         * Runs install/uninstall hooks for SDK
+         * 
+         * @param PlatformSDKRoot absolute path to platform SDK root
+         * @param SDKVersionString version string to run for (can be empty!)
+         * @param Hook which one of hooks to run
+         * @param bHookCanBeNonExistent whether a non-existing hook means failure
+         * 
+         * @return true if succeeded
+         */
+        protected virtual bool RunAutoSDKHooks(string PlatformSDKRoot, string SDKVersionString, SDKHookType Hook, bool bHookCanBeNonExistent = true)
+        {
+            if (!IsAutoSDKSafe())
+            {
+                Console.ForegroundColor = ConsoleColor.Red;
+                Console.WriteLine(GetSDKTargetPlatformName() + " attempted to run SDK hook which could have damaged manual SDK install!");
+                Console.ResetColor();
+
+                return false;
+            }
+            if (SDKVersionString != "")
+            {
+                string SDKDirectory = Path.Combine(PlatformSDKRoot, SDKVersionString);
+                string HookExe = Path.Combine(SDKDirectory, GetHookExecutableName(Hook));
+
+                if (File.Exists(HookExe))
+                {
+                    Console.WriteLine("Running {0} hook {1}", Hook, HookExe);
+
+                    // run it
+                    Process HookProcess = new Process();
+                    HookProcess.StartInfo.WorkingDirectory = SDKDirectory;
+                    HookProcess.StartInfo.FileName = HookExe;
+                    HookProcess.StartInfo.Arguments = "";
+                    HookProcess.StartInfo.WindowStyle = ProcessWindowStyle.Hidden;
+
+                    //installers may require administrator access to succeed. so run as an admmin.
+                    HookProcess.StartInfo.Verb = "runas";
+                    HookProcess.Start();
+                    HookProcess.WaitForExit();
+
+                    if (HookProcess.ExitCode != 0)
+                    {
+                        Console.WriteLine("Hook exited uncleanly (returned {0}), considering it failed.", HookProcess.ExitCode);
+                        return false;
+                    }
+
+                    return true;
+                }
+            }
+
+            return bHookCanBeNonExistent;
+        }
+
+        /**
+         * Loads environment variables from SDK
+         * 
+         * @param PlatformSDKRoot absolute path to platform SDK         
+         * 
+         * @return true if succeeded
+         */
+        protected bool SetupEnvironmentFromAutoSDK(string PlatformSDKRoot)
+        {
+            string EnvVarFile = Path.Combine(PlatformSDKRoot, SDKEnvironmentVarsFile);
+            if (File.Exists(EnvVarFile))
+            {
+                using (StreamReader Reader = new StreamReader(EnvVarFile))
+                {
+                    List<string> PathAdds = new List<string>();
+                    List<string> PathRemoves = new List<string>();
+
+                    List<string> EnvVarNames = new List<string>();
+                    List<string> EnvVarValues = new List<string>();
+                    for (; ; )
+                    {
+                        string VariableString = Reader.ReadLine();
+                        if (VariableString == null)
+                        {
+                            break;
+                        }
+
+                        string[] Parts = VariableString.Split('=');
+                        if (Parts.Length != 2)
+                        {
+                            Console.WriteLine("Incorrect environment variable declaration:");
+                            Console.WriteLine(VariableString);
+                            return false;
+                        }
+
+                        if (String.Compare(Parts[0], "strippath", true) == 0)
+                        {
+                            PathRemoves.Add(Parts[1]);
+                        }
+                        else if (String.Compare(Parts[0], "addpath", true) == 0)
+                        {
+                            PathAdds.Add(Parts[1]);
+                        }
+                        else
+                        {
+                            // convenience for setup.bat writers.  Trim any accidental whitespace from var names/values.
+                            EnvVarNames.Add(Parts[0].Trim());
+                            EnvVarValues.Add(Parts[1].Trim());
+                        }
+                    }
+
+                    // don't actually set anything until we successfully validate and read all values in.
+                    // we don't want to set a few vars, return a failure, and then have a platform try to
+                    // build against a manually installed SDK with half-set env vars.
+                    for (int i = 0; i < EnvVarNames.Count; ++i)
+                    {
+                        string EnvVarName = EnvVarNames[i];
+                        string EnvVarValue = EnvVarValues[i];
+                        if (BuildConfiguration.bPrintDebugInfo)
+                        {
+                            Console.WriteLine("Setting variable '{0}' to '{1}'", EnvVarName, EnvVarValue);
+                        }
+                        Environment.SetEnvironmentVariable(EnvVarName, EnvVarValue);
+                    }
+
+
+                    // actually perform the PATH stripping / adding.
+                    String OrigPathVar = Environment.GetEnvironmentVariable("PATH");
+                    String PathDelimiter = GetPathVarDelimiter();
+                    String[] PathVars = OrigPathVar.Split(PathDelimiter.ToCharArray());
+
+                    List<String> ModifiedPathVars = new List<string>();
+                    ModifiedPathVars.AddRange(PathVars);
+
+                    // perform removes first, in case they overlap with any adds.
+                    foreach (String PathRemove in PathRemoves)
+                    {
+                        foreach (String PathVar in PathVars)
+                        {
+                            if (PathVar.IndexOf(PathRemove, StringComparison.OrdinalIgnoreCase) >= 0)
+                            {
+                                if (BuildConfiguration.bPrintDebugInfo)
+                                {
+                                    Console.WriteLine("Removing Path: '{0}'", PathVar);
+                                }
+                                ModifiedPathVars.Remove(PathVar);
+                            }
+                        }
+                    }
+
+                    // remove all the of ADDs so that if this function is executed multiple times, the paths will be guarateed to be in the same order after each run.
+                    // If we did not do this, a 'remove' that matched some, but not all, of our 'adds' would cause the order to change.
+                    foreach (String PathAdd in PathAdds)
+                    {
+                        foreach (String PathVar in PathVars)
+                        {
+                            if (String.Compare(PathAdd, PathVar, true) == 0)
+                            {
+                                if (BuildConfiguration.bPrintDebugInfo)
+                                {
+                                    Console.WriteLine("Removing Path: '{0}'", PathVar);
+                                }
+                                ModifiedPathVars.Remove(PathVar);
+                            }
+                        }
+                    }
+
+                    // perform adds, but don't add duplicates
+                    foreach (String PathAdd in PathAdds)
+                    {
+                        if (!ModifiedPathVars.Contains(PathAdd))
+                        {
+                            if (BuildConfiguration.bPrintDebugInfo)
+                            {
+                                Console.WriteLine("Adding Path: '{0}'", PathAdd);
+                            }
+                            ModifiedPathVars.Add(PathAdd);
+                        }
+                    }
+
+                    String ModifiedPath = String.Join(PathDelimiter, ModifiedPathVars);
+                    Environment.SetEnvironmentVariable("PATH", ModifiedPath);
+
+                    // make sure we know that we've modified the local environment, invalidating manual installs for this run.
+                    bProcessEnvSetupAutoSDK = true;
+
+                    return true;
+                }
+            }
+
+            return false;
+        }
+
+        protected void InvalidateCurrentlyInstalledAutoSDK(String PlatformSDKRoot)
+        {
+            SetCurrentlyInstalledAutoSDKString(PlatformSDKRoot, "");
+            SetLastRunAutoSDKScriptVersion(PlatformSDKRoot, "");
+        }
+
+        /** 
+         * Currently installed AutoSDK is written out to a text file in a known location.  
+         * This function just compares the file's contents with the current requirements.
+         */
+        protected SDKStatus HasRequiredAutoSDKInstalled()
+        {
+            if (PlatformSupportsAutoSDKs() && HasAutoSDKSystemEnabled())
+            {                
+                string AutoSDKRoot = GetPathToPlatformAutoSDKs();                
+                if (AutoSDKRoot != "")
+                {
+                    // check script version so script fixes can be propagated without touching every build machine's CurrentlyInstalled file manually.
+                    bool bScriptVersionMatches = false;
+                    string CurrentScriptVersionString;
+                    if (GetLastRunScriptVersionString(AutoSDKRoot, out CurrentScriptVersionString) && CurrentScriptVersionString == GetRequiredScriptVersionString())
+                    {
+                        bScriptVersionMatches = true;
+                    }
+
+                    string CurrentSDKString;
+                    if (GetCurrentlyInstalledSDKString(AutoSDKRoot, out CurrentSDKString) && CurrentSDKString == GetRequiredSDKString() && bScriptVersionMatches)
+                    {
+                        return SDKStatus.Valid;
+                    }
+                    return SDKStatus.Invalid;
+                }                
+            }
+            return SDKStatus.Invalid;
+        }
+
+        // This tracks if we have already checked the sdk installation.
+        private Int32 SDKCheckStatus = -1;
+
+        // true if we've ever overridden the process's environment with AutoSDK data.  After that, manual installs cannot be considered valid ever again.
+        private bool bProcessEnvSetupAutoSDK = false;
+
+        protected bool HasSetupAutoSDK()
+        {
+            return bProcessEnvSetupAutoSDK;
+        }
+
+        protected SDKStatus HasRequiredManualSDK()
+        {
+            if (bProcessEnvSetupAutoSDK)
+            {
+                return SDKStatus.Invalid;
+            }
+
+            // manual installs are always invalid if we have modified the process's environment for AutoSDKs
+            return HasRequiredManualSDKInternal();
+        }
+
+        // for platforms with destructive AutoSDK.  Report if any manual sdk is installed that may be damaged by an autosdk.
+        protected virtual bool HasAnyManualInstall()
+        {
+            return false;
+        }
+
+        // tells us if the user has a valid manual install.
+        protected abstract SDKStatus HasRequiredManualSDKInternal();
+
+        // some platforms will fail if there is a manual install that is the WRONG manual install.
+        protected virtual bool AllowInvalidManualInstall()
+        {
+            return true;
+        }
+
+        // platforms can choose if they prefer a correct the the AutoSDK install over the manual install.
+        protected virtual bool PreferAutoSDK()
+        {
+            return false;
+        }
+
+        // some platforms don't support parallel SDK installs.  AutoSDK on these platforms will
+        // actively damage an existing manual install by overwriting files in it.  AutoSDK must NOT
+        // run any setup if a manual install exists in this case.
+        protected virtual bool IsAutoSDKDestructive()
+        {
+            return false;
+        }
+
+        /**
+        * Runs batch files if necessary to set up required AutoSDK.
+        * AutoSDKs are SDKs that have not been setup through a formal installer, but rather come from
+        * a source control directory, or other local copy.
+        */
+        private void SetupAutoSDK()
+        {
+            if (IsAutoSDKSafe() && PlatformSupportsAutoSDKs() && HasAutoSDKSystemEnabled())
+            {
+                // run installation for autosdk if necessary.
+                if (HasRequiredAutoSDKInstalled() == SDKStatus.Invalid)
+                {
+                    //reset check status so any checking sdk status after the attempted setup will do a real check again.
+                    SDKCheckStatus = -1;
+
+                    string AutoSDKRoot = GetPathToPlatformAutoSDKs();
+                    string CurrentSDKString;
+                    GetCurrentlyInstalledSDKString(AutoSDKRoot, out CurrentSDKString);
+
+                    // switch over (note that version string can be empty)
+                    if (!RunAutoSDKHooks(AutoSDKRoot, CurrentSDKString, SDKHookType.Uninstall))
+                    {
+                        Console.WriteLine("Failed to uninstall currently installed SDK {0}", CurrentSDKString);
+                        InvalidateCurrentlyInstalledAutoSDK(AutoSDKRoot);
+                        return;
+                    }
+                    // delete Manifest file to avoid multiple uninstalls
+                    InvalidateCurrentlyInstalledAutoSDK(AutoSDKRoot);
+
+                    if (!RunAutoSDKHooks(AutoSDKRoot, GetRequiredSDKString(), SDKHookType.Install, false))
+                    {
+                        Console.WriteLine("Failed to install required SDK {0}.  Attemping to uninstall", GetRequiredSDKString());
+                        RunAutoSDKHooks(AutoSDKRoot, GetRequiredSDKString(), SDKHookType.Uninstall, false);
+                        return;
+                    }
+                    SetCurrentlyInstalledAutoSDKString(AutoSDKRoot, GetRequiredSDKString());
+                    SetLastRunAutoSDKScriptVersion(AutoSDKRoot, GetRequiredScriptVersionString());
+                }
+
+                // fixup process environment to match autosdk
+                SetupEnvironmentFromAutoSDK();                
+            }
+        }        
+
+        #endregion
+
+        #region public AutoSDKs Utility
+
+        /** Enum describing types of hooks a platform SDK can have */
+        public enum SDKHookType
+        {
+            Install,
+            Uninstall
+        };
+
+        /** 
+         * Returns platform-specific name used in SDK repository
+         * 
+         * @return path to SDK Repository
+         */
+        public virtual string GetSDKTargetPlatformName()
+        {
+            return "";
+        }  
+
+        /* Whether or not we should try to automatically switch SDKs when asked to validate the platform's SDK state. */
+        public static bool bAllowAutoSDKSwitching = true;               
+
+        public SDKStatus SetupEnvironmentFromAutoSDK()
+        {            
+            string PlatformSDKRoot = GetPathToPlatformAutoSDKs();
+
+            // load environment variables from current SDK
+            if (!SetupEnvironmentFromAutoSDK(PlatformSDKRoot))
+            {
+                Console.WriteLine("Failed to load environment from required SDK {0}", GetRequiredSDKString());
+                InvalidateCurrentlyInstalledAutoSDK(PlatformSDKRoot);
+                return SDKStatus.Invalid;
+            }            
+            return SDKStatus.Valid;
+        }
+
+        /**
+         *	Whether the required external SDKs are installed for this platform.  
+         *	Could be either a manual install or an AutoSDK.
+         */
+        public SDKStatus HasRequiredSDKsInstalled()
+        {
+            // avoid redundant potentially expensive SDK checks.
+            if (SDKCheckStatus == -1)
+            {                
+                bool bHasManualSDK = HasRequiredManualSDK() == SDKStatus.Valid;
+                bool bHasAutoSDK = HasRequiredAutoSDKInstalled() == SDKStatus.Valid;
+                
+                // Per-Platform implementations can choose how to handle non-Auto SDK detection / handling.
+                SDKCheckStatus = (bHasManualSDK || bHasAutoSDK) ? 1 : 0;
+            }
+            return SDKCheckStatus == 1 ? SDKStatus.Valid : SDKStatus.Invalid;
+        }
+
+        // Arbitrates between manual SDKs and setting up AutoSDK based on program options and platform preferences.
+        public void ManageAndValidateSDK()
+        {
+            bool bHasRequiredManualSDK = HasRequiredManualSDK() == SDKStatus.Valid;
+            if (bAllowAutoSDKSwitching && IsAutoSDKSafe() && (PreferAutoSDK() || !bHasRequiredManualSDK))
+            {
+                SetupAutoSDK();
+            }
+
+            if (BuildConfiguration.bPrintDebugInfo)
+            {
+                PrintSDKInfo();
+            }
+        }
+
+        public void PrintSDKInfo()
+        {
+            
+            
+            if (HasRequiredSDKsInstalled() == SDKStatus.Valid)
+            {
+                bool bHasRequiredManualSDK = HasRequiredManualSDK() == SDKStatus.Valid;
+                if (HasSetupAutoSDK())
+                {
+                    string PlatformSDKRoot = GetPathToPlatformAutoSDKs();
+                    Console.WriteLine(GetSDKTargetPlatformName() + " using SDK from: " + Path.Combine(PlatformSDKRoot, GetRequiredSDKString()));
+                }
+                else if (bHasRequiredManualSDK)
+                {
+                    Console.WriteLine(this.ToString() + " using manually installed SDK " + GetRequiredSDKString());
+                }
+                else
+                {
+                    Console.WriteLine(this.ToString() + " setup error.  Inform platform team.");
+                }
+            }
+            else
+            {
+                Console.WriteLine(this.ToString() + " has no valid SDK");
+            }
+        }
+
+        #endregion
+    }
 }
diff --git a/Engine/Source/Programs/UnrealBuildTool/HTML5/UEBuildHTML5.cs b/Engine/Source/Programs/UnrealBuildTool/HTML5/UEBuildHTML5.cs
--- a/Engine/Source/Programs/UnrealBuildTool/HTML5/UEBuildHTML5.cs
+++ b/Engine/Source/Programs/UnrealBuildTool/HTML5/UEBuildHTML5.cs
@@ -8,11 +8,11 @@ using System.IO;
 
 namespace UnrealBuildTool
 {
-	class HTML5Platform : UEBuildPlatform
-	{
+    class HTML5Platform : UEBuildPlatform
+    {
         // use -win32 for win32 builds. ( build html5 platform as a win32 binary for debugging )
-		[XmlConfig]
-		public static string HTML5Architecture = "";
+        [XmlConfig]
+        public static string HTML5Architecture = "";
 
         // Enable and link in code for serving HTTP request via NFS. 
         [XmlConfig]
@@ -22,65 +22,72 @@ namespace UnrealBuildTool
         [XmlConfig]
         public static bool UseHTTPForNFS = false;
 
+        static private string TargetPlatformName = "HTML5";
+
+
+        public override string GetSDKTargetPlatformName()
+        {
+            return TargetPlatformName;
+        }
+
+        // The current architecture - affects everything about how UBT operates on HTML5
+        public override string GetActiveArchitecture()
+        {
+            // by default, use an empty architecture (which is really just a modifier to the platform for some paths/names)
+            return HTML5Architecture;
+        }
+
+        // F5 should always try to run the Win32 version
+        public override string ModifyNMakeOutput(string ExeName)
+        {
+            // nmake Run should always run the win32 version
+            return Path.ChangeExtension(ExeName+"-win32", ".exe");
+        }
 
-		// The current architecture - affects everything about how UBT operates on HTML5
-		public override string GetActiveArchitecture()
-		{
-			// by default, use an empty architecture (which is really just a modifier to the platform for some paths/names)
-			return HTML5Architecture;
-		}
-
-		// F5 should always try to run the Win32 version
-		public override string ModifyNMakeOutput(string ExeName)
-		{
-			// nmake Run should always run the win32 version
-			return Path.ChangeExtension(ExeName+"-win32", ".exe");
-		}
-
-		/**
-		 *	Whether the required external SDKs are installed for this platform
-		 */
-		public override SDKStatus HasRequiredSDKsInstalled()
-		{
-			string BaseSDKPath = Environment.GetEnvironmentVariable("EMSCRIPTEN");
-			return (string.IsNullOrEmpty(BaseSDKPath) == false) ? SDKStatus.Valid : SDKStatus.Invalid;
-		}
+        /**
+         *	Whether the required external SDKs are installed for this platform
+         */
+        protected override SDKStatus HasRequiredManualSDKInternal()
+        {
+            string BaseSDKPath = Environment.GetEnvironmentVariable("EMSCRIPTEN");
+            return (string.IsNullOrEmpty(BaseSDKPath) == false) ? SDKStatus.Valid : SDKStatus.Invalid;
+        }
 
         public override bool CanUseXGE()
         {
             return (GetActiveArchitecture() == "-win32" );
         }
 
-		/**
-		 *	Register the platform with the UEBuildPlatform class
-		 */
-		public override void RegisterBuildPlatform()
-		{
-			[REDACTED_EMAIL]: Add platform support
-			if (UnrealBuildTool.RunningRocket() || UnrealBuildTool.BuildingRocket())
-			{
-				return;
-			}
-
-			// Make sure the SDK is installed
-			if ((ProjectFileGenerator.bGenerateProjectFiles == true) || (HasRequiredSDKsInstalled() == SDKStatus.Valid))
-			{
-				bool bRegisterBuildPlatform = true;
-
-				// make sure we have the HTML5 files; if not, then this user doesn't really have HTML5 access/files, no need to compile HTML5!
-				string EngineSourcePath = Path.Combine(ProjectFileGenerator.EngineRelativePath, "Source");
-				string HTML5TargetPlatformFile = Path.Combine(EngineSourcePath, "Developer", "HTML5", "HTML5TargetPlatform", "HTML5TargetPlatform.Build.cs");
- 				if ((File.Exists(HTML5TargetPlatformFile) == false))
- 				{
- 					bRegisterBuildPlatform = false;
- 					Log.TraceWarning("Missing required components (.... HTML5TargetPlatformFile, others here...). Check source control filtering, or try resyncing.");
- 				}
-
-				if (bRegisterBuildPlatform == true)
-				{
-					// Register this build platform for HTML5
-					Log.TraceVerbose("        Registering for {0}", UnrealTargetPlatform.HTML5.ToString());
-					UEBuildPlatform.RegisterBuildPlatform(UnrealTargetPlatform.HTML5, this);
+        /**
+         *	Register the platform with the UEBuildPlatform class
+         */
+        protected override void RegisterBuildPlatformInternal()
+        {
+            [REDACTED_EMAIL]: Add platform support
+            if (UnrealBuildTool.RunningRocket() || UnrealBuildTool.BuildingRocket())
+            {
+                return;
+            }
+
+            // Make sure the SDK is installed
+            if ((ProjectFileGenerator.bGenerateProjectFiles == true) || (HasRequiredSDKsInstalled() == SDKStatus.Valid))
+            {
+                bool bRegisterBuildPlatform = true;
+
+                // make sure we have the HTML5 files; if not, then this user doesn't really have HTML5 access/files, no need to compile HTML5!
+                string EngineSourcePath = Path.Combine(ProjectFileGenerator.EngineRelativePath, "Source");
+                string HTML5TargetPlatformFile = Path.Combine(EngineSourcePath, "Developer", "HTML5", "HTML5TargetPlatform", "HTML5TargetPlatform.Build.cs");
+                if ((File.Exists(HTML5TargetPlatformFile) == false))
+                {
+                    bRegisterBuildPlatform = false;
+                    Log.TraceWarning("Missing required components (.... HTML5TargetPlatformFile, others here...). Check source control filtering, or try resyncing.");
+                }
+
+                if (bRegisterBuildPlatform == true)
+                {
+                    // Register this build platform for HTML5
+                    Log.TraceVerbose("        Registering for {0}", UnrealTargetPlatform.HTML5.ToString());
+                    UEBuildPlatform.RegisterBuildPlatform(UnrealTargetPlatform.HTML5, this);
                     if (GetActiveArchitecture() == "-win32")
                     {
                         UEBuildPlatform.RegisterPlatformWithGroup(UnrealTargetPlatform.HTML5, UnrealPlatformGroup.Simulator);
@@ -89,59 +96,59 @@ namespace UnrealBuildTool
                     {
                         UEBuildPlatform.RegisterPlatformWithGroup(UnrealTargetPlatform.HTML5, UnrealPlatformGroup.Device);
                     }
-				}
-			}
+                }
+            }
 
             // if HTTP is not enabled, we can't use it. 
             if (EnableHTTPForNFS == false)
                 UseHTTPForNFS = false; 
-		}
-
-		/**
-		 *	Retrieve the CPPTargetPlatform for the given UnrealTargetPlatform
-		 *
-		 *	@param	InUnrealTargetPlatform		The UnrealTargetPlatform being build
-		 *	
-		 *	@return	CPPTargetPlatform			The CPPTargetPlatform to compile for
-		 */
-		public override CPPTargetPlatform GetCPPTargetPlatform(UnrealTargetPlatform InUnrealTargetPlatform)
-		{
-			switch (InUnrealTargetPlatform)
-			{
-				case UnrealTargetPlatform.HTML5:
-					return CPPTargetPlatform.HTML5;
-			}
-			throw new BuildException("HTML5Platform::GetCPPTargetPlatform: Invalid request for {0}", InUnrealTargetPlatform.ToString());
-		}
-
-		/**
-		 *	Get the extension to use for the given binary type
-		 *	
-		 *	@param	InBinaryType		The binary type being built
-		 *	
-		 *	@return	string				The binary extension (ie 'exe' or 'dll')
-		 */
-		public override string GetBinaryExtension(UEBuildBinaryType InBinaryType)
-		{
-			if (GetActiveArchitecture() == "-win32")
-			{
-				switch (InBinaryType)
-				{
-					case UEBuildBinaryType.DynamicLinkLibrary:
-						return ".dll";
-					case UEBuildBinaryType.Executable:
-						return ".exe";
-					case UEBuildBinaryType.StaticLibrary:
-						return ".lib";
-					case UEBuildBinaryType.Object:
-						return ".o";
-					case UEBuildBinaryType.PrecompiledHeader:
-						return ".gch";
-				}
-				return base.GetBinaryExtension(InBinaryType);
-			}
-			else
-			{
+        }
+
+        /**
+         *	Retrieve the CPPTargetPlatform for the given UnrealTargetPlatform
+         *
+         *	@param	InUnrealTargetPlatform		The UnrealTargetPlatform being build
+         *	
+         *	@return	CPPTargetPlatform			The CPPTargetPlatform to compile for
+         */
+        public override CPPTargetPlatform GetCPPTargetPlatform(UnrealTargetPlatform InUnrealTargetPlatform)
+        {
+            switch (InUnrealTargetPlatform)
+            {
+                case UnrealTargetPlatform.HTML5:
+                    return CPPTargetPlatform.HTML5;
+            }
+            throw new BuildException("HTML5Platform::GetCPPTargetPlatform: Invalid request for {0}", InUnrealTargetPlatform.ToString());
+        }
+
+        /**
+         *	Get the extension to use for the given binary type
+         *	
+         *	@param	InBinaryType		The binary type being built
+         *	
+         *	@return	string				The binary extension (ie 'exe' or 'dll')
+         */
+        public override string GetBinaryExtension(UEBuildBinaryType InBinaryType)
+        {
+            if (GetActiveArchitecture() == "-win32")
+            {
+                switch (InBinaryType)
+                {
+                    case UEBuildBinaryType.DynamicLinkLibrary:
+                        return ".dll";
+                    case UEBuildBinaryType.Executable:
+                        return ".exe";
+                    case UEBuildBinaryType.StaticLibrary:
+                        return ".lib";
+                    case UEBuildBinaryType.Object:
+                        return ".o";
+                    case UEBuildBinaryType.PrecompiledHeader:
+                        return ".gch";
+                }
+                return base.GetBinaryExtension(InBinaryType);
+            }
+            else
+            {
                 switch (InBinaryType)
                 {
                     case UEBuildBinaryType.DynamicLinkLibrary:
@@ -158,138 +165,138 @@ namespace UnrealBuildTool
 
                 return base.GetBinaryExtension(InBinaryType);
             }
-		}
-
-		/**
-		 *	Get the extension to use for debug info for the given binary type
-		 *	
-		 *	@param	InBinaryType		The binary type being built
-		 *	
-		 *	@return	string				The debug info extension (i.e. 'pdb')
-		 */
-		public override string GetDebugInfoExtension(UEBuildBinaryType InBinaryType)
-		{
-			if (GetActiveArchitecture() == "-win32")
-			{
-				switch (InBinaryType)
-				{
-					case UEBuildBinaryType.DynamicLinkLibrary:
-						return ".pdb";
-					case UEBuildBinaryType.Executable:
-						return ".pdb";
-				}
-				return "";
-			}
-			else
-			{
-				return "";
-			}
-		}
-
-		/**
-		 *	Gives the platform a chance to 'override' the configuration settings 
-		 *	that are overridden on calls to RunUBT.
-		 *	
-		 *	@param	InPlatform			The UnrealTargetPlatform being built
-		 *	@param	InConfiguration		The UnrealTargetConfiguration being built
-		 *	
-		 *	@return	bool				true if debug info should be generated, false if not
-		 */
-		public override void ResetBuildConfiguration(UnrealTargetPlatform InPlatform, UnrealTargetConfiguration InConfiguration)
-		{
-			ValidateUEBuildConfiguration();
-		}
-
-		/**
-		 *	Validate the UEBuildConfiguration for this platform
-		 *	This is called BEFORE calling UEBuildConfiguration to allow setting 
-		 *	various fields used in that function such as CompileLeanAndMean...
-		 */
-		public override void ValidateUEBuildConfiguration()
-		{
-			UEBuildConfiguration.bCompileLeanAndMeanUE = true;
-			UEBuildConfiguration.bCompileAPEX = false;
-			UEBuildConfiguration.bCompilePhysX = true;
-			UEBuildConfiguration.bCompileSimplygon = false;
+        }
+
+        /**
+         *	Get the extension to use for debug info for the given binary type
+         *	
+         *	@param	InBinaryType		The binary type being built
+         *	
+         *	@return	string				The debug info extension (i.e. 'pdb')
+         */
+        public override string GetDebugInfoExtension(UEBuildBinaryType InBinaryType)
+        {
+            if (GetActiveArchitecture() == "-win32")
+            {
+                switch (InBinaryType)
+                {
+                    case UEBuildBinaryType.DynamicLinkLibrary:
+                        return ".pdb";
+                    case UEBuildBinaryType.Executable:
+                        return ".pdb";
+                }
+                return "";
+            }
+            else
+            {
+                return "";
+            }
+        }
+
+        /**
+         *	Gives the platform a chance to 'override' the configuration settings 
+         *	that are overridden on calls to RunUBT.
+         *	
+         *	@param	InPlatform			The UnrealTargetPlatform being built
+         *	@param	InConfiguration		The UnrealTargetConfiguration being built
+         *	
+         *	@return	bool				true if debug info should be generated, false if not
+         */
+        public override void ResetBuildConfiguration(UnrealTargetPlatform InPlatform, UnrealTargetConfiguration InConfiguration)
+        {
+            ValidateUEBuildConfiguration();
+        }
+
+        /**
+         *	Validate the UEBuildConfiguration for this platform
+         *	This is called BEFORE calling UEBuildConfiguration to allow setting 
+         *	various fields used in that function such as CompileLeanAndMean...
+         */
+        public override void ValidateUEBuildConfiguration()
+        {
+            UEBuildConfiguration.bCompileLeanAndMeanUE = true;
+            UEBuildConfiguration.bCompileAPEX = false;
+            UEBuildConfiguration.bCompilePhysX = true;
+            UEBuildConfiguration.bCompileSimplygon = false;
             UEBuildConfiguration.bCompileICU = false;
-		}
-
-		/**
-		 * Whether this platform should build a monolithic binary
-		 */
-		public override bool ShouldCompileMonolithicBinary(UnrealTargetPlatform InPlatform)
-		{
-			// This platform currently always compiles monolithic
-			return true;
-		}
-
-		/**
-		 *	Whether PDB files should be used
-		 *	
-		 *	@param	InPlatform			The CPPTargetPlatform being built
-		 *	@param	InConfiguration		The CPPTargetConfiguration being built
-		 *	@param	bInCreateDebugInfo	true if debug info is getting create, false if not
-		 *	
-		 *	@return	bool	true if PDB files should be used, false if not
-		 */
-		public override bool ShouldUsePDBFiles(CPPTargetPlatform Platform, CPPTargetConfiguration Configuration, bool bCreateDebugInfo)
-		{
-			return bCreateDebugInfo;
-		}
-
-		/**
-		 *	Whether PCH files should be used
-		 *	
-		 *	@param	InPlatform			The CPPTargetPlatform being built
-		 *	@param	InConfiguration		The CPPTargetConfiguration being built
-		 *	
-		 *	@return	bool				true if PCH files should be used, false if not
-		 */
-		public override bool ShouldUsePCHFiles(CPPTargetPlatform Platform, CPPTargetConfiguration Configuration)
-		{
-			return false;
-		}
-
-		/**
-		 *	Whether the editor should be built for this platform or not
-		 *	
-		 *	@param	InPlatform		The UnrealTargetPlatform being built
-		 *	@param	InConfiguration	The UnrealTargetConfiguration being built
-		 *	@return	bool			true if the editor should be built, false if not
-		 */
-		public override bool ShouldNotBuildEditor(UnrealTargetPlatform InPlatform, UnrealTargetConfiguration InConfiguration)
-		{
-			return true;
-		}
-
-		public override bool BuildRequiresCookedData(UnrealTargetPlatform InPlatform, UnrealTargetConfiguration InConfiguration)
-		{
-			return true;
-		}
-
-		/**
-		 *	Modify the newly created module passed in for this platform.
-		 *	This is not required - but allows for hiding details of a
-		 *	particular platform.
-		 *	
-		 *	@param	InModule		The newly loaded module
-		 *	@param	Target			The target being build
-		 */
-		public override void ModifyNewlyLoadedModule(UEBuildModule InModule, TargetInfo Target)
-		{
-			if (Target.Platform == UnrealTargetPlatform.HTML5)
-			{
-				if (InModule.ToString() == "Core")
-				{
-					InModule.AddPublicIncludePath("Runtime/Core/Public/HTML5");
-					InModule.AddPublicDependencyModule("zlib");
-				}
-				else if (InModule.ToString() == "Engine")
-				{
-					InModule.AddPrivateDependencyModule("zlib");
-					InModule.AddPrivateDependencyModule("UElibPNG");
-					InModule.AddPublicDependencyModule("UEOgg");
-					InModule.AddPublicDependencyModule("Vorbis");
+        }
+
+        /**
+         * Whether this platform should build a monolithic binary
+         */
+        public override bool ShouldCompileMonolithicBinary(UnrealTargetPlatform InPlatform)
+        {
+            // This platform currently always compiles monolithic
+            return true;
+        }
+
+        /**
+         *	Whether PDB files should be used
+         *	
+         *	@param	InPlatform			The CPPTargetPlatform being built
+         *	@param	InConfiguration		The CPPTargetConfiguration being built
+         *	@param	bInCreateDebugInfo	true if debug info is getting create, false if not
+         *	
+         *	@return	bool	true if PDB files should be used, false if not
+         */
+        public override bool ShouldUsePDBFiles(CPPTargetPlatform Platform, CPPTargetConfiguration Configuration, bool bCreateDebugInfo)
+        {
+            return bCreateDebugInfo;
+        }
+
+        /**
+         *	Whether PCH files should be used
+         *	
+         *	@param	InPlatform			The CPPTargetPlatform being built
+         *	@param	InConfiguration		The CPPTargetConfiguration being built
+         *	
+         *	@return	bool				true if PCH files should be used, false if not
+         */
+        public override bool ShouldUsePCHFiles(CPPTargetPlatform Platform, CPPTargetConfiguration Configuration)
+        {
+            return false;
+        }
+
+        /**
+         *	Whether the editor should be built for this platform or not
+         *	
+         *	@param	InPlatform		The UnrealTargetPlatform being built
+         *	@param	InConfiguration	The UnrealTargetConfiguration being built
+         *	@return	bool			true if the editor should be built, false if not
+         */
+        public override bool ShouldNotBuildEditor(UnrealTargetPlatform InPlatform, UnrealTargetConfiguration InConfiguration)
+        {
+            return true;
+        }
+
+        public override bool BuildRequiresCookedData(UnrealTargetPlatform InPlatform, UnrealTargetConfiguration InConfiguration)
+        {
+            return true;
+        }
+
+        /**
+         *	Modify the newly created module passed in for this platform.
+         *	This is not required - but allows for hiding details of a
+         *	particular platform.
+         *	
+         *	@param	InModule		The newly loaded module
+         *	@param	Target			The target being build
+         */
+        public override void ModifyNewlyLoadedModule(UEBuildModule InModule, TargetInfo Target)
+        {
+            if (Target.Platform == UnrealTargetPlatform.HTML5)
+            {
+                if (InModule.ToString() == "Core")
+                {
+                    InModule.AddPublicIncludePath("Runtime/Core/Public/HTML5");
+                    InModule.AddPublicDependencyModule("zlib");
+                }
+                else if (InModule.ToString() == "Engine")
+                {
+                    InModule.AddPrivateDependencyModule("zlib");
+                    InModule.AddPrivateDependencyModule("UElibPNG");
+                    InModule.AddPublicDependencyModule("UEOgg");
+                    InModule.AddPublicDependencyModule("Vorbis");
                 }
 
                 if (InModule.ToString() == "NetworkFile")
@@ -304,14 +311,14 @@ namespace UnrealBuildTool
                         InModule.AddPrivateDependencyModule("HTML5JS");
                     }
                 }
-			}
-			else if (Target.Platform == UnrealTargetPlatform.Win64 || Target.Platform == UnrealTargetPlatform.Mac )
-			{
-				if ( (!UEBuildConfiguration.bBuildRequiresCookedData
-					&& InModule.ToString() == "Engine"
-					&& UEBuildConfiguration.bBuildDeveloperTools)
-					|| UEBuildConfiguration.bForceBuildTargetPlatforms)
-				{
+            }
+            else if (Target.Platform == UnrealTargetPlatform.Win64 || Target.Platform == UnrealTargetPlatform.Mac )
+            {
+                if ( (!UEBuildConfiguration.bBuildRequiresCookedData
+                    && InModule.ToString() == "Engine"
+                    && UEBuildConfiguration.bBuildDeveloperTools)
+                    || UEBuildConfiguration.bForceBuildTargetPlatforms)
+                {
                     InModule.AddPlatformSpecificDynamicallyLoadedModule("HTML5TargetPlatform");
                 }
                 if (InModule.ToString() == "NetworkFile") // client
@@ -324,74 +331,74 @@ namespace UnrealBuildTool
                     InModule.AddPublicDependencyModule("WebSockets");
                     InModule.AddPublicDefinition("ENABLE_HTTP_FOR_NFS=1");
                 }
-			}
-		}
+            }
+        }
 
 
-		/**
-		 *	Setup the target environment for building
-		 *	
-		 *	@param	InBuildTarget		The target being built
-		 */
-		public override void SetUpEnvironment(UEBuildTarget InBuildTarget)
-		{
+        /**
+         *	Setup the target environment for building
+         *	
+         *	@param	InBuildTarget		The target being built
+         */
+        public override void SetUpEnvironment(UEBuildTarget InBuildTarget)
+        {
             if (GetActiveArchitecture() == "-win32")
                 InBuildTarget.GlobalLinkEnvironment.Config.ExcludedLibraries.Add("LIBCMT");
 
-			InBuildTarget.GlobalCompileEnvironment.Config.Definitions.Add("PLATFORM_HTML5=1");
+            InBuildTarget.GlobalCompileEnvironment.Config.Definitions.Add("PLATFORM_HTML5=1");
             InBuildTarget.GlobalCompileEnvironment.Config.Definitions.Add("HTML5=1");
             if (InBuildTarget.GlobalCompileEnvironment.Config.Target.Architecture == "-win32")
-			{
-				InBuildTarget.GlobalCompileEnvironment.Config.Definitions.Add("PLATFORM_HTML5_WIN32=1");
-				InBuildTarget.GlobalLinkEnvironment.Config.AdditionalLibraries.Add("delayimp.lib");
-			}
+            {
+                InBuildTarget.GlobalCompileEnvironment.Config.Definitions.Add("PLATFORM_HTML5_WIN32=1");
+                InBuildTarget.GlobalLinkEnvironment.Config.AdditionalLibraries.Add("delayimp.lib");
+            }
             else
             {
                 InBuildTarget.GlobalCompileEnvironment.Config.Definitions.Add("PLATFORM_HTML5_BROWSER=1");
             }
 
-			// @todo needed?
-			InBuildTarget.GlobalCompileEnvironment.Config.Definitions.Add("UNICODE");
-			InBuildTarget.GlobalCompileEnvironment.Config.Definitions.Add("_UNICODE");
+            // @todo needed?
+            InBuildTarget.GlobalCompileEnvironment.Config.Definitions.Add("UNICODE");
+            InBuildTarget.GlobalCompileEnvironment.Config.Definitions.Add("_UNICODE");
             InBuildTarget.GlobalCompileEnvironment.Config.Definitions.Add("WITH_AUTOMATION_WORKER=0");
             InBuildTarget.GlobalCompileEnvironment.Config.Definitions.Add("REQUIRES_ALIGNED_INT_ACCESS");
             InBuildTarget.GlobalCompileEnvironment.Config.Definitions.Add("WITH_OGGVORBIS=1");
             InBuildTarget.GlobalCompileEnvironment.Config.Definitions.Add("USE_SCENE_LOCK=0");
-         	BuildConfiguration.bDeployAfterCompile = true;
-
-		}
-
-		/**
-		 *	Whether this platform should create debug information or not
-		 *	
-		 *	@param	InPlatform			The UnrealTargetPlatform being built
-		 *	@param	InConfiguration		The UnrealTargetConfiguration being built
-		 *	
-		 *	@return	bool				true if debug info should be generated, false if not
-		 */
-		public override bool ShouldCreateDebugInfo(UnrealTargetPlatform Platform, UnrealTargetConfiguration Configuration)
-		{
-			switch (Configuration)
-			{
-				case UnrealTargetConfiguration.Development:
-				case UnrealTargetConfiguration.Shipping:
-				case UnrealTargetConfiguration.Test:
-					return !BuildConfiguration.bOmitPCDebugInfoInDevelopment;
-				case UnrealTargetConfiguration.Debug:
-				default:
-					// We don't need debug info for Emscripten, and it causes a bunch of issues with linking
-					return true;
-			};
-		}
-
-		/**
-		 *	Setup the binaries for this specific platform.
-		 *	
-		 *	@param	InBuildTarget		The target being built
-		 */
-		public override void SetupBinaries(UEBuildTarget InBuildTarget)
-		{
+            BuildConfiguration.bDeployAfterCompile = true;
+
+        }
+
+        /**
+         *	Whether this platform should create debug information or not
+         *	
+         *	@param	InPlatform			The UnrealTargetPlatform being built
+         *	@param	InConfiguration		The UnrealTargetConfiguration being built
+         *	
+         *	@return	bool				true if debug info should be generated, false if not
+         */
+        public override bool ShouldCreateDebugInfo(UnrealTargetPlatform Platform, UnrealTargetConfiguration Configuration)
+        {
+            switch (Configuration)
+            {
+                case UnrealTargetConfiguration.Development:
+                case UnrealTargetConfiguration.Shipping:
+                case UnrealTargetConfiguration.Test:
+                    return !BuildConfiguration.bOmitPCDebugInfoInDevelopment;
+                case UnrealTargetConfiguration.Debug:
+                default:
+                    // We don't need debug info for Emscripten, and it causes a bunch of issues with linking
+                    return true;
+            };
+        }
+
+        /**
+         *	Setup the binaries for this specific platform.
+         *	
+         *	@param	InBuildTarget		The target being built
+         */
+        public override void SetupBinaries(UEBuildTarget InBuildTarget)
+        {
 // 			InBuildTarget.ExtraModuleNames.Add("ES2RHI");
-		}
-	}
+        }
+    }
 }
diff --git a/Engine/Source/Programs/UnrealBuildTool/IOS/UEBuildIOS.cs b/Engine/Source/Programs/UnrealBuildTool/IOS/UEBuildIOS.cs
--- a/Engine/Source/Programs/UnrealBuildTool/IOS/UEBuildIOS.cs
+++ b/Engine/Source/Programs/UnrealBuildTool/IOS/UEBuildIOS.cs
@@ -9,33 +9,33 @@ using System.Xml;
 
 namespace UnrealBuildTool
 {
-	class IOSPlatform : UEBuildPlatform
-	{
-		// by default, use an empty architecture (which is really just a modifer to the platform for some paths/names)
-		[XmlConfig]
-		public static string IOSArchitecture = "";
-
-		// The current architecture - affects everything about how UBT operates on IOS
-		public override string GetActiveArchitecture()
-		{
-			return IOSArchitecture;
-		}
-
-		public override SDKStatus HasRequiredSDKsInstalled()
-		{
-			return SDKStatus.Valid;
-		}
-
-		/**
-		 *	Register the platform with the UEBuildPlatform class
-		 */
-		public override void RegisterBuildPlatform()
-		{
-			// Register this build platform for IOS
-			Log.TraceVerbose("        Registering for {0}", UnrealTargetPlatform.IOS.ToString());
-			UEBuildPlatform.RegisterBuildPlatform(UnrealTargetPlatform.IOS, this);
-			UEBuildPlatform.RegisterPlatformWithGroup(UnrealTargetPlatform.IOS, UnrealPlatformGroup.Unix);
-			UEBuildPlatform.RegisterPlatformWithGroup(UnrealTargetPlatform.IOS, UnrealPlatformGroup.Apple);
+    class IOSPlatform : UEBuildPlatform
+    {
+        // by default, use an empty architecture (which is really just a modifer to the platform for some paths/names)
+        [XmlConfig]
+        public static string IOSArchitecture = "";
+
+        // The current architecture - affects everything about how UBT operates on IOS
+        public override string GetActiveArchitecture()
+        {
+            return IOSArchitecture;
+        }
+
+        protected override SDKStatus HasRequiredManualSDKInternal()
+        {
+            return SDKStatus.Valid;
+        }
+
+        /**
+         *	Register the platform with the UEBuildPlatform class
+         */
+        protected override void RegisterBuildPlatformInternal()
+        {
+            // Register this build platform for IOS
+            Log.TraceVerbose("        Registering for {0}", UnrealTargetPlatform.IOS.ToString());
+            UEBuildPlatform.RegisterBuildPlatform(UnrealTargetPlatform.IOS, this);
+            UEBuildPlatform.RegisterPlatformWithGroup(UnrealTargetPlatform.IOS, UnrealPlatformGroup.Unix);
+            UEBuildPlatform.RegisterPlatformWithGroup(UnrealTargetPlatform.IOS, UnrealPlatformGroup.Apple);
 
             if (GetActiveArchitecture() == "-simulator")
             {
@@ -45,286 +45,286 @@ namespace UnrealBuildTool
             {
                 UEBuildPlatform.RegisterPlatformWithGroup(UnrealTargetPlatform.IOS, UnrealPlatformGroup.Device);
             }
-		}
-
-		/**
-		 *	Retrieve the CPPTargetPlatform for the given UnrealTargetPlatform
-		 *
-		 *	@param	InUnrealTargetPlatform		The UnrealTargetPlatform being build
-		 *	
-		 *	@return	CPPTargetPlatform			The CPPTargetPlatform to compile for
-		 */
-		public override CPPTargetPlatform GetCPPTargetPlatform(UnrealTargetPlatform InUnrealTargetPlatform)
-		{
-			switch (InUnrealTargetPlatform)
-			{
-				case UnrealTargetPlatform.IOS:
-					return CPPTargetPlatform.IOS;
-			}
-			throw new BuildException("IOSPlatform::GetCPPTargetPlatform: Invalid request for {0}", InUnrealTargetPlatform.ToString());
-		}
-
-		/**
-		 *	Get the extension to use for the given binary type
-		 *	
-		 *	@param	InBinaryType		The binary type being built
-		 *	
-		 *	@return	string				The binary extenstion (ie 'exe' or 'dll')
-		 */
-		public override string GetBinaryExtension(UEBuildBinaryType InBinaryType)
-		{
-			switch (InBinaryType)
-			{
-				case UEBuildBinaryType.DynamicLinkLibrary:
-					return ".dylib";
-				case UEBuildBinaryType.Executable:
-					return "";
-				case UEBuildBinaryType.StaticLibrary:
-					return ".a";
-				case UEBuildBinaryType.Object:
-					return ".o";
-				case UEBuildBinaryType.PrecompiledHeader:
-					return ".gch";
-			}
-			return base.GetBinaryExtension(InBinaryType);
-		}
-
-		public override string GetDebugInfoExtension(UEBuildBinaryType InBinaryType)
-		{
-			return BuildConfiguration.bGeneratedSYMFile ? ".dsym" : "";
-		}
+        }
+
+        /**
+         *	Retrieve the CPPTargetPlatform for the given UnrealTargetPlatform
+         *
+         *	@param	InUnrealTargetPlatform		The UnrealTargetPlatform being build
+         *	
+         *	@return	CPPTargetPlatform			The CPPTargetPlatform to compile for
+         */
+        public override CPPTargetPlatform GetCPPTargetPlatform(UnrealTargetPlatform InUnrealTargetPlatform)
+        {
+            switch (InUnrealTargetPlatform)
+            {
+                case UnrealTargetPlatform.IOS:
+                    return CPPTargetPlatform.IOS;
+            }
+            throw new BuildException("IOSPlatform::GetCPPTargetPlatform: Invalid request for {0}", InUnrealTargetPlatform.ToString());
+        }
+
+        /**
+         *	Get the extension to use for the given binary type
+         *	
+         *	@param	InBinaryType		The binary type being built
+         *	
+         *	@return	string				The binary extenstion (ie 'exe' or 'dll')
+         */
+        public override string GetBinaryExtension(UEBuildBinaryType InBinaryType)
+        {
+            switch (InBinaryType)
+            {
+                case UEBuildBinaryType.DynamicLinkLibrary:
+                    return ".dylib";
+                case UEBuildBinaryType.Executable:
+                    return "";
+                case UEBuildBinaryType.StaticLibrary:
+                    return ".a";
+                case UEBuildBinaryType.Object:
+                    return ".o";
+                case UEBuildBinaryType.PrecompiledHeader:
+                    return ".gch";
+            }
+            return base.GetBinaryExtension(InBinaryType);
+        }
+
+        public override string GetDebugInfoExtension(UEBuildBinaryType InBinaryType)
+        {
+            return BuildConfiguration.bGeneratedSYMFile ? ".dsym" : "";
+        }
 
         public override bool CanUseXGE()
         {
             return false;
         }
 
-		/**
-		 *	Setup the target environment for building
-		 *	
-		 *	@param	InBuildTarget		The target being built
-		 */
-		public override void SetUpEnvironment(UEBuildTarget InBuildTarget)
-		{
-			InBuildTarget.GlobalCompileEnvironment.Config.Definitions.Add("PLATFORM_IOS=1");
-			InBuildTarget.GlobalCompileEnvironment.Config.Definitions.Add("PLATFORM_APPLE=1");
-
-			InBuildTarget.GlobalCompileEnvironment.Config.Definitions.Add("WITH_TTS=0");
-			InBuildTarget.GlobalCompileEnvironment.Config.Definitions.Add("WITH_SPEECH_RECOGNITION=0");
-			InBuildTarget.GlobalCompileEnvironment.Config.Definitions.Add("WITH_DATABASE_SUPPORT=0");
-			InBuildTarget.GlobalCompileEnvironment.Config.Definitions.Add("WITH_EDITOR=0");
-			InBuildTarget.GlobalCompileEnvironment.Config.Definitions.Add("USE_NULL_RHI=0");
+        /**
+         *	Setup the target environment for building
+         *	
+         *	@param	InBuildTarget		The target being built
+         */
+        public override void SetUpEnvironment(UEBuildTarget InBuildTarget)
+        {
+            InBuildTarget.GlobalCompileEnvironment.Config.Definitions.Add("PLATFORM_IOS=1");
+            InBuildTarget.GlobalCompileEnvironment.Config.Definitions.Add("PLATFORM_APPLE=1");
+
+            InBuildTarget.GlobalCompileEnvironment.Config.Definitions.Add("WITH_TTS=0");
+            InBuildTarget.GlobalCompileEnvironment.Config.Definitions.Add("WITH_SPEECH_RECOGNITION=0");
+            InBuildTarget.GlobalCompileEnvironment.Config.Definitions.Add("WITH_DATABASE_SUPPORT=0");
+            InBuildTarget.GlobalCompileEnvironment.Config.Definitions.Add("WITH_EDITOR=0");
+            InBuildTarget.GlobalCompileEnvironment.Config.Definitions.Add("USE_NULL_RHI=0");
             InBuildTarget.GlobalCompileEnvironment.Config.Definitions.Add("REQUIRES_ALIGNED_INT_ACCESS");
 
-			if (GetActiveArchitecture() == "-simulator")
-			{
-				InBuildTarget.GlobalCompileEnvironment.Config.Definitions.Add("WITH_SIMULATOR=1");
-			}
-
-			// needs IOS8 for Metal
-			if (IOSToolChain.IOSSDKVersionFloat >= 8.0 && UEBuildConfiguration.bCompileAgainstEngine)
-			{
-				InBuildTarget.GlobalCompileEnvironment.Config.Definitions.Add("HAS_METAL=1");
-				InBuildTarget.ExtraModuleNames.Add("MetalRHI");
-			}
-			else
-			{
-				InBuildTarget.GlobalCompileEnvironment.Config.Definitions.Add("HAS_METAL=0");
-			}
+            if (GetActiveArchitecture() == "-simulator")
+            {
+                InBuildTarget.GlobalCompileEnvironment.Config.Definitions.Add("WITH_SIMULATOR=1");
+            }
+
+            // needs IOS8 for Metal
+            if (IOSToolChain.IOSSDKVersionFloat >= 8.0 && UEBuildConfiguration.bCompileAgainstEngine)
+            {
+                InBuildTarget.GlobalCompileEnvironment.Config.Definitions.Add("HAS_METAL=1");
+                InBuildTarget.ExtraModuleNames.Add("MetalRHI");
+            }
+            else
+            {
+                InBuildTarget.GlobalCompileEnvironment.Config.Definitions.Add("HAS_METAL=0");
+            }
 
             InBuildTarget.GlobalLinkEnvironment.Config.AdditionalFrameworks.Add( new UEBuildFramework( "GameKit" ) );
             InBuildTarget.GlobalLinkEnvironment.Config.AdditionalFrameworks.Add( new UEBuildFramework( "StoreKit" ) );
         }
 
-		/**
-		*	Whether the editor should be built for this platform or not
-		*	
-		*	@param	InPlatform		The UnrealTargetPlatform being built
-		*	@param	InConfiguration	The UnrealTargetConfiguration being built
-		*	@return	bool			true if the editor should be built, false if not
-		*/
-		public override bool ShouldNotBuildEditor(UnrealTargetPlatform InPlatform, UnrealTargetConfiguration InConfiguration)
-		{
-			return true;
-		}
-
-		public override bool BuildRequiresCookedData(UnrealTargetPlatform InPlatform, UnrealTargetConfiguration InConfiguration)
-		{
-			return true; // for iOS can only run cooked. this is mostly for testing console code paths.
-		}
-
-		/**
-		 *	Whether this platform should create debug information or not
-		 *	
-		 *	@param	InPlatform			The UnrealTargetPlatform being built
-		 *	@param	InConfiguration		The UnrealTargetConfiguration being built
-		 *	
-		 *	@return	bool				true if debug info should be generated, false if not
-		 */
-		public override bool ShouldCreateDebugInfo(UnrealTargetPlatform Platform, UnrealTargetConfiguration Configuration)
-		{
-			return true;
-		}
-
-		public override void ResetBuildConfiguration(UnrealTargetPlatform InPlatform, UnrealTargetConfiguration InConfiguration)
-		{
-			UEBuildConfiguration.bBuildEditor = false;
-			UEBuildConfiguration.bBuildDeveloperTools = false;
-			UEBuildConfiguration.bCompileAPEX = false;
-			UEBuildConfiguration.bCompileSimplygon = false;
-			UEBuildConfiguration.bBuildDeveloperTools = false;
-
-			// we currently don't have any simulator libs for PhysX
-			if (GetActiveArchitecture() == "-simulator")
-			{
-				UEBuildConfiguration.bCompilePhysX = false;
-			}
-		}
-
-		public override bool ShouldCompileMonolithicBinary(UnrealTargetPlatform InPlatform)
-		{
-			// This platform currently always compiles monolithic
-			return true;
-		}
-
-		public override void ValidateBuildConfiguration(CPPTargetConfiguration Configuration, CPPTargetPlatform Platform, bool bCreateDebugInfo)
-		{
-			BuildConfiguration.bUsePCHFiles = false;
-			BuildConfiguration.bUseSharedPCHs = false;
-			BuildConfiguration.bCheckExternalHeadersForModification = false;
-			BuildConfiguration.bCheckSystemHeadersForModification = false;
-			BuildConfiguration.ProcessorCountMultiplier = IOSToolChain.GetAdjustedProcessorCountMultiplier();
-			BuildConfiguration.bDeployAfterCompile = true;
-		}
-
-		/**
-		 *	Whether the platform requires the extra UnityCPPWriter
-		 *	This is used to add an extra file for UBT to get the #include dependencies from
-		 *	
-		 *	@return	bool	true if it is required, false if not
-		 */
-		public override bool RequiresExtraUnityCPPWriter()
-		{
-			return true;
-		}
-
-		/**
-		 *     Modify the newly created module passed in for this platform.
-		 *     This is not required - but allows for hiding details of a
-		 *     particular platform.
-		 *     
-		 *     @param InModule             The newly loaded module
-		 *     @param GameName             The game being build
-		 *     @param Target               The target being build
-		 */
-		public override void ModifyNewlyLoadedModule(UEBuildModule InModule, TargetInfo Target)
-		{
-			if ((Target.Platform == UnrealTargetPlatform.Win32) || (Target.Platform == UnrealTargetPlatform.Win64) || (Target.Platform == UnrealTargetPlatform.Mac))
-			{
-				bool bBuildShaderFormats = UEBuildConfiguration.bForceBuildShaderFormats;
-				if (!UEBuildConfiguration.bBuildRequiresCookedData)
-				{
-					if (InModule.ToString() == "Engine")
-					{
-						if (UEBuildConfiguration.bBuildDeveloperTools)
-						{
+        /**
+        *	Whether the editor should be built for this platform or not
+        *	
+        *	@param	InPlatform		The UnrealTargetPlatform being built
+        *	@param	InConfiguration	The UnrealTargetConfiguration being built
+        *	@return	bool			true if the editor should be built, false if not
+        */
+        public override bool ShouldNotBuildEditor(UnrealTargetPlatform InPlatform, UnrealTargetConfiguration InConfiguration)
+        {
+            return true;
+        }
+
+        public override bool BuildRequiresCookedData(UnrealTargetPlatform InPlatform, UnrealTargetConfiguration InConfiguration)
+        {
+            return true; // for iOS can only run cooked. this is mostly for testing console code paths.
+        }
+
+        /**
+         *	Whether this platform should create debug information or not
+         *	
+         *	@param	InPlatform			The UnrealTargetPlatform being built
+         *	@param	InConfiguration		The UnrealTargetConfiguration being built
+         *	
+         *	@return	bool				true if debug info should be generated, false if not
+         */
+        public override bool ShouldCreateDebugInfo(UnrealTargetPlatform Platform, UnrealTargetConfiguration Configuration)
+        {
+            return true;
+        }
+
+        public override void ResetBuildConfiguration(UnrealTargetPlatform InPlatform, UnrealTargetConfiguration InConfiguration)
+        {
+            UEBuildConfiguration.bBuildEditor = false;
+            UEBuildConfiguration.bBuildDeveloperTools = false;
+            UEBuildConfiguration.bCompileAPEX = false;
+            UEBuildConfiguration.bCompileSimplygon = false;
+            UEBuildConfiguration.bBuildDeveloperTools = false;
+
+            // we currently don't have any simulator libs for PhysX
+            if (GetActiveArchitecture() == "-simulator")
+            {
+                UEBuildConfiguration.bCompilePhysX = false;
+            }
+        }
+
+        public override bool ShouldCompileMonolithicBinary(UnrealTargetPlatform InPlatform)
+        {
+            // This platform currently always compiles monolithic
+            return true;
+        }
+
+        public override void ValidateBuildConfiguration(CPPTargetConfiguration Configuration, CPPTargetPlatform Platform, bool bCreateDebugInfo)
+        {
+            BuildConfiguration.bUsePCHFiles = false;
+            BuildConfiguration.bUseSharedPCHs = false;
+            BuildConfiguration.bCheckExternalHeadersForModification = false;
+            BuildConfiguration.bCheckSystemHeadersForModification = false;
+            BuildConfiguration.ProcessorCountMultiplier = IOSToolChain.GetAdjustedProcessorCountMultiplier();
+            BuildConfiguration.bDeployAfterCompile = true;
+        }
+
+        /**
+         *	Whether the platform requires the extra UnityCPPWriter
+         *	This is used to add an extra file for UBT to get the #include dependencies from
+         *	
+         *	@return	bool	true if it is required, false if not
+         */
+        public override bool RequiresExtraUnityCPPWriter()
+        {
+            return true;
+        }
+
+        /**
+         *     Modify the newly created module passed in for this platform.
+         *     This is not required - but allows for hiding details of a
+         *     particular platform.
+         *     
+         *     @param InModule             The newly loaded module
+         *     @param GameName             The game being build
+         *     @param Target               The target being build
+         */
+        public override void ModifyNewlyLoadedModule(UEBuildModule InModule, TargetInfo Target)
+        {
+            if ((Target.Platform == UnrealTargetPlatform.Win32) || (Target.Platform == UnrealTargetPlatform.Win64) || (Target.Platform == UnrealTargetPlatform.Mac))
+            {
+                bool bBuildShaderFormats = UEBuildConfiguration.bForceBuildShaderFormats;
+                if (!UEBuildConfiguration.bBuildRequiresCookedData)
+                {
+                    if (InModule.ToString() == "Engine")
+                    {
+                        if (UEBuildConfiguration.bBuildDeveloperTools)
+                        {
                             InModule.AddPlatformSpecificDynamicallyLoadedModule("IOSTargetPlatform");
-						}
-					}
-					else if (InModule.ToString() == "UnrealEd")
-					{
-						InModule.AddPlatformSpecificDynamicallyLoadedModule("IOSPlatformEditor");
-					}
-					else if (InModule.ToString() == "TargetPlatform")
-					{
-						bBuildShaderFormats = true;
+                        }
+                    }
+                    else if (InModule.ToString() == "UnrealEd")
+                    {
+                        InModule.AddPlatformSpecificDynamicallyLoadedModule("IOSPlatformEditor");
+                    }
+                    else if (InModule.ToString() == "TargetPlatform")
+                    {
+                        bBuildShaderFormats = true;
                         InModule.AddDynamicallyLoadedModule("TextureFormatPVR");
-						if (UEBuildConfiguration.bBuildDeveloperTools)
-						{
+                        if (UEBuildConfiguration.bBuildDeveloperTools)
+                        {
                             InModule.AddPlatformSpecificDynamicallyLoadedModule("AudioFormatADPCM");
-						}
-					}
-				}
+                        }
+                    }
+                }
 
-				// allow standalone tools to use targetplatform modules, without needing Engine
-				if (UEBuildConfiguration.bForceBuildTargetPlatforms)
-				{
+                // allow standalone tools to use targetplatform modules, without needing Engine
+                if (UEBuildConfiguration.bForceBuildTargetPlatforms)
+                {
                     InModule.AddPlatformSpecificDynamicallyLoadedModule("IOSTargetPlatform");
-				}
-
- 				if (bBuildShaderFormats)
- 				{
-					InModule.AddPlatformSpecificDynamicallyLoadedModule("MetalShaderFormat");
- 				}
-			}
-		}
-
-		public override void SetupBinaries(UEBuildTarget InBuildTarget)
-		{
-			if (ExternalExecution.GetRuntimePlatform () != UnrealTargetPlatform.Mac)
-			{
-				// dangerously fast mode doesn't generate stub files
-				if (!IOSToolChain.bUseDangerouslyFastMode)
-				{
-					List<UEBuildBinary> NewBinaries = new List<UEBuildBinary> ();
-					// add the .stub to the binaries
-					foreach (var Binary in InBuildTarget.AppBinaries)
-					{
-						// make a binary that just points to the .stub of this executable
-						UEBuildBinaryConfiguration NewConfig = new UEBuildBinaryConfiguration(
-							                                      InType: Binary.Config.Type,
-							                                      InOutputFilePath: Binary.Config.OutputFilePath + ".stub",
-							                                      InIntermediateDirectory: Binary.Config.IntermediateDirectory,
-							                                      bInCreateImportLibrarySeparately: Binary.Config.bCreateImportLibrarySeparately,
-							                                      bInAllowExports: Binary.Config.bAllowExports,
-							                                      InModuleNames: Binary.Config.ModuleNames);
-
-						NewBinaries.Add (new UEStubDummyBinary (InBuildTarget, NewConfig));
-					}
-
-					InBuildTarget.AppBinaries.AddRange (NewBinaries);
-				}
-			}
-		}
-	}
-
-	/// <summary>
-	/// A .stub that has the executable and metadata included
-	/// Note that this doesn't actually build anything. It could potentially be used to perform the IPhonePackage stuff that happens in IOSToolChain.PostBuildSync()
-	/// </summary>
-	class UEStubDummyBinary : UEBuildBinary
-	{
-		/// <summary>
-		/// Create an instance initialized to the given configuration
-		/// </summary>
-		/// <param name="InConfig">The build binary configuration to initialize the instance to</param>
-		public UEStubDummyBinary(UEBuildTarget InTarget, UEBuildBinaryConfiguration InConfig)
-			: base(InTarget, InConfig)
-		{
-		}
-
-		/// <summary>
-		/// Builds the binary.
-		/// </summary>
-		/// <param name="CompileEnvironment">The environment to compile the binary in</param>
-		/// <param name="LinkEnvironment">The environment to link the binary in</param>
-		/// <returns></returns>
-		public override IEnumerable<FileItem> Build(IUEToolChain ToolChain, CPPEnvironment CompileEnvironment, LinkEnvironment LinkEnvironment)
-		{
-			// generate the .stub!
-			return new FileItem[] { FileItem.GetItemByPath(Config.OutputFilePath) };
-		}
-
-		/// <summary>
-		/// Writes an XML summary of the build environment for this binary
-		/// </summary>
-		/// <param name="CompileEnvironment">The environment to compile the binary in</param>
-		/// <param name="LinkEnvironment">The environment to link the binary in</param>
-		/// <returns></returns>
-		public override void WriteBuildEnvironment(CPPEnvironment CompileEnvironment, LinkEnvironment LinkEnvironment, XmlWriter Writer)
-		{
-		}
-	}
+                }
+
+                if (bBuildShaderFormats)
+                {
+                    InModule.AddPlatformSpecificDynamicallyLoadedModule("MetalShaderFormat");
+                }
+            }
+        }
+
+        public override void SetupBinaries(UEBuildTarget InBuildTarget)
+        {
+            if (ExternalExecution.GetRuntimePlatform () != UnrealTargetPlatform.Mac)
+            {
+                // dangerously fast mode doesn't generate stub files
+                if (!IOSToolChain.bUseDangerouslyFastMode)
+                {
+                    List<UEBuildBinary> NewBinaries = new List<UEBuildBinary> ();
+                    // add the .stub to the binaries
+                    foreach (var Binary in InBuildTarget.AppBinaries)
+                    {
+                        // make a binary that just points to the .stub of this executable
+                        UEBuildBinaryConfiguration NewConfig = new UEBuildBinaryConfiguration(
+                                                                  InType: Binary.Config.Type,
+                                                                  InOutputFilePath: Binary.Config.OutputFilePath + ".stub",
+                                                                  InIntermediateDirectory: Binary.Config.IntermediateDirectory,
+                                                                  bInCreateImportLibrarySeparately: Binary.Config.bCreateImportLibrarySeparately,
+                                                                  bInAllowExports: Binary.Config.bAllowExports,
+                                                                  InModuleNames: Binary.Config.ModuleNames);
+
+                        NewBinaries.Add (new UEStubDummyBinary (InBuildTarget, NewConfig));
+                    }
+
+                    InBuildTarget.AppBinaries.AddRange (NewBinaries);
+                }
+            }
+        }
+    }
+
+    /// <summary>
+    /// A .stub that has the executable and metadata included
+    /// Note that this doesn't actually build anything. It could potentially be used to perform the IPhonePackage stuff that happens in IOSToolChain.PostBuildSync()
+    /// </summary>
+    class UEStubDummyBinary : UEBuildBinary
+    {
+        /// <summary>
+        /// Create an instance initialized to the given configuration
+        /// </summary>
+        /// <param name="InConfig">The build binary configuration to initialize the instance to</param>
+        public UEStubDummyBinary(UEBuildTarget InTarget, UEBuildBinaryConfiguration InConfig)
+            : base(InTarget, InConfig)
+        {
+        }
+
+        /// <summary>
+        /// Builds the binary.
+        /// </summary>
+        /// <param name="CompileEnvironment">The environment to compile the binary in</param>
+        /// <param name="LinkEnvironment">The environment to link the binary in</param>
+        /// <returns></returns>
+        public override IEnumerable<FileItem> Build(IUEToolChain ToolChain, CPPEnvironment CompileEnvironment, LinkEnvironment LinkEnvironment)
+        {
+            // generate the .stub!
+            return new FileItem[] { FileItem.GetItemByPath(Config.OutputFilePath) };
+        }
+
+        /// <summary>
+        /// Writes an XML summary of the build environment for this binary
+        /// </summary>
+        /// <param name="CompileEnvironment">The environment to compile the binary in</param>
+        /// <param name="LinkEnvironment">The environment to link the binary in</param>
+        /// <returns></returns>
+        public override void WriteBuildEnvironment(CPPEnvironment CompileEnvironment, LinkEnvironment LinkEnvironment, XmlWriter Writer)
+        {
+        }
+    }
 
 }
 
diff --git a/Engine/Source/Programs/UnrealBuildTool/Linux/UEBuildLinux.cs b/Engine/Source/Programs/UnrealBuildTool/Linux/UEBuildLinux.cs
--- a/Engine/Source/Programs/UnrealBuildTool/Linux/UEBuildLinux.cs
+++ b/Engine/Source/Programs/UnrealBuildTool/Linux/UEBuildLinux.cs
@@ -39,7 +39,7 @@ namespace UnrealBuildTool
          * 
          * @return true if supports
          */
-        public override bool PlatformSupportsSDKSwitching()
+        protected override bool PlatformSupportsAutoSDKs()
         {
             return true;
         }
@@ -59,27 +59,26 @@ namespace UnrealBuildTool
          * 
          * @return Valid SDK string
          */
-        public override string GetRequiredSDKString()
+        protected override string GetRequiredSDKString()
         {
             return ExpectedSDKVersion;
         }
 
+        protected override String GetRequiredScriptVersionString()
+        {
+            return "2.0";
+        }
+
         /**
          *	Whether the required external SDKs are installed for this platform
          */
-        public override SDKStatus HasRequiredSDKsInstalled()
+        protected override SDKStatus HasRequiredManualSDKInternal()
         {
             if (ExternalExecution.GetRuntimePlatform() == UnrealTargetPlatform.Linux)
             {
                 return SDKStatus.Valid;
             }
 
-            // attempt to switch SDKs
-            if (base.HasRequiredSDKsInstalled() == SDKStatus.Valid)
-            {
-                return SDKStatus.Valid;
-            }
-
             string BaseLinuxPath = Environment.GetEnvironmentVariable("LINUX_ROOT");
 
             // we don't have an LINUX_ROOT specified
@@ -105,7 +104,7 @@ namespace UnrealBuildTool
         /**
          *	Register the platform with the UEBuildPlatform class
          */
-        public override void RegisterBuildPlatform()
+        protected override void RegisterBuildPlatformInternal()
         {
 			[REDACTED_EMAIL]: Add platform support
 			if (UnrealBuildTool.RunningRocket() || UnrealBuildTool.BuildingRocket())
diff --git a/Engine/Source/Programs/UnrealBuildTool/Mac/UEBuildMac.cs b/Engine/Source/Programs/UnrealBuildTool/Mac/UEBuildMac.cs
--- a/Engine/Source/Programs/UnrealBuildTool/Mac/UEBuildMac.cs
+++ b/Engine/Source/Programs/UnrealBuildTool/Mac/UEBuildMac.cs
@@ -16,7 +16,7 @@ namespace UnrealBuildTool
             return false;
         }
 
-		public override SDKStatus HasRequiredSDKsInstalled()
+		protected override SDKStatus HasRequiredManualSDKInternal()
 		{
 			return SDKStatus.Valid;
 		}
@@ -24,7 +24,7 @@ namespace UnrealBuildTool
 		/**
 		 *	Register the platform with the UEBuildPlatform class
 		 */
-		public override void RegisterBuildPlatform()
+		protected override void RegisterBuildPlatformInternal()
 		{
 			// Register this build platform for Mac
 			Log.TraceVerbose("        Registering for {0}", UnrealTargetPlatform.Mac.ToString());
diff --git a/Engine/Source/Programs/UnrealBuildTool/WinRT/UEBuildWinRT.cs b/Engine/Source/Programs/UnrealBuildTool/WinRT/UEBuildWinRT.cs
--- a/Engine/Source/Programs/UnrealBuildTool/WinRT/UEBuildWinRT.cs
+++ b/Engine/Source/Programs/UnrealBuildTool/WinRT/UEBuildWinRT.cs
@@ -8,479 +8,479 @@ using System.IO;
 
 namespace UnrealBuildTool
 {
-	class WinRTPlatform : UEBuildPlatform
-	{
-		/// <summary>
-		/// Should the app be compiled as WinRT
-		/// </summary>
-		[XmlConfig]
-		public static bool bCompileWinRT = false;
-
-		public static bool IsVisualStudioInstalled()
-		{
-			string BaseVSToolPath = WindowsPlatform.GetVSComnToolsPath();
-			if (string.IsNullOrEmpty(BaseVSToolPath) == false)
-			{
-				return true;
-			}
-			return false;
-		}
-
-		public static bool IsWindows8()
-		{
-			// Are we a Windows8 machine?
-			if ((Environment.OSVersion.Version.Major == 6) && (Environment.OSVersion.Version.Minor == 2))
-			{
-				return true;
-			}
-			return false;
-		}
+    class WinRTPlatform : UEBuildPlatform
+    {
+        /// <summary>
+        /// Should the app be compiled as WinRT
+        /// </summary>
+        [XmlConfig]
+        public static bool bCompileWinRT = false;
+
+        public static bool IsVisualStudioInstalled()
+        {
+            string BaseVSToolPath = WindowsPlatform.GetVSComnToolsPath();
+            if (string.IsNullOrEmpty(BaseVSToolPath) == false)
+            {
+                return true;
+            }
+            return false;
+        }
+
+        public static bool IsWindows8()
+        {
+            // Are we a Windows8 machine?
+            if ((Environment.OSVersion.Version.Major == 6) && (Environment.OSVersion.Version.Minor == 2))
+            {
+                return true;
+            }
+            return false;
+        }
         public override bool CanUseXGE()
         {
             return false;
         }
 
-		/**
-		 *	Whether the required external SDKs are installed for this platform
-		 */
-		public override SDKStatus HasRequiredSDKsInstalled()
-		{
-			return IsVisualStudioInstalled() ? SDKStatus.Valid : SDKStatus.Invalid;
-		}
-
-		/**
-		 *	Register the platform with the UEBuildPlatform class
-		 */
-		public override void RegisterBuildPlatform()
-		{
-			[REDACTED_EMAIL]: Add platform support
-			if (UnrealBuildTool.RunningRocket() || UnrealBuildTool.BuildingRocket() || Utils.IsRunningOnMono)
-			{
-				return;
-			}
-
-			if ((ProjectFileGenerator.bGenerateProjectFiles == true) || (IsVisualStudioInstalled() == true))
-			{
-				bool bRegisterBuildPlatform = true;
-
-				// We also need to check for the generated projects... to handle the case where someone generates projects w/out WinRT.
-				// Hardcoding this for now - but ideally it would be dynamically discovered.
- 				string EngineSourcePath = Path.Combine(ProjectFileGenerator.EngineRelativePath, "Source");
-				string WinRTRHIFile = Path.Combine(EngineSourcePath, "Runtime", "Windows", "D3D11RHI", "D3D11RHI.build.cs");
-				if (File.Exists(WinRTRHIFile) == false)
- 				{
- 					bRegisterBuildPlatform = false;
- 				}
-
-				if (bRegisterBuildPlatform == true)
-				{
-					// Register this build platform for WinRT
-					Log.TraceVerbose("        Registering for {0}", UnrealTargetPlatform.WinRT.ToString());
-					UEBuildPlatform.RegisterBuildPlatform(UnrealTargetPlatform.WinRT, this);
-					UEBuildPlatform.RegisterPlatformWithGroup(UnrealTargetPlatform.WinRT, UnrealPlatformGroup.Microsoft);
-
-					// For now only register WinRT_ARM is truly a Windows 8 machine.
-					// This will prevent people who do all platform builds from running into the compiler issue.
-					if (WinRTPlatform.IsWindows8() == true)
-					{
-						Log.TraceVerbose("        Registering for {0}", UnrealTargetPlatform.WinRT_ARM.ToString());
-						UEBuildPlatform.RegisterBuildPlatform(UnrealTargetPlatform.WinRT_ARM, this);
-						UEBuildPlatform.RegisterPlatformWithGroup(UnrealTargetPlatform.WinRT_ARM, UnrealPlatformGroup.Microsoft);
-					}
-				}
-			}
-		}
-
-		/**
-		 *	Retrieve the CPPTargetPlatform for the given UnrealTargetPlatform
-		 *
-		 *	@param	InUnrealTargetPlatform		The UnrealTargetPlatform being build
-		 *	
-		 *	@return	CPPTargetPlatform			The CPPTargetPlatform to compile for
-		 */
-		public override CPPTargetPlatform GetCPPTargetPlatform(UnrealTargetPlatform InUnrealTargetPlatform)
-		{
-			switch (InUnrealTargetPlatform)
-			{
-				case UnrealTargetPlatform.WinRT:
-					return CPPTargetPlatform.WinRT;
-				case UnrealTargetPlatform.WinRT_ARM:
-					return CPPTargetPlatform.WinRT_ARM;
-			}
-			throw new BuildException("WinRTPlatform::GetCPPTargetPlatform: Invalid request for {0}", InUnrealTargetPlatform.ToString());
-		}
-
-		/**
-		 *	Get the extension to use for the given binary type
-		 *	
-		 *	@param	InBinaryType		The binary type being built
-		 *	
-		 *	@return	string				The binary extension (i.e. 'exe' or 'dll')
-		 */
-		public override string GetBinaryExtension(UEBuildBinaryType InBinaryType)
-		{
-			switch (InBinaryType)
-			{
-				case UEBuildBinaryType.DynamicLinkLibrary:
-					return ".dll";
-				case UEBuildBinaryType.Executable:
-					return ".exe";
-				case UEBuildBinaryType.StaticLibrary:
-					return ".lib";
-				case UEBuildBinaryType.Object:
-					return ".obj";
-				case UEBuildBinaryType.PrecompiledHeader:
-					return ".pch";
-			}
-			return base.GetBinaryExtension(InBinaryType);
-		}
-
-		/**
-		 *	Get the extension to use for debug info for the given binary type
-		 *	
-		 *	@param	InBinaryType		The binary type being built
-		 *	
-		 *	@return	string				The debug info extension (i.e. 'pdb')
-		 */
-		public override string GetDebugInfoExtension(UEBuildBinaryType InBinaryType)
-		{
-			return ".pdb";
-		}
-
-		/**
-		 *	Gives the platform a chance to 'override' the configuration settings 
-		 *	that are overridden on calls to RunUBT.
-		 *	
-		 *	@param	InPlatform			The UnrealTargetPlatform being built
-		 *	@param	InConfiguration		The UnrealTargetConfiguration being built
-		 *	
-		 *	@return	bool				true if debug info should be generated, false if not
-		 */
-		public override void ResetBuildConfiguration(UnrealTargetPlatform InPlatform, UnrealTargetConfiguration InConfiguration)
-		{
-			ValidateUEBuildConfiguration();
-		}
-
-		/**
-		 *	Validate the UEBuildConfiguration for this platform
-		 *	This is called BEFORE calling UEBuildConfiguration to allow setting 
-		 *	various fields used in that function such as CompileLeanAndMean...
-		 */
-		public override void ValidateUEBuildConfiguration()
-		{
-			UEBuildConfiguration.bCompileLeanAndMeanUE = true;
-			UEBuildConfiguration.bCompilePhysX = false;
-			UEBuildConfiguration.bCompileAPEX = false;
-
-			BuildConfiguration.bUseUnityBuild = false;
-
-			// Don't stop compilation at first error...
-			BuildConfiguration.bStopXGECompilationAfterErrors = false;
-		}
-
-		/**
-		 * Whether this platform should build a monolithic binary
-		 */
-		public override bool ShouldCompileMonolithicBinary(UnrealTargetPlatform InPlatform)
-		{
-			return true;
-		}
-
-		/**
-		 *	Whether PDB files should be used
-		 *	
-		 *	@param	InPlatform			The CPPTargetPlatform being built
-		 *	@param	InConfiguration		The CPPTargetConfiguration being built
-		 *	@param	bInCreateDebugInfo	true if debug info is getting create, false if not
-		 *	
-		 *	@return	bool	true if PDB files should be used, false if not
-		 */
-		public override bool ShouldUsePDBFiles(CPPTargetPlatform Platform, CPPTargetConfiguration Configuration, bool bCreateDebugInfo)
-		{
-			return true;
-		}
-
-		/**
-		 *	Whether the editor should be built for this platform or not
-		 *	
-		 *	@param	InPlatform		The UnrealTargetPlatform being built
-		 *	@param	InConfiguration	The UnrealTargetConfiguration being built
-		 *	@return	bool			true if the editor should be built, false if not
-		 */
-		public override bool ShouldNotBuildEditor(UnrealTargetPlatform InPlatform, UnrealTargetConfiguration InConfiguration)
-		{
-			return true;
-		}
-
-		/**
-		 *	Whether the platform requires cooked data
-		 *	
-		 *	@param	InPlatform		The UnrealTargetPlatform being built
-		 *	@param	InConfiguration	The UnrealTargetConfiguration being built
-		 *	@return	bool			true if the platform requires cooked data, false if not
-		 */
-		public override bool BuildRequiresCookedData(UnrealTargetPlatform InPlatform, UnrealTargetConfiguration InConfiguration)
-		{
-			return true;
-		}
-
-		/**
-		 *	Get a list of extra modules the platform requires.
-		 *	This is to allow undisclosed platforms to add modules they need without exposing information about the platfomr.
-		 *	
-		 *	@param	Target						The target being build
-		 *	@param	BuildTarget					The UEBuildTarget getting build
-		 *	@param	PlatformExtraModules		OUTPUT the list of extra modules the platform needs to add to the target
-		 */
-		public override void GetExtraModules(TargetInfo Target, UEBuildTarget BuildTarget, ref List<string> PlatformExtraModules)
-		{
-			if (Target.Platform == UnrealTargetPlatform.WinRT)
-			{
-				PlatformExtraModules.Add("XAudio2");
-			}
-		}
-
-		/**
-		 *	Modify the newly created module passed in for this platform.
-		 *	This is not required - but allows for hiding details of a
-		 *	particular platform.
-		 *	
-		 *	@param	InModule		The newly loaded module
-		 *	@param	Target			The target being build
-		 */
-		public override void ModifyNewlyLoadedModule(UEBuildModule InModule, TargetInfo Target)
-		{
-			if ((Target.Platform == UnrealTargetPlatform.WinRT) || 
-				(Target.Platform == UnrealTargetPlatform.WinRT_ARM))
-			{
-				if (InModule.ToString() == "Core")
-				{
-					InModule.AddPublicIncludePath("Runtime/Core/Public/WinRT");
-					InModule.AddPublicDependencyModule("zlib");
-				}
-				else if (InModule.ToString() == "Engine")
-				{
-					InModule.AddPrivateDependencyModule("zlib");
-					InModule.AddPrivateDependencyModule("UElibPNG");
- 					InModule.AddPublicDependencyModule("UEOgg");
- 					InModule.AddPublicDependencyModule("Vorbis");
-				}
-				else if (InModule.ToString() == "Launch")
-				{
-				}
-				else if (InModule.ToString() == "D3D11RHI")
-				{
-					InModule.AddPublicDefinition("D3D11_CUSTOM_VIEWPORT_CONSTRUCTOR=1");
-					// To enable platform specific D3D11 RHI Types
-					InModule.AddPrivateIncludePath("Runtime/Windows/D3D11RHI/Private/WinRT");
-					// Hack to enable AllowWindowsPlatformTypes.h/HideWindowsPlatformTypes.h
-					InModule.AddPublicIncludePath("Runtime/Core/Public/Windows");
-				}
-				else if (InModule.ToString() == "Sockets")
-				{
-					// Hack to enable AllowWindowsPlatformTypes.h/HideWindowsPlatformTypes.h
-					InModule.AddPublicIncludePath("Runtime/Core/Public/Windows");
-				}
-				else if (InModule.ToString() == "PhysX")
-				{
-					string PhysXDir = UEBuildConfiguration.UEThirdPartySourceDirectory + "PhysX/PhysX-3.3/";
-
-					InModule.AddPublicIncludePath("include/foundation/WinRT");
-					if (Target.Platform == UnrealTargetPlatform.WinRT)
-					{
-						InModule.AddPublicLibraryPath(PhysXDir + "Lib/WinRT");
-					}
-					else
-					{
-						InModule.AddPublicLibraryPath(PhysXDir + "Lib/WinRT/ARM");
-					}
-
-					if (Target.Configuration == UnrealTargetConfiguration.Debug)
-					{
-						InModule.AddPublicAdditionalLibrary("PhysX3DEBUG.lib");
-						InModule.AddPublicAdditionalLibrary("PhysX3ExtensionsDEBUG.lib");
-						InModule.AddPublicAdditionalLibrary("PhysX3CookingDEBUG.lib");
-						InModule.AddPublicAdditionalLibrary("PhysX3CommonDEBUG.lib");
-						InModule.AddPublicAdditionalLibrary("PhysX3VehicleDEBUG.lib");
-						InModule.AddPublicAdditionalLibrary("PxTaskDEBUG.lib");
-						InModule.AddPublicAdditionalLibrary("PhysXVisualDebuggerSDKDEBUG.lib");
-						InModule.AddPublicAdditionalLibrary("PhysXProfileSDKDEBUG.lib");
-					}
-					else if (Target.Configuration == UnrealTargetConfiguration.Development)
-					{
-						InModule.AddPublicAdditionalLibrary("PhysX3PROFILE.lib");
-						InModule.AddPublicAdditionalLibrary("PhysX3ExtensionsPROFILE.lib");
-						InModule.AddPublicAdditionalLibrary("PhysX3CookingPROFILE.lib");
-						InModule.AddPublicAdditionalLibrary("PhysX3CommonPROFILE.lib");
-						InModule.AddPublicAdditionalLibrary("PhysX3VehiclePROFILE.lib");
-						InModule.AddPublicAdditionalLibrary("PxTaskPROFILE.lib");
-						InModule.AddPublicAdditionalLibrary("PhysXVisualDebuggerSDKPROFILE.lib");
-						InModule.AddPublicAdditionalLibrary("PhysXProfileSDKPROFILE.lib");
-					}
-					else // Test or Shipping
-					{
-						InModule.AddPublicAdditionalLibrary("PhysX3.lib");
-						InModule.AddPublicAdditionalLibrary("PhysX3Extensions.lib");
-						InModule.AddPublicAdditionalLibrary("PhysX3Cooking.lib");
-						InModule.AddPublicAdditionalLibrary("PhysX3Common.lib");
-						InModule.AddPublicAdditionalLibrary("PhysX3Vehicle.lib");
-						InModule.AddPublicAdditionalLibrary("PxTask.lib");
-						InModule.AddPublicAdditionalLibrary("PhysXVisualDebuggerSDK.lib");
-						InModule.AddPublicAdditionalLibrary("PhysXProfileSDK.lib");
-					}
-				}
-				else if (InModule.ToString() == "APEX")
-				{
-					InModule.RemovePublicDefinition("APEX_STATICALLY_LINKED=0");
-					InModule.AddPublicDefinition("APEX_STATICALLY_LINKED=1");
-					
-					string APEXDir = UEBuildConfiguration.UEThirdPartySourceDirectory + "PhysX/APEX-1.3/";
-					if (Target.Platform == UnrealTargetPlatform.WinRT)
-					{
-						InModule.AddPublicLibraryPath(APEXDir + "lib/WinRT");
-					}
-					else
-					{
-						InModule.AddPublicLibraryPath(APEXDir + "lib/WinRT/ARM");
-					}
-
-					if (Target.Configuration == UnrealTargetConfiguration.Debug)
-					{
-						InModule.AddPublicAdditionalLibrary("ApexCommonDEBUG.lib");
-						InModule.AddPublicAdditionalLibrary("ApexFrameworkDEBUG.lib");
-						InModule.AddPublicAdditionalLibrary("ApexSharedDEBUG.lib");
-						InModule.AddPublicAdditionalLibrary("APEX_DestructibleDEBUG.lib");
-
-					}
-					else if (Target.Configuration == UnrealTargetConfiguration.Development)
-					{
-						InModule.AddPublicAdditionalLibrary("ApexCommonPROFILE.lib");
-						InModule.AddPublicAdditionalLibrary("ApexFrameworkPROFILE.lib");
-						InModule.AddPublicAdditionalLibrary("ApexSharedPROFILE.lib");
-						InModule.AddPublicAdditionalLibrary("APEX_DestructiblePROFILE.lib");
-					}
-					else // Test or Shipping
-					{
-						InModule.AddPublicAdditionalLibrary("ApexCommon.lib");
-						InModule.AddPublicAdditionalLibrary("ApexFramework.lib");
-						InModule.AddPublicAdditionalLibrary("ApexShared.lib");
-						InModule.AddPublicAdditionalLibrary("APEX_Destructible.lib");
-					}
-				}
-				else if (InModule.ToString() == "FreeType2")
-				{
-					string FreeType2Path = UEBuildConfiguration.UEThirdPartySourceDirectory + "FreeType2/FreeType2-2.4.12/";
-					if (Target.Platform == UnrealTargetPlatform.WinRT)
-					{
-						InModule.AddPublicLibraryPath(FreeType2Path + "Lib/WinRT/Win64");
-					}
-					else
-					{
-						InModule.AddPublicLibraryPath(FreeType2Path + "Lib/WinRT/ARM");
-					}
-					InModule.AddPublicAdditionalLibrary("freetype2412MT.lib");
-				}
-				else if (InModule.ToString() == "UElibPNG")
-				{
-					string libPNGPath = UEBuildConfiguration.UEThirdPartySourceDirectory + "libPNG/libPNG-1.5.2";
-					if (Target.Platform == UnrealTargetPlatform.WinRT)
-					{
-						InModule.AddPublicLibraryPath(libPNGPath + "/lib/WinRT/Win64");
-					}
-					else
-					{
-						InModule.AddPublicLibraryPath(libPNGPath + "/lib/WinRT/ARM");
-					}
-					InModule.AddPublicAdditionalLibrary("libpng125.lib");
-				}
-				else if (InModule.ToString() == "DX11")
-				{
-					// Clear out all the Windows include paths and libraries...
-					// The WinRTSDK module handles proper paths and libs for WinRT.
-					// However, the D3D11RHI module will include the DX11 module.
-					InModule.ClearPublicIncludePaths();
-					InModule.ClearPublicLibraryPaths();
-					InModule.ClearPublicAdditionalLibraries();
-					InModule.RemovePublicDefinition("WITH_D3DX_LIBS=1");
-					InModule.AddPublicDefinition("D3D11_WITH_DWMAPI=0");
-					InModule.AddPublicDefinition("WITH_D3DX_LIBS=0");
-					InModule.AddPublicDefinition("WITH_DX_PERF=0");
- 					InModule.RemovePublicAdditionalLibrary("X3DAudio.lib");
- 					InModule.RemovePublicAdditionalLibrary("XAPOFX.lib");
-				}
-				else if (InModule.ToString() == "XInput")
-				{
-					InModule.AddPublicAdditionalLibrary("XInput.lib");
-				}
-				else if (InModule.ToString() == "XAudio2")
-				{
-					InModule.AddPublicDefinition("XAUDIO_SUPPORTS_XMA2WAVEFORMATEX=0");
-					InModule.AddPublicDefinition("XAUDIO_SUPPORTS_DEVICE_DETAILS=0");
-					InModule.AddPublicDefinition("XAUDIO2_SUPPORTS_MUSIC=0");
-					InModule.AddPublicDefinition("XAUDIO2_SUPPORTS_SENDLIST=0");
-					InModule.AddPublicAdditionalLibrary("XAudio2.lib");
-					// Hack to enable AllowWindowsPlatformTypes.h/HideWindowsPlatformTypes.h
-					InModule.AddPublicIncludePath("Runtime/Core/Public/Windows");
-				}
-				else if (InModule.ToString() == "UEOgg")
-				{
-					string OggPath = UEBuildConfiguration.UEThirdPartySourceDirectory + "Ogg/libogg-1.2.2/";
-					if (Target.Platform == UnrealTargetPlatform.WinRT)
-					{
-						InModule.AddPublicLibraryPath(OggPath + "WinRT/VS2012/WinRT/x64/Release");
-					}
-					else
-					{
-						InModule.AddPublicLibraryPath(OggPath + "WinRT/VS2012/WinRT/ARM/Release");
-					}
-					InModule.AddPublicAdditionalLibrary("libogg_static.lib");
-				}
-				else if (InModule.ToString() == "Vorbis")
-				{
-					string VorbisPath = UEBuildConfiguration.UEThirdPartySourceDirectory + "Vorbis/libvorbis-1.3.2/";
-					if (Target.Platform == UnrealTargetPlatform.WinRT)
-					{
-						InModule.AddPublicLibraryPath(VorbisPath + "WinRT/VS2012/WinRT/x64/Release");
-					}
-					else
-					{
-						InModule.AddPublicLibraryPath(VorbisPath + "WinRT/VS2012/WinRT/ARM/Release");
-					}
-					InModule.AddPublicAdditionalLibrary("libvorbis_static.lib");
-				}
-				else if (InModule.ToString() == "VorbisFile")
-				{
-					string VorbisPath = UEBuildConfiguration.UEThirdPartySourceDirectory + "Vorbis/libvorbis-1.3.2/";
-					if (Target.Platform == UnrealTargetPlatform.WinRT)
-					{
-						InModule.AddPublicLibraryPath(VorbisPath + "WinRT/VS2012/WinRT/x64/Release");
-					}
-					else
-					{
-						InModule.AddPublicLibraryPath(VorbisPath + "WinRT/VS2012/WinRT/ARM/Release");
-					}
-					InModule.AddPublicAdditionalLibrary("libvorbisfile_static.lib");
-				}
-				else if (InModule.ToString() == "DX11Audio")
-				{
-					InModule.RemovePublicAdditionalLibrary("X3DAudio.lib");
-					InModule.RemovePublicAdditionalLibrary("XAPOFX.lib");
-				}
-				else if (InModule.ToString() == "zlib")
-				{
-					if (Target.Platform == UnrealTargetPlatform.WinRT)
-					{
-						InModule.AddPublicLibraryPath(UEBuildConfiguration.UEThirdPartySourceDirectory + "zlib/zlib-1.2.5/Lib/WinRT/Win64");
-					}
-					else
-					{
-						InModule.AddPublicLibraryPath(UEBuildConfiguration.UEThirdPartySourceDirectory + "zlib/zlib-1.2.5/Lib/WinRT/ARM");
-					}
-					InModule.AddPublicAdditionalLibrary("zlib125.lib");
-				}
-			}
-			else if ((Target.Platform == UnrealTargetPlatform.Win32) || (Target.Platform == UnrealTargetPlatform.Win64))
-			{
+        /**
+         *	Whether the required external SDKs are installed for this platform
+         */
+        protected override SDKStatus HasRequiredManualSDKInternal()
+        {
+            return IsVisualStudioInstalled() ? SDKStatus.Valid : SDKStatus.Invalid;
+        }
+
+        /**
+         *	Register the platform with the UEBuildPlatform class
+         */
+        protected override void RegisterBuildPlatformInternal()
+        {
+            [REDACTED_EMAIL]: Add platform support
+            if (UnrealBuildTool.RunningRocket() || UnrealBuildTool.BuildingRocket() || Utils.IsRunningOnMono)
+            {
+                return;
+            }
+
+            if ((ProjectFileGenerator.bGenerateProjectFiles == true) || (IsVisualStudioInstalled() == true))
+            {
+                bool bRegisterBuildPlatform = true;
+
+                // We also need to check for the generated projects... to handle the case where someone generates projects w/out WinRT.
+                // Hardcoding this for now - but ideally it would be dynamically discovered.
+                string EngineSourcePath = Path.Combine(ProjectFileGenerator.EngineRelativePath, "Source");
+                string WinRTRHIFile = Path.Combine(EngineSourcePath, "Runtime", "Windows", "D3D11RHI", "D3D11RHI.build.cs");
+                if (File.Exists(WinRTRHIFile) == false)
+                {
+                    bRegisterBuildPlatform = false;
+                }
+
+                if (bRegisterBuildPlatform == true)
+                {
+                    // Register this build platform for WinRT
+                    Log.TraceVerbose("        Registering for {0}", UnrealTargetPlatform.WinRT.ToString());
+                    UEBuildPlatform.RegisterBuildPlatform(UnrealTargetPlatform.WinRT, this);
+                    UEBuildPlatform.RegisterPlatformWithGroup(UnrealTargetPlatform.WinRT, UnrealPlatformGroup.Microsoft);
+
+                    // For now only register WinRT_ARM is truly a Windows 8 machine.
+                    // This will prevent people who do all platform builds from running into the compiler issue.
+                    if (WinRTPlatform.IsWindows8() == true)
+                    {
+                        Log.TraceVerbose("        Registering for {0}", UnrealTargetPlatform.WinRT_ARM.ToString());
+                        UEBuildPlatform.RegisterBuildPlatform(UnrealTargetPlatform.WinRT_ARM, this);
+                        UEBuildPlatform.RegisterPlatformWithGroup(UnrealTargetPlatform.WinRT_ARM, UnrealPlatformGroup.Microsoft);
+                    }
+                }
+            }
+        }
+
+        /**
+         *	Retrieve the CPPTargetPlatform for the given UnrealTargetPlatform
+         *
+         *	@param	InUnrealTargetPlatform		The UnrealTargetPlatform being build
+         *	
+         *	@return	CPPTargetPlatform			The CPPTargetPlatform to compile for
+         */
+        public override CPPTargetPlatform GetCPPTargetPlatform(UnrealTargetPlatform InUnrealTargetPlatform)
+        {
+            switch (InUnrealTargetPlatform)
+            {
+                case UnrealTargetPlatform.WinRT:
+                    return CPPTargetPlatform.WinRT;
+                case UnrealTargetPlatform.WinRT_ARM:
+                    return CPPTargetPlatform.WinRT_ARM;
+            }
+            throw new BuildException("WinRTPlatform::GetCPPTargetPlatform: Invalid request for {0}", InUnrealTargetPlatform.ToString());
+        }
+
+        /**
+         *	Get the extension to use for the given binary type
+         *	
+         *	@param	InBinaryType		The binary type being built
+         *	
+         *	@return	string				The binary extension (i.e. 'exe' or 'dll')
+         */
+        public override string GetBinaryExtension(UEBuildBinaryType InBinaryType)
+        {
+            switch (InBinaryType)
+            {
+                case UEBuildBinaryType.DynamicLinkLibrary:
+                    return ".dll";
+                case UEBuildBinaryType.Executable:
+                    return ".exe";
+                case UEBuildBinaryType.StaticLibrary:
+                    return ".lib";
+                case UEBuildBinaryType.Object:
+                    return ".obj";
+                case UEBuildBinaryType.PrecompiledHeader:
+                    return ".pch";
+            }
+            return base.GetBinaryExtension(InBinaryType);
+        }
+
+        /**
+         *	Get the extension to use for debug info for the given binary type
+         *	
+         *	@param	InBinaryType		The binary type being built
+         *	
+         *	@return	string				The debug info extension (i.e. 'pdb')
+         */
+        public override string GetDebugInfoExtension(UEBuildBinaryType InBinaryType)
+        {
+            return ".pdb";
+        }
+
+        /**
+         *	Gives the platform a chance to 'override' the configuration settings 
+         *	that are overridden on calls to RunUBT.
+         *	
+         *	@param	InPlatform			The UnrealTargetPlatform being built
+         *	@param	InConfiguration		The UnrealTargetConfiguration being built
+         *	
+         *	@return	bool				true if debug info should be generated, false if not
+         */
+        public override void ResetBuildConfiguration(UnrealTargetPlatform InPlatform, UnrealTargetConfiguration InConfiguration)
+        {
+            ValidateUEBuildConfiguration();
+        }
+
+        /**
+         *	Validate the UEBuildConfiguration for this platform
+         *	This is called BEFORE calling UEBuildConfiguration to allow setting 
+         *	various fields used in that function such as CompileLeanAndMean...
+         */
+        public override void ValidateUEBuildConfiguration()
+        {
+            UEBuildConfiguration.bCompileLeanAndMeanUE = true;
+            UEBuildConfiguration.bCompilePhysX = false;
+            UEBuildConfiguration.bCompileAPEX = false;
+
+            BuildConfiguration.bUseUnityBuild = false;
+
+            // Don't stop compilation at first error...
+            BuildConfiguration.bStopXGECompilationAfterErrors = false;
+        }
+
+        /**
+         * Whether this platform should build a monolithic binary
+         */
+        public override bool ShouldCompileMonolithicBinary(UnrealTargetPlatform InPlatform)
+        {
+            return true;
+        }
+
+        /**
+         *	Whether PDB files should be used
+         *	
+         *	@param	InPlatform			The CPPTargetPlatform being built
+         *	@param	InConfiguration		The CPPTargetConfiguration being built
+         *	@param	bInCreateDebugInfo	true if debug info is getting create, false if not
+         *	
+         *	@return	bool	true if PDB files should be used, false if not
+         */
+        public override bool ShouldUsePDBFiles(CPPTargetPlatform Platform, CPPTargetConfiguration Configuration, bool bCreateDebugInfo)
+        {
+            return true;
+        }
+
+        /**
+         *	Whether the editor should be built for this platform or not
+         *	
+         *	@param	InPlatform		The UnrealTargetPlatform being built
+         *	@param	InConfiguration	The UnrealTargetConfiguration being built
+         *	@return	bool			true if the editor should be built, false if not
+         */
+        public override bool ShouldNotBuildEditor(UnrealTargetPlatform InPlatform, UnrealTargetConfiguration InConfiguration)
+        {
+            return true;
+        }
+
+        /**
+         *	Whether the platform requires cooked data
+         *	
+         *	@param	InPlatform		The UnrealTargetPlatform being built
+         *	@param	InConfiguration	The UnrealTargetConfiguration being built
+         *	@return	bool			true if the platform requires cooked data, false if not
+         */
+        public override bool BuildRequiresCookedData(UnrealTargetPlatform InPlatform, UnrealTargetConfiguration InConfiguration)
+        {
+            return true;
+        }
+
+        /**
+         *	Get a list of extra modules the platform requires.
+         *	This is to allow undisclosed platforms to add modules they need without exposing information about the platfomr.
+         *	
+         *	@param	Target						The target being build
+         *	@param	BuildTarget					The UEBuildTarget getting build
+         *	@param	PlatformExtraModules		OUTPUT the list of extra modules the platform needs to add to the target
+         */
+        public override void GetExtraModules(TargetInfo Target, UEBuildTarget BuildTarget, ref List<string> PlatformExtraModules)
+        {
+            if (Target.Platform == UnrealTargetPlatform.WinRT)
+            {
+                PlatformExtraModules.Add("XAudio2");
+            }
+        }
+
+        /**
+         *	Modify the newly created module passed in for this platform.
+         *	This is not required - but allows for hiding details of a
+         *	particular platform.
+         *	
+         *	@param	InModule		The newly loaded module
+         *	@param	Target			The target being build
+         */
+        public override void ModifyNewlyLoadedModule(UEBuildModule InModule, TargetInfo Target)
+        {
+            if ((Target.Platform == UnrealTargetPlatform.WinRT) || 
+                (Target.Platform == UnrealTargetPlatform.WinRT_ARM))
+            {
+                if (InModule.ToString() == "Core")
+                {
+                    InModule.AddPublicIncludePath("Runtime/Core/Public/WinRT");
+                    InModule.AddPublicDependencyModule("zlib");
+                }
+                else if (InModule.ToString() == "Engine")
+                {
+                    InModule.AddPrivateDependencyModule("zlib");
+                    InModule.AddPrivateDependencyModule("UElibPNG");
+                    InModule.AddPublicDependencyModule("UEOgg");
+                    InModule.AddPublicDependencyModule("Vorbis");
+                }
+                else if (InModule.ToString() == "Launch")
+                {
+                }
+                else if (InModule.ToString() == "D3D11RHI")
+                {
+                    InModule.AddPublicDefinition("D3D11_CUSTOM_VIEWPORT_CONSTRUCTOR=1");
+                    // To enable platform specific D3D11 RHI Types
+                    InModule.AddPrivateIncludePath("Runtime/Windows/D3D11RHI/Private/WinRT");
+                    // Hack to enable AllowWindowsPlatformTypes.h/HideWindowsPlatformTypes.h
+                    InModule.AddPublicIncludePath("Runtime/Core/Public/Windows");
+                }
+                else if (InModule.ToString() == "Sockets")
+                {
+                    // Hack to enable AllowWindowsPlatformTypes.h/HideWindowsPlatformTypes.h
+                    InModule.AddPublicIncludePath("Runtime/Core/Public/Windows");
+                }
+                else if (InModule.ToString() == "PhysX")
+                {
+                    string PhysXDir = UEBuildConfiguration.UEThirdPartySourceDirectory + "PhysX/PhysX-3.3/";
+
+                    InModule.AddPublicIncludePath("include/foundation/WinRT");
+                    if (Target.Platform == UnrealTargetPlatform.WinRT)
+                    {
+                        InModule.AddPublicLibraryPath(PhysXDir + "Lib/WinRT");
+                    }
+                    else
+                    {
+                        InModule.AddPublicLibraryPath(PhysXDir + "Lib/WinRT/ARM");
+                    }
+
+                    if (Target.Configuration == UnrealTargetConfiguration.Debug)
+                    {
+                        InModule.AddPublicAdditionalLibrary("PhysX3DEBUG.lib");
+                        InModule.AddPublicAdditionalLibrary("PhysX3ExtensionsDEBUG.lib");
+                        InModule.AddPublicAdditionalLibrary("PhysX3CookingDEBUG.lib");
+                        InModule.AddPublicAdditionalLibrary("PhysX3CommonDEBUG.lib");
+                        InModule.AddPublicAdditionalLibrary("PhysX3VehicleDEBUG.lib");
+                        InModule.AddPublicAdditionalLibrary("PxTaskDEBUG.lib");
+                        InModule.AddPublicAdditionalLibrary("PhysXVisualDebuggerSDKDEBUG.lib");
+                        InModule.AddPublicAdditionalLibrary("PhysXProfileSDKDEBUG.lib");
+                    }
+                    else if (Target.Configuration == UnrealTargetConfiguration.Development)
+                    {
+                        InModule.AddPublicAdditionalLibrary("PhysX3PROFILE.lib");
+                        InModule.AddPublicAdditionalLibrary("PhysX3ExtensionsPROFILE.lib");
+                        InModule.AddPublicAdditionalLibrary("PhysX3CookingPROFILE.lib");
+                        InModule.AddPublicAdditionalLibrary("PhysX3CommonPROFILE.lib");
+                        InModule.AddPublicAdditionalLibrary("PhysX3VehiclePROFILE.lib");
+                        InModule.AddPublicAdditionalLibrary("PxTaskPROFILE.lib");
+                        InModule.AddPublicAdditionalLibrary("PhysXVisualDebuggerSDKPROFILE.lib");
+                        InModule.AddPublicAdditionalLibrary("PhysXProfileSDKPROFILE.lib");
+                    }
+                    else // Test or Shipping
+                    {
+                        InModule.AddPublicAdditionalLibrary("PhysX3.lib");
+                        InModule.AddPublicAdditionalLibrary("PhysX3Extensions.lib");
+                        InModule.AddPublicAdditionalLibrary("PhysX3Cooking.lib");
+                        InModule.AddPublicAdditionalLibrary("PhysX3Common.lib");
+                        InModule.AddPublicAdditionalLibrary("PhysX3Vehicle.lib");
+                        InModule.AddPublicAdditionalLibrary("PxTask.lib");
+                        InModule.AddPublicAdditionalLibrary("PhysXVisualDebuggerSDK.lib");
+                        InModule.AddPublicAdditionalLibrary("PhysXProfileSDK.lib");
+                    }
+                }
+                else if (InModule.ToString() == "APEX")
+                {
+                    InModule.RemovePublicDefinition("APEX_STATICALLY_LINKED=0");
+                    InModule.AddPublicDefinition("APEX_STATICALLY_LINKED=1");
+                    
+                    string APEXDir = UEBuildConfiguration.UEThirdPartySourceDirectory + "PhysX/APEX-1.3/";
+                    if (Target.Platform == UnrealTargetPlatform.WinRT)
+                    {
+                        InModule.AddPublicLibraryPath(APEXDir + "lib/WinRT");
+                    }
+                    else
+                    {
+                        InModule.AddPublicLibraryPath(APEXDir + "lib/WinRT/ARM");
+                    }
+
+                    if (Target.Configuration == UnrealTargetConfiguration.Debug)
+                    {
+                        InModule.AddPublicAdditionalLibrary("ApexCommonDEBUG.lib");
+                        InModule.AddPublicAdditionalLibrary("ApexFrameworkDEBUG.lib");
+                        InModule.AddPublicAdditionalLibrary("ApexSharedDEBUG.lib");
+                        InModule.AddPublicAdditionalLibrary("APEX_DestructibleDEBUG.lib");
+
+                    }
+                    else if (Target.Configuration == UnrealTargetConfiguration.Development)
+                    {
+                        InModule.AddPublicAdditionalLibrary("ApexCommonPROFILE.lib");
+                        InModule.AddPublicAdditionalLibrary("ApexFrameworkPROFILE.lib");
+                        InModule.AddPublicAdditionalLibrary("ApexSharedPROFILE.lib");
+                        InModule.AddPublicAdditionalLibrary("APEX_DestructiblePROFILE.lib");
+                    }
+                    else // Test or Shipping
+                    {
+                        InModule.AddPublicAdditionalLibrary("ApexCommon.lib");
+                        InModule.AddPublicAdditionalLibrary("ApexFramework.lib");
+                        InModule.AddPublicAdditionalLibrary("ApexShared.lib");
+                        InModule.AddPublicAdditionalLibrary("APEX_Destructible.lib");
+                    }
+                }
+                else if (InModule.ToString() == "FreeType2")
+                {
+                    string FreeType2Path = UEBuildConfiguration.UEThirdPartySourceDirectory + "FreeType2/FreeType2-2.4.12/";
+                    if (Target.Platform == UnrealTargetPlatform.WinRT)
+                    {
+                        InModule.AddPublicLibraryPath(FreeType2Path + "Lib/WinRT/Win64");
+                    }
+                    else
+                    {
+                        InModule.AddPublicLibraryPath(FreeType2Path + "Lib/WinRT/ARM");
+                    }
+                    InModule.AddPublicAdditionalLibrary("freetype2412MT.lib");
+                }
+                else if (InModule.ToString() == "UElibPNG")
+                {
+                    string libPNGPath = UEBuildConfiguration.UEThirdPartySourceDirectory + "libPNG/libPNG-1.5.2";
+                    if (Target.Platform == UnrealTargetPlatform.WinRT)
+                    {
+                        InModule.AddPublicLibraryPath(libPNGPath + "/lib/WinRT/Win64");
+                    }
+                    else
+                    {
+                        InModule.AddPublicLibraryPath(libPNGPath + "/lib/WinRT/ARM");
+                    }
+                    InModule.AddPublicAdditionalLibrary("libpng125.lib");
+                }
+                else if (InModule.ToString() == "DX11")
+                {
+                    // Clear out all the Windows include paths and libraries...
+                    // The WinRTSDK module handles proper paths and libs for WinRT.
+                    // However, the D3D11RHI module will include the DX11 module.
+                    InModule.ClearPublicIncludePaths();
+                    InModule.ClearPublicLibraryPaths();
+                    InModule.ClearPublicAdditionalLibraries();
+                    InModule.RemovePublicDefinition("WITH_D3DX_LIBS=1");
+                    InModule.AddPublicDefinition("D3D11_WITH_DWMAPI=0");
+                    InModule.AddPublicDefinition("WITH_D3DX_LIBS=0");
+                    InModule.AddPublicDefinition("WITH_DX_PERF=0");
+                    InModule.RemovePublicAdditionalLibrary("X3DAudio.lib");
+                    InModule.RemovePublicAdditionalLibrary("XAPOFX.lib");
+                }
+                else if (InModule.ToString() == "XInput")
+                {
+                    InModule.AddPublicAdditionalLibrary("XInput.lib");
+                }
+                else if (InModule.ToString() == "XAudio2")
+                {
+                    InModule.AddPublicDefinition("XAUDIO_SUPPORTS_XMA2WAVEFORMATEX=0");
+                    InModule.AddPublicDefinition("XAUDIO_SUPPORTS_DEVICE_DETAILS=0");
+                    InModule.AddPublicDefinition("XAUDIO2_SUPPORTS_MUSIC=0");
+                    InModule.AddPublicDefinition("XAUDIO2_SUPPORTS_SENDLIST=0");
+                    InModule.AddPublicAdditionalLibrary("XAudio2.lib");
+                    // Hack to enable AllowWindowsPlatformTypes.h/HideWindowsPlatformTypes.h
+                    InModule.AddPublicIncludePath("Runtime/Core/Public/Windows");
+                }
+                else if (InModule.ToString() == "UEOgg")
+                {
+                    string OggPath = UEBuildConfiguration.UEThirdPartySourceDirectory + "Ogg/libogg-1.2.2/";
+                    if (Target.Platform == UnrealTargetPlatform.WinRT)
+                    {
+                        InModule.AddPublicLibraryPath(OggPath + "WinRT/VS2012/WinRT/x64/Release");
+                    }
+                    else
+                    {
+                        InModule.AddPublicLibraryPath(OggPath + "WinRT/VS2012/WinRT/ARM/Release");
+                    }
+                    InModule.AddPublicAdditionalLibrary("libogg_static.lib");
+                }
+                else if (InModule.ToString() == "Vorbis")
+                {
+                    string VorbisPath = UEBuildConfiguration.UEThirdPartySourceDirectory + "Vorbis/libvorbis-1.3.2/";
+                    if (Target.Platform == UnrealTargetPlatform.WinRT)
+                    {
+                        InModule.AddPublicLibraryPath(VorbisPath + "WinRT/VS2012/WinRT/x64/Release");
+                    }
+                    else
+                    {
+                        InModule.AddPublicLibraryPath(VorbisPath + "WinRT/VS2012/WinRT/ARM/Release");
+                    }
+                    InModule.AddPublicAdditionalLibrary("libvorbis_static.lib");
+                }
+                else if (InModule.ToString() == "VorbisFile")
+                {
+                    string VorbisPath = UEBuildConfiguration.UEThirdPartySourceDirectory + "Vorbis/libvorbis-1.3.2/";
+                    if (Target.Platform == UnrealTargetPlatform.WinRT)
+                    {
+                        InModule.AddPublicLibraryPath(VorbisPath + "WinRT/VS2012/WinRT/x64/Release");
+                    }
+                    else
+                    {
+                        InModule.AddPublicLibraryPath(VorbisPath + "WinRT/VS2012/WinRT/ARM/Release");
+                    }
+                    InModule.AddPublicAdditionalLibrary("libvorbisfile_static.lib");
+                }
+                else if (InModule.ToString() == "DX11Audio")
+                {
+                    InModule.RemovePublicAdditionalLibrary("X3DAudio.lib");
+                    InModule.RemovePublicAdditionalLibrary("XAPOFX.lib");
+                }
+                else if (InModule.ToString() == "zlib")
+                {
+                    if (Target.Platform == UnrealTargetPlatform.WinRT)
+                    {
+                        InModule.AddPublicLibraryPath(UEBuildConfiguration.UEThirdPartySourceDirectory + "zlib/zlib-1.2.5/Lib/WinRT/Win64");
+                    }
+                    else
+                    {
+                        InModule.AddPublicLibraryPath(UEBuildConfiguration.UEThirdPartySourceDirectory + "zlib/zlib-1.2.5/Lib/WinRT/ARM");
+                    }
+                    InModule.AddPublicAdditionalLibrary("zlib125.lib");
+                }
+            }
+            else if ((Target.Platform == UnrealTargetPlatform.Win32) || (Target.Platform == UnrealTargetPlatform.Win64))
+            {
 //              bool bBuildShaderFormats = UEBuildConfiguration.bForceBuildShaderFormats;
 // 				if (!UEBuildConfiguration.bBuildRequiresCookedData)
 // 				{
@@ -507,148 +507,148 @@ namespace UnrealBuildTool
 //              {
                 //                  InModule.AddPlatformSpecificDynamicallyLoadedModule("ShaderFormatWinRT");
 //              }
-			}
-		}
-
-		/**
-		 *	Setup the target environment for building
-		 *	
-		 *	@param	InBuildTarget		The target being built
-		 */
-		public override void SetUpEnvironment(UEBuildTarget InBuildTarget)
-		{
-			InBuildTarget.GlobalCompileEnvironment.Config.Definitions.Add("PLATFORM_DESKTOP=0");
-			if (InBuildTarget.Platform == UnrealTargetPlatform.WinRT)
-			{
-				InBuildTarget.GlobalCompileEnvironment.Config.Definitions.Add("PLATFORM_64BITS=1");
-			}
-			else
-			{
-				//WinRT_ARM
-				InBuildTarget.GlobalCompileEnvironment.Config.Definitions.Add("PLATFORM_64BITS=0");
-			}
-			InBuildTarget.GlobalCompileEnvironment.Config.Definitions.Add("EXCEPTIONS_DISABLED=0");
-			InBuildTarget.GlobalCompileEnvironment.Config.Definitions.Add("WITH_OGGVORBIS=1");
-			InBuildTarget.GlobalCompileEnvironment.Config.Definitions.Add("PLATFORM_CAN_SUPPORT_EDITORONLY_DATA=0");
-
-			[REDACTED_EMAIL]: REMOVE THIS
-			// Temporary disabling of automation... to get things working
-			InBuildTarget.GlobalCompileEnvironment.Config.Definitions.Add("WITH_AUTOMATION_WORKER=0");
-
-			[REDACTED_EMAIL]: Add support for version enforcement
-			//InBuildTarget.GlobalCompileEnvironment.Config.Definitions.Add("MIN_EXPECTED_XDK_VER=1070");
-			//InBuildTarget.GlobalCompileEnvironment.Config.Definitions.Add("MAX_EXPECTED_XDK_VER=1070");
-			InBuildTarget.GlobalCompileEnvironment.Config.Definitions.Add("UNICODE");
-			InBuildTarget.GlobalCompileEnvironment.Config.Definitions.Add("_UNICODE");
-			//InBuildTarget.GlobalCompileEnvironment.Config.Definitions.Add("WIN32_LEAN_AND_MEAN");
-			InBuildTarget.GlobalCompileEnvironment.Config.Definitions.Add("WINAPI_FAMILY=WINAPI_FAMILY_APP");
-
-			if (InBuildTarget.Platform == UnrealTargetPlatform.WinRT)
-			{
-				InBuildTarget.GlobalCompileEnvironment.Config.Definitions.Add("WINRT=1");
-				InBuildTarget.GlobalCompileEnvironment.Config.Definitions.Add("PLATFORM_WINRT=1");
-			}
-			else
-			{
-				//WinRT_ARM
-				InBuildTarget.GlobalCompileEnvironment.Config.Definitions.Add("WINRT=1");
-				InBuildTarget.GlobalCompileEnvironment.Config.Definitions.Add("PLATFORM_WINRT_ARM=1");
-				InBuildTarget.GlobalCompileEnvironment.Config.Definitions.Add("PLATFORM_WINRT=1");
-				// 				InBuildTarget.GlobalCompileEnvironment.Config.Definitions.Add("WinRT_ARM=1");
-			}
-
-			// No D3DX on WinRT!
-			InBuildTarget.GlobalCompileEnvironment.Config.Definitions.Add("NO_D3DX_LIBS=1");
-
-			// Explicitly exclude the MS C++ runtime libraries we're not using, to ensure other libraries we link with use the same
-			// runtime library as the engine.
-			if (InBuildTarget.Configuration == UnrealTargetConfiguration.Debug)
-			{
-				InBuildTarget.GlobalLinkEnvironment.Config.ExcludedLibraries.Add("MSVCRT");
-				InBuildTarget.GlobalLinkEnvironment.Config.ExcludedLibraries.Add("MSVCPRT");
-			}
-			else
-			{
-				InBuildTarget.GlobalLinkEnvironment.Config.ExcludedLibraries.Add("MSVCRTD");
-				InBuildTarget.GlobalLinkEnvironment.Config.ExcludedLibraries.Add("MSVCPRTD");
-			}
-			InBuildTarget.GlobalLinkEnvironment.Config.ExcludedLibraries.Add("LIBC");
-			InBuildTarget.GlobalLinkEnvironment.Config.ExcludedLibraries.Add("LIBCMT");
-			InBuildTarget.GlobalLinkEnvironment.Config.ExcludedLibraries.Add("LIBCPMT");
-			InBuildTarget.GlobalLinkEnvironment.Config.ExcludedLibraries.Add("LIBCP");
-			InBuildTarget.GlobalLinkEnvironment.Config.ExcludedLibraries.Add("LIBCD");
-			InBuildTarget.GlobalLinkEnvironment.Config.ExcludedLibraries.Add("LIBCMTD");
-			InBuildTarget.GlobalLinkEnvironment.Config.ExcludedLibraries.Add("LIBCPMTD");
-			InBuildTarget.GlobalLinkEnvironment.Config.ExcludedLibraries.Add("LIBCPD");
-
-			// Compile and link with Win32 API libraries.
+            }
+        }
+
+        /**
+         *	Setup the target environment for building
+         *	
+         *	@param	InBuildTarget		The target being built
+         */
+        public override void SetUpEnvironment(UEBuildTarget InBuildTarget)
+        {
+            InBuildTarget.GlobalCompileEnvironment.Config.Definitions.Add("PLATFORM_DESKTOP=0");
+            if (InBuildTarget.Platform == UnrealTargetPlatform.WinRT)
+            {
+                InBuildTarget.GlobalCompileEnvironment.Config.Definitions.Add("PLATFORM_64BITS=1");
+            }
+            else
+            {
+                //WinRT_ARM
+                InBuildTarget.GlobalCompileEnvironment.Config.Definitions.Add("PLATFORM_64BITS=0");
+            }
+            InBuildTarget.GlobalCompileEnvironment.Config.Definitions.Add("EXCEPTIONS_DISABLED=0");
+            InBuildTarget.GlobalCompileEnvironment.Config.Definitions.Add("WITH_OGGVORBIS=1");
+            InBuildTarget.GlobalCompileEnvironment.Config.Definitions.Add("PLATFORM_CAN_SUPPORT_EDITORONLY_DATA=0");
+
+            [REDACTED_EMAIL]: REMOVE THIS
+            // Temporary disabling of automation... to get things working
+            InBuildTarget.GlobalCompileEnvironment.Config.Definitions.Add("WITH_AUTOMATION_WORKER=0");
+
+            [REDACTED_EMAIL]: Add support for version enforcement
+            //InBuildTarget.GlobalCompileEnvironment.Config.Definitions.Add("MIN_EXPECTED_XDK_VER=1070");
+            //InBuildTarget.GlobalCompileEnvironment.Config.Definitions.Add("MAX_EXPECTED_XDK_VER=1070");
+            InBuildTarget.GlobalCompileEnvironment.Config.Definitions.Add("UNICODE");
+            InBuildTarget.GlobalCompileEnvironment.Config.Definitions.Add("_UNICODE");
+            //InBuildTarget.GlobalCompileEnvironment.Config.Definitions.Add("WIN32_LEAN_AND_MEAN");
+            InBuildTarget.GlobalCompileEnvironment.Config.Definitions.Add("WINAPI_FAMILY=WINAPI_FAMILY_APP");
+
+            if (InBuildTarget.Platform == UnrealTargetPlatform.WinRT)
+            {
+                InBuildTarget.GlobalCompileEnvironment.Config.Definitions.Add("WINRT=1");
+                InBuildTarget.GlobalCompileEnvironment.Config.Definitions.Add("PLATFORM_WINRT=1");
+            }
+            else
+            {
+                //WinRT_ARM
+                InBuildTarget.GlobalCompileEnvironment.Config.Definitions.Add("WINRT=1");
+                InBuildTarget.GlobalCompileEnvironment.Config.Definitions.Add("PLATFORM_WINRT_ARM=1");
+                InBuildTarget.GlobalCompileEnvironment.Config.Definitions.Add("PLATFORM_WINRT=1");
+                // 				InBuildTarget.GlobalCompileEnvironment.Config.Definitions.Add("WinRT_ARM=1");
+            }
+
+            // No D3DX on WinRT!
+            InBuildTarget.GlobalCompileEnvironment.Config.Definitions.Add("NO_D3DX_LIBS=1");
+
+            // Explicitly exclude the MS C++ runtime libraries we're not using, to ensure other libraries we link with use the same
+            // runtime library as the engine.
+            if (InBuildTarget.Configuration == UnrealTargetConfiguration.Debug)
+            {
+                InBuildTarget.GlobalLinkEnvironment.Config.ExcludedLibraries.Add("MSVCRT");
+                InBuildTarget.GlobalLinkEnvironment.Config.ExcludedLibraries.Add("MSVCPRT");
+            }
+            else
+            {
+                InBuildTarget.GlobalLinkEnvironment.Config.ExcludedLibraries.Add("MSVCRTD");
+                InBuildTarget.GlobalLinkEnvironment.Config.ExcludedLibraries.Add("MSVCPRTD");
+            }
+            InBuildTarget.GlobalLinkEnvironment.Config.ExcludedLibraries.Add("LIBC");
+            InBuildTarget.GlobalLinkEnvironment.Config.ExcludedLibraries.Add("LIBCMT");
+            InBuildTarget.GlobalLinkEnvironment.Config.ExcludedLibraries.Add("LIBCPMT");
+            InBuildTarget.GlobalLinkEnvironment.Config.ExcludedLibraries.Add("LIBCP");
+            InBuildTarget.GlobalLinkEnvironment.Config.ExcludedLibraries.Add("LIBCD");
+            InBuildTarget.GlobalLinkEnvironment.Config.ExcludedLibraries.Add("LIBCMTD");
+            InBuildTarget.GlobalLinkEnvironment.Config.ExcludedLibraries.Add("LIBCPMTD");
+            InBuildTarget.GlobalLinkEnvironment.Config.ExcludedLibraries.Add("LIBCPD");
+
+            // Compile and link with Win32 API libraries.
 // 			d2d1.lib
 // 			d3d11.lib
 // 			dxgi.lib
 // 			ole32.lib
 // 			windowscodecs.lib
 // 			dwrite.lib
-			InBuildTarget.GlobalLinkEnvironment.Config.AdditionalLibraries.Add("runtimeobject.lib");
-			InBuildTarget.GlobalLinkEnvironment.Config.AdditionalLibraries.Add("mincore.lib");
-			//InBuildTarget.GlobalLinkEnvironment.Config.AdditionalLibraries.Add("mincore_legacy.lib");
-			//InBuildTarget.GlobalLinkEnvironment.Config.AdditionalLibraries.Add("mincore_obsolete.lib");
-			InBuildTarget.GlobalLinkEnvironment.Config.AdditionalLibraries.Add("user32.lib");
-
-			// Windows Vista/7 Desktop Windows Manager API for Slate Windows Compliance
-			//InBuildTarget.GlobalLinkEnvironment.Config.AdditionalLibraries.Add("dwmapi.lib");
-
-			UEBuildConfiguration.bCompileSimplygon = false;
-
-			// For 64-bit builds, we'll forcibly ignore a linker warning with DirectInput.  This is
-			// Microsoft's recommended solution as they don't have a fixed .lib for us.
-			InBuildTarget.GlobalLinkEnvironment.Config.AdditionalArguments += " /ignore:4078";
-
-			// WinRT assemblies
-			//InBuildTarget.GlobalCompileEnvironment.Config.SystemDotNetAssemblyPaths.Add("$(DurangoXDK)\\console\\Metadata");
-			//InBuildTarget.GlobalCompileEnvironment.Config.FrameworkAssemblyDependencies.Add("Windows.Foundation.winmd");
-			//InBuildTarget.GlobalCompileEnvironment.Config.FrameworkAssemblyDependencies.Add("Windows.winmd");
-		}
-
-		/**
-		 *	Whether this platform should create debug information or not
-		 *	
-		 *	@param	InPlatform			The UnrealTargetPlatform being built
-		 *	@param	InConfiguration		The UnrealTargetConfiguration being built
-		 *	
-		 *	@return	bool				true if debug info should be generated, false if not
-		 */
-		public override bool ShouldCreateDebugInfo(UnrealTargetPlatform Platform, UnrealTargetConfiguration Configuration)
-		{
-			switch (Configuration)
-			{
-				case UnrealTargetConfiguration.Development:
-				case UnrealTargetConfiguration.Shipping:
-				case UnrealTargetConfiguration.Test:
-				case UnrealTargetConfiguration.Debug:
-				default:
-					return true;
-			};
-		}
-
-		/**
-		 *	Setup the binaries for this specific platform.
-		 *	
-		 *	@param	InBuildTarget		The target being built
-		 */
-		public override void SetupBinaries(UEBuildTarget InBuildTarget)
-		{
-			InBuildTarget.ExtraModuleNames.Add("XInput");
-		}
-
-		//
-		// WinRT specific functions
-		//
-		/**
-		 *	Should the app be compiled as WinRT
-		 */
-		public static bool ShouldCompileWinRT()
-		{
-			return WinRTPlatform.bCompileWinRT;
-		}
-	}
+            InBuildTarget.GlobalLinkEnvironment.Config.AdditionalLibraries.Add("runtimeobject.lib");
+            InBuildTarget.GlobalLinkEnvironment.Config.AdditionalLibraries.Add("mincore.lib");
+            //InBuildTarget.GlobalLinkEnvironment.Config.AdditionalLibraries.Add("mincore_legacy.lib");
+            //InBuildTarget.GlobalLinkEnvironment.Config.AdditionalLibraries.Add("mincore_obsolete.lib");
+            InBuildTarget.GlobalLinkEnvironment.Config.AdditionalLibraries.Add("user32.lib");
+
+            // Windows Vista/7 Desktop Windows Manager API for Slate Windows Compliance
+            //InBuildTarget.GlobalLinkEnvironment.Config.AdditionalLibraries.Add("dwmapi.lib");
+
+            UEBuildConfiguration.bCompileSimplygon = false;
+
+            // For 64-bit builds, we'll forcibly ignore a linker warning with DirectInput.  This is
+            // Microsoft's recommended solution as they don't have a fixed .lib for us.
+            InBuildTarget.GlobalLinkEnvironment.Config.AdditionalArguments += " /ignore:4078";
+
+            // WinRT assemblies
+            //InBuildTarget.GlobalCompileEnvironment.Config.SystemDotNetAssemblyPaths.Add("$(DurangoXDK)\\console\\Metadata");
+            //InBuildTarget.GlobalCompileEnvironment.Config.FrameworkAssemblyDependencies.Add("Windows.Foundation.winmd");
+            //InBuildTarget.GlobalCompileEnvironment.Config.FrameworkAssemblyDependencies.Add("Windows.winmd");
+        }
+
+        /**
+         *	Whether this platform should create debug information or not
+         *	
+         *	@param	InPlatform			The UnrealTargetPlatform being built
+         *	@param	InConfiguration		The UnrealTargetConfiguration being built
+         *	
+         *	@return	bool				true if debug info should be generated, false if not
+         */
+        public override bool ShouldCreateDebugInfo(UnrealTargetPlatform Platform, UnrealTargetConfiguration Configuration)
+        {
+            switch (Configuration)
+            {
+                case UnrealTargetConfiguration.Development:
+                case UnrealTargetConfiguration.Shipping:
+                case UnrealTargetConfiguration.Test:
+                case UnrealTargetConfiguration.Debug:
+                default:
+                    return true;
+            };
+        }
+
+        /**
+         *	Setup the binaries for this specific platform.
+         *	
+         *	@param	InBuildTarget		The target being built
+         */
+        public override void SetupBinaries(UEBuildTarget InBuildTarget)
+        {
+            InBuildTarget.ExtraModuleNames.Add("XInput");
+        }
+
+        //
+        // WinRT specific functions
+        //
+        /**
+         *	Should the app be compiled as WinRT
+         */
+        public static bool ShouldCompileWinRT()
+        {
+            return WinRTPlatform.bCompileWinRT;
+        }
+    }
 }
diff --git a/Engine/Source/Programs/UnrealBuildTool/Windows/UEBuildWindows.cs b/Engine/Source/Programs/UnrealBuildTool/Windows/UEBuildWindows.cs
--- a/Engine/Source/Programs/UnrealBuildTool/Windows/UEBuildWindows.cs
+++ b/Engine/Source/Programs/UnrealBuildTool/Windows/UEBuildWindows.cs
@@ -9,328 +9,328 @@ using Microsoft.Win32;
 
 namespace UnrealBuildTool
 {
-	/// <summary>
-	/// Available compiler toolchains on Windows platform
-	/// </summary>
-	public enum WindowsCompiler
-	{
-		/// Visual Studio 2012 (Visual C++ 11.0)
-		VisualStudio2012,
-
-		/// Visual Studio 2013 (Visual C++ 12.0)
-		VisualStudio2013,
-	}
-
-
-	public class WindowsPlatform : UEBuildPlatform
-	{
-		/// Property caching.
-		private static WindowsCompiler? CachedCompiler;
-
-		/// Version of the compiler toolchain to use on Windows platform
-		public static WindowsCompiler Compiler
-		{
-			get
-			{
-				// Cache the result because Compiler is often used.
-				if (CachedCompiler.HasValue)
-				{
-					return CachedCompiler.Value;
-				}
-
-				// First check if Visual Sudio 2013 is installed.
-				// If both 2012 and 2013 are installed, prefer 2013
-				if (!String.IsNullOrEmpty(WindowsPlatform.GetVSComnToolsPath(WindowsCompiler.VisualStudio2013)))
-				{
-					CachedCompiler = WindowsCompiler.VisualStudio2013;
-				}
-				// Next try Visual Studio 2012
-				else if (!String.IsNullOrEmpty(WindowsPlatform.GetVSComnToolsPath(WindowsCompiler.VisualStudio2012)))
-				{
-					CachedCompiler = WindowsCompiler.VisualStudio2012;
-				}
-				// Finally assume 2013 is installed to defer errors somewhere else like VCToolChain
-				else
-				{
-					CachedCompiler = WindowsCompiler.VisualStudio2013;
-				}
-
-				return CachedCompiler.Value;
-			}
-		}
-
-		/// True if we are using "clang-cl" to compile instead of MSVC on Windows platform
-		public static readonly bool bCompileWithClang = false;
-			// !UnrealBuildTool.CommandLineContains( "UnrealHeaderTool" );	// @todo clang: Avoid using Clang to compile UHT because 64-bit toolchain is not ready yet
-
-		/// True if we're targeting Windows XP as a minimum spec.  In Visual Studio 2012 and higher, this may change how
-		/// we compile and link the application (http://blogs.msdn.com/b/vcblog/archive/2012/10/08/10357555.aspx)
-		[XmlConfig]
-		public static bool SupportWindowsXP = false;
-
-		/** True if VS EnvDTE is available (false when building using Visual Studio Express) */
-		public static bool bHasVisualStudioDTE
-		{
-			get
-			{
-				try
-				{
-					// Interrogate the Win32 registry
+    /// <summary>
+    /// Available compiler toolchains on Windows platform
+    /// </summary>
+    public enum WindowsCompiler
+    {
+        /// Visual Studio 2012 (Visual C++ 11.0)
+        VisualStudio2012,
+
+        /// Visual Studio 2013 (Visual C++ 12.0)
+        VisualStudio2013,
+    }
+
+
+    public class WindowsPlatform : UEBuildPlatform
+    {
+        /// Property caching.
+        private static WindowsCompiler? CachedCompiler;
+
+        /// Version of the compiler toolchain to use on Windows platform
+        public static WindowsCompiler Compiler
+        {
+            get
+            {
+                // Cache the result because Compiler is often used.
+                if (CachedCompiler.HasValue)
+                {
+                    return CachedCompiler.Value;
+                }
+
+                // First check if Visual Sudio 2013 is installed.
+                // If both 2012 and 2013 are installed, prefer 2013
+                if (!String.IsNullOrEmpty(WindowsPlatform.GetVSComnToolsPath(WindowsCompiler.VisualStudio2013)))
+                {
+                    CachedCompiler = WindowsCompiler.VisualStudio2013;
+                }
+                // Next try Visual Studio 2012
+                else if (!String.IsNullOrEmpty(WindowsPlatform.GetVSComnToolsPath(WindowsCompiler.VisualStudio2012)))
+                {
+                    CachedCompiler = WindowsCompiler.VisualStudio2012;
+                }
+                // Finally assume 2013 is installed to defer errors somewhere else like VCToolChain
+                else
+                {
+                    CachedCompiler = WindowsCompiler.VisualStudio2013;
+                }
+
+                return CachedCompiler.Value;
+            }
+        }
+
+        /// True if we are using "clang-cl" to compile instead of MSVC on Windows platform
+        public static readonly bool bCompileWithClang = false;
+            // !UnrealBuildTool.CommandLineContains( "UnrealHeaderTool" );	// @todo clang: Avoid using Clang to compile UHT because 64-bit toolchain is not ready yet
+
+        /// True if we're targeting Windows XP as a minimum spec.  In Visual Studio 2012 and higher, this may change how
+        /// we compile and link the application (http://blogs.msdn.com/b/vcblog/archive/2012/10/08/10357555.aspx)
+        [XmlConfig]
+        public static bool SupportWindowsXP = false;
+
+        /** True if VS EnvDTE is available (false when building using Visual Studio Express) */
+        public static bool bHasVisualStudioDTE
+        {
+            get
+            {
+                try
+                {
+                    // Interrogate the Win32 registry
                     string DTEKey = (Compiler == WindowsCompiler.VisualStudio2013) ? "VisualStudio.DTE.12.0" : "VisualStudio.DTE.11.0";
                     return RegistryKey.OpenBaseKey(RegistryHive.ClassesRoot, RegistryView.Registry32).OpenSubKey(DTEKey) != null;
-				}
-				catch(Exception) 
-				{
-					return false;
-				}
-			}
-		}
-
-		public override SDKStatus HasRequiredSDKsInstalled()
-		{
-			return SDKStatus.Valid;
-		}
-
-		/**
-		 * Returns VisualStudio common tools path for current compiler.
-		 * 
-		 * @return Common tools path.
-		 */
-		public static string GetVSComnToolsPath()
-		{
-			return GetVSComnToolsPath(Compiler);
-		}
-
-		/**
-		 * Returns VisualStudio common tools path for given compiler.
-		 * 
-		 * @param Compiler Compiler for which to return tools path.
-		 * 
-		 * @return Common tools path.
-		 */
-		public static string GetVSComnToolsPath(WindowsCompiler Compiler)
-		{
-			int VSVersion;
-
-			switch(Compiler)
-			{
-				case WindowsCompiler.VisualStudio2012:
-					VSVersion = 11;
-					break;
-				case WindowsCompiler.VisualStudio2013:
-					VSVersion = 12;
-					break;
-				default:
-					throw new NotSupportedException("Not supported compiler.");
-			}
-
-			string[] PossibleRegPaths = new string[] {
-				@"Wow6432Node\Microsoft\VisualStudio",	// Non-express VS2013 on 64-bit machine.
-				@"Microsoft\VisualStudio",				// Non-express VS2013 on 32-bit machine.
-				@"Wow6432Node\Microsoft\WDExpress",		// Express VS2013 on 64-bit machine.
-				@"Microsoft\WDExpress"					// Express VS2013 on 32-bit machine.
-			};
-
-			string VSPath = null;
-
-			foreach(var PossibleRegPath in PossibleRegPaths)
-			{
-				VSPath = (string) Registry.GetValue(string.Format(@"HKEY_LOCAL_MACHINE\SOFTWARE\{0}\{1}.0", PossibleRegPath, VSVersion), "InstallDir", null);
-
-				if(VSPath != null)
-				{
-					break;
-				}
-			}
-
-			if(VSPath == null)
-			{
-				return null;
-			}
-
-			return new DirectoryInfo(Path.Combine(VSPath, "..", "Tools")).FullName;
-		}
-
-
-		/**
-		 *	Register the platform with the UEBuildPlatform class
-		 */
-		public override void RegisterBuildPlatform()
-		{
-			// Register this build platform for both Win64 and Win32
-			Log.TraceVerbose("        Registering for {0}", UnrealTargetPlatform.Win64.ToString());
-			UEBuildPlatform.RegisterBuildPlatform(UnrealTargetPlatform.Win64, this);
-			UEBuildPlatform.RegisterPlatformWithGroup(UnrealTargetPlatform.Win64, UnrealPlatformGroup.Windows);
-			UEBuildPlatform.RegisterPlatformWithGroup(UnrealTargetPlatform.Win64, UnrealPlatformGroup.Microsoft);
-
-			Log.TraceVerbose("        Registering for {0}", UnrealTargetPlatform.Win32.ToString());
-			UEBuildPlatform.RegisterBuildPlatform(UnrealTargetPlatform.Win32, this);
-			UEBuildPlatform.RegisterPlatformWithGroup(UnrealTargetPlatform.Win32, UnrealPlatformGroup.Windows);
-			UEBuildPlatform.RegisterPlatformWithGroup(UnrealTargetPlatform.Win32, UnrealPlatformGroup.Microsoft);
-		}
-
-		/**
-		 *	Retrieve the CPPTargetPlatform for the given UnrealTargetPlatform
-		 *
-		 *	@param	InUnrealTargetPlatform		The UnrealTargetPlatform being build
-		 *	
-		 *	@return	CPPTargetPlatform			The CPPTargetPlatform to compile for
-		 */
-		public override CPPTargetPlatform GetCPPTargetPlatform(UnrealTargetPlatform InUnrealTargetPlatform)
-		{
-			switch (InUnrealTargetPlatform)
-			{
-				case UnrealTargetPlatform.Win32:
-					return CPPTargetPlatform.Win32;
-
-				case UnrealTargetPlatform.Win64:
-					return CPPTargetPlatform.Win64;
-			}
-			throw new BuildException("WindowsPlatform::GetCPPTargetPlatform: Invalid request for {0}", InUnrealTargetPlatform.ToString());
-		}
-
-		/**
-		 *	Get the extension to use for the given binary type
-		 *	
-		 *	@param	InBinaryType		The binrary type being built
-		 *	
-		 *	@return	string				The binary extenstion (ie 'exe' or 'dll')
-		 */
-		public override string GetBinaryExtension(UEBuildBinaryType InBinaryType)
-		{
-			switch (InBinaryType)
-			{
-				case UEBuildBinaryType.DynamicLinkLibrary: 
-					return ".dll";
-				case UEBuildBinaryType.Executable:
-					return ".exe";
-				case UEBuildBinaryType.StaticLibrary:
-					return ".lib";
-				case UEBuildBinaryType.Object:
-					return ".obj";
-				case UEBuildBinaryType.PrecompiledHeader:
-					return ".pch";
-			}
-			return base.GetBinaryExtension(InBinaryType);
-		}
-
-		
-		/// <summary>
-		/// When using a Visual Studio compiler, returns the version name as a string
-		/// </summary>
-		/// <returns>The Visual Studio compiler version name (e.g. "2012")</returns>
-		public static string GetVisualStudioCompilerVersionName()
-		{
-			switch( Compiler )
-			{ 
-				case WindowsCompiler.VisualStudio2012:
-					return "2012";
-
-				case WindowsCompiler.VisualStudio2013:
-					return "2013";
-
-				default:
-					throw new BuildException( "Unexpected WindowsCompiler version for GetVisualStudioCompilerVersionName().  Either not using a Visual Studio compiler or switch block needs to be updated");
-			}
-		}
-
-
-		/**
-		 *	Get the extension to use for debug info for the given binary type
-		 *	
-		 *	@param	InBinaryType		The binary type being built
-		 *	
-		 *	@return	string				The debug info extension (i.e. 'pdb')
-		 */
-		public override string GetDebugInfoExtension( UEBuildBinaryType InBinaryType )
-		{
+                }
+                catch(Exception) 
+                {
+                    return false;
+                }
+            }
+        }
+
+        protected override SDKStatus HasRequiredManualSDKInternal()
+        {
+            return SDKStatus.Valid;
+        }
+
+        /**
+         * Returns VisualStudio common tools path for current compiler.
+         * 
+         * @return Common tools path.
+         */
+        public static string GetVSComnToolsPath()
+        {
+            return GetVSComnToolsPath(Compiler);
+        }
+
+        /**
+         * Returns VisualStudio common tools path for given compiler.
+         * 
+         * @param Compiler Compiler for which to return tools path.
+         * 
+         * @return Common tools path.
+         */
+        public static string GetVSComnToolsPath(WindowsCompiler Compiler)
+        {
+            int VSVersion;
+
+            switch(Compiler)
+            {
+                case WindowsCompiler.VisualStudio2012:
+                    VSVersion = 11;
+                    break;
+                case WindowsCompiler.VisualStudio2013:
+                    VSVersion = 12;
+                    break;
+                default:
+                    throw new NotSupportedException("Not supported compiler.");
+            }
+
+            string[] PossibleRegPaths = new string[] {
+                @"Wow6432Node\Microsoft\VisualStudio",	// Non-express VS2013 on 64-bit machine.
+                @"Microsoft\VisualStudio",				// Non-express VS2013 on 32-bit machine.
+                @"Wow6432Node\Microsoft\WDExpress",		// Express VS2013 on 64-bit machine.
+                @"Microsoft\WDExpress"					// Express VS2013 on 32-bit machine.
+            };
+
+            string VSPath = null;
+
+            foreach(var PossibleRegPath in PossibleRegPaths)
+            {
+                VSPath = (string) Registry.GetValue(string.Format(@"HKEY_LOCAL_MACHINE\SOFTWARE\{0}\{1}.0", PossibleRegPath, VSVersion), "InstallDir", null);
+
+                if(VSPath != null)
+                {
+                    break;
+                }
+            }
+
+            if(VSPath == null)
+            {
+                return null;
+            }
+
+            return new DirectoryInfo(Path.Combine(VSPath, "..", "Tools")).FullName;
+        }
+
+
+        /**
+         *	Register the platform with the UEBuildPlatform class
+         */
+        protected override void RegisterBuildPlatformInternal()
+        {
+            // Register this build platform for both Win64 and Win32
+            Log.TraceVerbose("        Registering for {0}", UnrealTargetPlatform.Win64.ToString());
+            UEBuildPlatform.RegisterBuildPlatform(UnrealTargetPlatform.Win64, this);
+            UEBuildPlatform.RegisterPlatformWithGroup(UnrealTargetPlatform.Win64, UnrealPlatformGroup.Windows);
+            UEBuildPlatform.RegisterPlatformWithGroup(UnrealTargetPlatform.Win64, UnrealPlatformGroup.Microsoft);
+
+            Log.TraceVerbose("        Registering for {0}", UnrealTargetPlatform.Win32.ToString());
+            UEBuildPlatform.RegisterBuildPlatform(UnrealTargetPlatform.Win32, this);
+            UEBuildPlatform.RegisterPlatformWithGroup(UnrealTargetPlatform.Win32, UnrealPlatformGroup.Windows);
+            UEBuildPlatform.RegisterPlatformWithGroup(UnrealTargetPlatform.Win32, UnrealPlatformGroup.Microsoft);
+        }
+
+        /**
+         *	Retrieve the CPPTargetPlatform for the given UnrealTargetPlatform
+         *
+         *	@param	InUnrealTargetPlatform		The UnrealTargetPlatform being build
+         *	
+         *	@return	CPPTargetPlatform			The CPPTargetPlatform to compile for
+         */
+        public override CPPTargetPlatform GetCPPTargetPlatform(UnrealTargetPlatform InUnrealTargetPlatform)
+        {
+            switch (InUnrealTargetPlatform)
+            {
+                case UnrealTargetPlatform.Win32:
+                    return CPPTargetPlatform.Win32;
+
+                case UnrealTargetPlatform.Win64:
+                    return CPPTargetPlatform.Win64;
+            }
+            throw new BuildException("WindowsPlatform::GetCPPTargetPlatform: Invalid request for {0}", InUnrealTargetPlatform.ToString());
+        }
+
+        /**
+         *	Get the extension to use for the given binary type
+         *	
+         *	@param	InBinaryType		The binrary type being built
+         *	
+         *	@return	string				The binary extenstion (ie 'exe' or 'dll')
+         */
+        public override string GetBinaryExtension(UEBuildBinaryType InBinaryType)
+        {
+            switch (InBinaryType)
+            {
+                case UEBuildBinaryType.DynamicLinkLibrary: 
+                    return ".dll";
+                case UEBuildBinaryType.Executable:
+                    return ".exe";
+                case UEBuildBinaryType.StaticLibrary:
+                    return ".lib";
+                case UEBuildBinaryType.Object:
+                    return ".obj";
+                case UEBuildBinaryType.PrecompiledHeader:
+                    return ".pch";
+            }
+            return base.GetBinaryExtension(InBinaryType);
+        }
+
+        
+        /// <summary>
+        /// When using a Visual Studio compiler, returns the version name as a string
+        /// </summary>
+        /// <returns>The Visual Studio compiler version name (e.g. "2012")</returns>
+        public static string GetVisualStudioCompilerVersionName()
+        {
+            switch( Compiler )
+            { 
+                case WindowsCompiler.VisualStudio2012:
+                    return "2012";
+
+                case WindowsCompiler.VisualStudio2013:
+                    return "2013";
+
+                default:
+                    throw new BuildException( "Unexpected WindowsCompiler version for GetVisualStudioCompilerVersionName().  Either not using a Visual Studio compiler or switch block needs to be updated");
+            }
+        }
+
+
+        /**
+         *	Get the extension to use for debug info for the given binary type
+         *	
+         *	@param	InBinaryType		The binary type being built
+         *	
+         *	@return	string				The debug info extension (i.e. 'pdb')
+         */
+        public override string GetDebugInfoExtension( UEBuildBinaryType InBinaryType )
+        {
             switch (InBinaryType)
             {
                 case UEBuildBinaryType.DynamicLinkLibrary:
-				case UEBuildBinaryType.Executable:
+                case UEBuildBinaryType.Executable:
                     return ".pdb";
             }
             return "";
         }
 
-		/**
-		 *	Whether incremental linking should be used
-		 *	
-		 *	@param	InPlatform			The CPPTargetPlatform being built
-		 *	@param	InConfiguration		The CPPTargetConfiguration being built
-		 *	
-		 *	@return	bool	true if incremental linking should be used, false if not
-		 */
-		public override bool ShouldUseIncrementalLinking(CPPTargetPlatform Platform, CPPTargetConfiguration Configuration)
-		{
-			return (Configuration == CPPTargetConfiguration.Debug);
-		}
-
-		/**
-		 *	Whether PDB files should be used
-		 *	
-		 *	@param	InPlatform			The CPPTargetPlatform being built
-		 *	@param	InConfiguration		The CPPTargetConfiguration being built
-		 *	@param	bInCreateDebugInfo	true if debug info is getting create, false if not
-		 *	
-		 *	@return	bool	true if PDB files should be used, false if not
-		 */
-		public override bool ShouldUsePDBFiles(CPPTargetPlatform Platform, CPPTargetConfiguration Configuration, bool bCreateDebugInfo)
-		{
-			// Only supported on PC.
-			if (bCreateDebugInfo && ShouldUseIncrementalLinking(Platform, Configuration))
-			{
-				return true;
-			}
-			return false;
-		}
-
-		/**
-		 *	Whether the editor should be built for this platform or not
-		 *	
-		 *	@param	InPlatform		The UnrealTargetPlatform being built
-		 *	@param	InConfiguration	The UnrealTargetConfiguration being built
-		 *	@return	bool			true if the editor should be built, false if not
-		 */
-		public override bool ShouldNotBuildEditor(UnrealTargetPlatform InPlatform, UnrealTargetConfiguration InConfiguration)
-		{
-			return false;
-		}
-
-		public override bool BuildRequiresCookedData(UnrealTargetPlatform InPlatform, UnrealTargetConfiguration InConfiguration)
-		{
-			return false;
-		}
-
-		public override void ResetBuildConfiguration(UnrealTargetPlatform InPlatform, UnrealTargetConfiguration InConfiguration)
-		{
-			UEBuildConfiguration.bCompileICU = true;
-		}
+        /**
+         *	Whether incremental linking should be used
+         *	
+         *	@param	InPlatform			The CPPTargetPlatform being built
+         *	@param	InConfiguration		The CPPTargetConfiguration being built
+         *	
+         *	@return	bool	true if incremental linking should be used, false if not
+         */
+        public override bool ShouldUseIncrementalLinking(CPPTargetPlatform Platform, CPPTargetConfiguration Configuration)
+        {
+            return (Configuration == CPPTargetConfiguration.Debug);
+        }
+
+        /**
+         *	Whether PDB files should be used
+         *	
+         *	@param	InPlatform			The CPPTargetPlatform being built
+         *	@param	InConfiguration		The CPPTargetConfiguration being built
+         *	@param	bInCreateDebugInfo	true if debug info is getting create, false if not
+         *	
+         *	@return	bool	true if PDB files should be used, false if not
+         */
+        public override bool ShouldUsePDBFiles(CPPTargetPlatform Platform, CPPTargetConfiguration Configuration, bool bCreateDebugInfo)
+        {
+            // Only supported on PC.
+            if (bCreateDebugInfo && ShouldUseIncrementalLinking(Platform, Configuration))
+            {
+                return true;
+            }
+            return false;
+        }
+
+        /**
+         *	Whether the editor should be built for this platform or not
+         *	
+         *	@param	InPlatform		The UnrealTargetPlatform being built
+         *	@param	InConfiguration	The UnrealTargetConfiguration being built
+         *	@return	bool			true if the editor should be built, false if not
+         */
+        public override bool ShouldNotBuildEditor(UnrealTargetPlatform InPlatform, UnrealTargetConfiguration InConfiguration)
+        {
+            return false;
+        }
+
+        public override bool BuildRequiresCookedData(UnrealTargetPlatform InPlatform, UnrealTargetConfiguration InConfiguration)
+        {
+            return false;
+        }
+
+        public override void ResetBuildConfiguration(UnrealTargetPlatform InPlatform, UnrealTargetConfiguration InConfiguration)
+        {
+            UEBuildConfiguration.bCompileICU = true;
+        }
 
         public override void ValidateBuildConfiguration(CPPTargetConfiguration Configuration, CPPTargetPlatform Platform, bool bCreateDebugInfo)
         {
-			// @todo clang: PCH files aren't quite working yet with "clang-cl" (no /Yc support, and "-x c++-header" cannot be specified)
-			if( WindowsPlatform.bCompileWithClang )
-			{
-				BuildConfiguration.bUsePCHFiles = false;
-				BuildConfiguration.bUseSharedPCHs = false;
-			}
+            // @todo clang: PCH files aren't quite working yet with "clang-cl" (no /Yc support, and "-x c++-header" cannot be specified)
+            if( WindowsPlatform.bCompileWithClang )
+            {
+                BuildConfiguration.bUsePCHFiles = false;
+                BuildConfiguration.bUseSharedPCHs = false;
+            }
         }
 
-		/**
+        /**
          *	Validate the UEBuildConfiguration for this platform
          *	This is called BEFORE calling UEBuildConfiguration to allow setting 
          *	various fields used in that function such as CompileLeanAndMean...
          */
         public override void ValidateUEBuildConfiguration()
         {
-			UEBuildConfiguration.bCompileICU &= !bCompileWithClang;	// @todo clang: ICU causes STL link errors when using Clang on Windows.  Needs debugging.
-		}
+            UEBuildConfiguration.bCompileICU &= !bCompileWithClang;	// @todo clang: ICU causes STL link errors when using Clang on Windows.  Needs debugging.
+        }
 
-		public override void ModifyNewlyLoadedModule(UEBuildModule InModule, TargetInfo Target)
-		{
-			if ((Target.Platform == UnrealTargetPlatform.Win32) || (Target.Platform == UnrealTargetPlatform.Win64))
-			{
+        public override void ModifyNewlyLoadedModule(UEBuildModule InModule, TargetInfo Target)
+        {
+            if ((Target.Platform == UnrealTargetPlatform.Win32) || (Target.Platform == UnrealTargetPlatform.Win64))
+            {
                 bool bBuildShaderFormats = UEBuildConfiguration.bForceBuildShaderFormats;
 
                 if (!UEBuildConfiguration.bBuildRequiresCookedData)
@@ -341,14 +341,14 @@ namespace UnrealBuildTool
                     }
                 }
 
-				// allow standalone tools to use target platform modules, without needing Engine
-				if (UEBuildConfiguration.bForceBuildTargetPlatforms)
-				{
-					InModule.AddDynamicallyLoadedModule("WindowsTargetPlatform");
-					InModule.AddDynamicallyLoadedModule("WindowsNoEditorTargetPlatform");
-					InModule.AddDynamicallyLoadedModule("WindowsServerTargetPlatform");
-					InModule.AddDynamicallyLoadedModule("WindowsClientTargetPlatform");
-				}
+                // allow standalone tools to use target platform modules, without needing Engine
+                if (UEBuildConfiguration.bForceBuildTargetPlatforms)
+                {
+                    InModule.AddDynamicallyLoadedModule("WindowsTargetPlatform");
+                    InModule.AddDynamicallyLoadedModule("WindowsNoEditorTargetPlatform");
+                    InModule.AddDynamicallyLoadedModule("WindowsServerTargetPlatform");
+                    InModule.AddDynamicallyLoadedModule("WindowsClientTargetPlatform");
+                }
 
                 if (bBuildShaderFormats)
                 {
@@ -356,248 +356,248 @@ namespace UnrealBuildTool
                     InModule.AddDynamicallyLoadedModule("ShaderFormatOpenGL");
                 }
 
-				if (InModule.ToString() == "D3D11RHI")
-				{
-					// To enable platform specific D3D11 RHI Types
-					InModule.AddPrivateIncludePath("Runtime/Windows/D3D11RHI/Private/Windows");
-				}
-			}
-		}
-
-		/**
-		 *	Setup the target environment for building
-		 *	
-		 *	@param	InBuildTarget		The target being built
-		 */
-		public override void SetUpEnvironment(UEBuildTarget InBuildTarget)
-		{
-			InBuildTarget.GlobalCompileEnvironment.Config.Definitions.Add("WIN32=1");
-
-			// Win32 XP is only supported at this time.
-			SupportWindowsXP = SupportWindowsXP && InBuildTarget.Platform == UnrealTargetPlatform.Win32;
-
-			if( SupportWindowsXP )
-			{
-				// Windows XP SP3 or higher required
-				InBuildTarget.GlobalCompileEnvironment.Config.Definitions.Add("_WIN32_WINNT=0x0502");
-				InBuildTarget.GlobalCompileEnvironment.Config.Definitions.Add("WINVER=0x0502");
-			}
-			else
-			{
-				// Windows Vista or higher required
-				InBuildTarget.GlobalCompileEnvironment.Config.Definitions.Add("_WIN32_WINNT=0x0600");
-				InBuildTarget.GlobalCompileEnvironment.Config.Definitions.Add("WINVER=0x0600");
-			}
-			InBuildTarget.GlobalCompileEnvironment.Config.Definitions.Add("PLATFORM_WINDOWS=1");
-
-			String MorpheusShaderPath = Path.Combine(BuildConfiguration.RelativeEnginePath, "Shaders/PS4/PostProcessHMDMorpheus.usf");
-			if (File.Exists(MorpheusShaderPath))
-			{
-				InBuildTarget.GlobalCompileEnvironment.Config.Definitions.Add("HAS_MORPHEUS=1");
-			}
-
-			if (InBuildTarget.Rules != null)
-			{
-				// Explicitly exclude the MS C++ runtime libraries we're not using, to ensure other libraries we link with use the same
-				// runtime library as the engine.
-				if (InBuildTarget.Configuration == UnrealTargetConfiguration.Debug && BuildConfiguration.bDebugBuildsActuallyUseDebugCRT)
-				{
-					InBuildTarget.GlobalLinkEnvironment.Config.ExcludedLibraries.Add("MSVCRT");
-					InBuildTarget.GlobalLinkEnvironment.Config.ExcludedLibraries.Add("MSVCPRT");
-				}
-				else
-				{
-					InBuildTarget.GlobalLinkEnvironment.Config.ExcludedLibraries.Add("MSVCRTD");
-					InBuildTarget.GlobalLinkEnvironment.Config.ExcludedLibraries.Add("MSVCPRTD");
-				}
-				InBuildTarget.GlobalLinkEnvironment.Config.ExcludedLibraries.Add("LIBC");
-				InBuildTarget.GlobalLinkEnvironment.Config.ExcludedLibraries.Add("LIBCMT");
-				InBuildTarget.GlobalLinkEnvironment.Config.ExcludedLibraries.Add("LIBCPMT");
-				InBuildTarget.GlobalLinkEnvironment.Config.ExcludedLibraries.Add("LIBCP");
-				InBuildTarget.GlobalLinkEnvironment.Config.ExcludedLibraries.Add("LIBCD");
-				InBuildTarget.GlobalLinkEnvironment.Config.ExcludedLibraries.Add("LIBCMTD");
-				InBuildTarget.GlobalLinkEnvironment.Config.ExcludedLibraries.Add("LIBCPMTD");
-				InBuildTarget.GlobalLinkEnvironment.Config.ExcludedLibraries.Add("LIBCPD");
-
-				//@todo ATL: Currently, only VSAccessor requires ATL (which is only used in editor builds)
-				// When compiling games, we do not want to include ATL - and we can't when compiling Rocket games
-				// due to VS 2012 Express not including ATL.
-				// If more modules end up requiring ATL, this should be refactored into a BuildTarget flag (bNeedsATL)
-				// that is set by the modules the target includes to allow for easier tracking.
-				// Alternatively, if VSAccessor is modified to not require ATL than we should always exclude the libraries.
-				if (InBuildTarget.ShouldCompileMonolithic() && (InBuildTarget.Rules != null) &&
-					(TargetRules.IsGameType(InBuildTarget.Rules.Type)) && (TargetRules.IsEditorType(InBuildTarget.Rules.Type) == false))
-				{
-					InBuildTarget.GlobalLinkEnvironment.Config.ExcludedLibraries.Add("atl");
-					InBuildTarget.GlobalLinkEnvironment.Config.ExcludedLibraries.Add("atls");
-					InBuildTarget.GlobalLinkEnvironment.Config.ExcludedLibraries.Add("atlsd");
-					InBuildTarget.GlobalLinkEnvironment.Config.ExcludedLibraries.Add("atlsn");
-					InBuildTarget.GlobalLinkEnvironment.Config.ExcludedLibraries.Add("atlsnd");
-				}
-
-				// Add the library used for the delayed loading of DLLs.
-				InBuildTarget.GlobalLinkEnvironment.Config.AdditionalLibraries.Add("delayimp.lib");
-
-				//@todo - remove once FB implementation uses Http module
-				if (UEBuildConfiguration.bCompileAgainstEngine)
-				{
-					// link against wininet (used by FBX and Facebook)
-					InBuildTarget.GlobalLinkEnvironment.Config.AdditionalLibraries.Add("wininet.lib");
-				}
-
-				// Compile and link with Win32 API libraries.
-				InBuildTarget.GlobalLinkEnvironment.Config.AdditionalLibraries.Add("rpcrt4.lib");
-				//InBuildTarget.GlobalLinkEnvironment.Config.AdditionalLibraries.Add("wsock32.lib");
-				InBuildTarget.GlobalLinkEnvironment.Config.AdditionalLibraries.Add("ws2_32.lib");
-				InBuildTarget.GlobalLinkEnvironment.Config.AdditionalLibraries.Add("dbghelp.lib");
-				InBuildTarget.GlobalLinkEnvironment.Config.AdditionalLibraries.Add("comctl32.lib");
-				InBuildTarget.GlobalLinkEnvironment.Config.AdditionalLibraries.Add("Winmm.lib");
-				InBuildTarget.GlobalLinkEnvironment.Config.AdditionalLibraries.Add("kernel32.lib");
-				InBuildTarget.GlobalLinkEnvironment.Config.AdditionalLibraries.Add("user32.lib");
-				InBuildTarget.GlobalLinkEnvironment.Config.AdditionalLibraries.Add("gdi32.lib");
-				InBuildTarget.GlobalLinkEnvironment.Config.AdditionalLibraries.Add("winspool.lib");
-				InBuildTarget.GlobalLinkEnvironment.Config.AdditionalLibraries.Add("comdlg32.lib");
-				InBuildTarget.GlobalLinkEnvironment.Config.AdditionalLibraries.Add("advapi32.lib");
-				InBuildTarget.GlobalLinkEnvironment.Config.AdditionalLibraries.Add("shell32.lib");
-				InBuildTarget.GlobalLinkEnvironment.Config.AdditionalLibraries.Add("ole32.lib");
-				InBuildTarget.GlobalLinkEnvironment.Config.AdditionalLibraries.Add("oleaut32.lib");
-				InBuildTarget.GlobalLinkEnvironment.Config.AdditionalLibraries.Add("uuid.lib");
-				InBuildTarget.GlobalLinkEnvironment.Config.AdditionalLibraries.Add("odbc32.lib");
-				InBuildTarget.GlobalLinkEnvironment.Config.AdditionalLibraries.Add("odbccp32.lib");
-				InBuildTarget.GlobalLinkEnvironment.Config.AdditionalLibraries.Add("netapi32.lib");
-				InBuildTarget.GlobalLinkEnvironment.Config.AdditionalLibraries.Add("iphlpapi.lib");
+                if (InModule.ToString() == "D3D11RHI")
+                {
+                    // To enable platform specific D3D11 RHI Types
+                    InModule.AddPrivateIncludePath("Runtime/Windows/D3D11RHI/Private/Windows");
+                }
+            }
+        }
+
+        /**
+         *	Setup the target environment for building
+         *	
+         *	@param	InBuildTarget		The target being built
+         */
+        public override void SetUpEnvironment(UEBuildTarget InBuildTarget)
+        {
+            InBuildTarget.GlobalCompileEnvironment.Config.Definitions.Add("WIN32=1");
+
+            // Win32 XP is only supported at this time.
+            SupportWindowsXP = SupportWindowsXP && InBuildTarget.Platform == UnrealTargetPlatform.Win32;
+
+            if( SupportWindowsXP )
+            {
+                // Windows XP SP3 or higher required
+                InBuildTarget.GlobalCompileEnvironment.Config.Definitions.Add("_WIN32_WINNT=0x0502");
+                InBuildTarget.GlobalCompileEnvironment.Config.Definitions.Add("WINVER=0x0502");
+            }
+            else
+            {
+                // Windows Vista or higher required
+                InBuildTarget.GlobalCompileEnvironment.Config.Definitions.Add("_WIN32_WINNT=0x0600");
+                InBuildTarget.GlobalCompileEnvironment.Config.Definitions.Add("WINVER=0x0600");
+            }
+            InBuildTarget.GlobalCompileEnvironment.Config.Definitions.Add("PLATFORM_WINDOWS=1");
+
+            String MorpheusShaderPath = Path.Combine(BuildConfiguration.RelativeEnginePath, "Shaders/PS4/PostProcessHMDMorpheus.usf");
+            if (File.Exists(MorpheusShaderPath))
+            {
+                InBuildTarget.GlobalCompileEnvironment.Config.Definitions.Add("HAS_MORPHEUS=1");
+            }
+
+            if (InBuildTarget.Rules != null)
+            {
+                // Explicitly exclude the MS C++ runtime libraries we're not using, to ensure other libraries we link with use the same
+                // runtime library as the engine.
+                if (InBuildTarget.Configuration == UnrealTargetConfiguration.Debug && BuildConfiguration.bDebugBuildsActuallyUseDebugCRT)
+                {
+                    InBuildTarget.GlobalLinkEnvironment.Config.ExcludedLibraries.Add("MSVCRT");
+                    InBuildTarget.GlobalLinkEnvironment.Config.ExcludedLibraries.Add("MSVCPRT");
+                }
+                else
+                {
+                    InBuildTarget.GlobalLinkEnvironment.Config.ExcludedLibraries.Add("MSVCRTD");
+                    InBuildTarget.GlobalLinkEnvironment.Config.ExcludedLibraries.Add("MSVCPRTD");
+                }
+                InBuildTarget.GlobalLinkEnvironment.Config.ExcludedLibraries.Add("LIBC");
+                InBuildTarget.GlobalLinkEnvironment.Config.ExcludedLibraries.Add("LIBCMT");
+                InBuildTarget.GlobalLinkEnvironment.Config.ExcludedLibraries.Add("LIBCPMT");
+                InBuildTarget.GlobalLinkEnvironment.Config.ExcludedLibraries.Add("LIBCP");
+                InBuildTarget.GlobalLinkEnvironment.Config.ExcludedLibraries.Add("LIBCD");
+                InBuildTarget.GlobalLinkEnvironment.Config.ExcludedLibraries.Add("LIBCMTD");
+                InBuildTarget.GlobalLinkEnvironment.Config.ExcludedLibraries.Add("LIBCPMTD");
+                InBuildTarget.GlobalLinkEnvironment.Config.ExcludedLibraries.Add("LIBCPD");
+
+                //@todo ATL: Currently, only VSAccessor requires ATL (which is only used in editor builds)
+                // When compiling games, we do not want to include ATL - and we can't when compiling Rocket games
+                // due to VS 2012 Express not including ATL.
+                // If more modules end up requiring ATL, this should be refactored into a BuildTarget flag (bNeedsATL)
+                // that is set by the modules the target includes to allow for easier tracking.
+                // Alternatively, if VSAccessor is modified to not require ATL than we should always exclude the libraries.
+                if (InBuildTarget.ShouldCompileMonolithic() && (InBuildTarget.Rules != null) &&
+                    (TargetRules.IsGameType(InBuildTarget.Rules.Type)) && (TargetRules.IsEditorType(InBuildTarget.Rules.Type) == false))
+                {
+                    InBuildTarget.GlobalLinkEnvironment.Config.ExcludedLibraries.Add("atl");
+                    InBuildTarget.GlobalLinkEnvironment.Config.ExcludedLibraries.Add("atls");
+                    InBuildTarget.GlobalLinkEnvironment.Config.ExcludedLibraries.Add("atlsd");
+                    InBuildTarget.GlobalLinkEnvironment.Config.ExcludedLibraries.Add("atlsn");
+                    InBuildTarget.GlobalLinkEnvironment.Config.ExcludedLibraries.Add("atlsnd");
+                }
+
+                // Add the library used for the delayed loading of DLLs.
+                InBuildTarget.GlobalLinkEnvironment.Config.AdditionalLibraries.Add("delayimp.lib");
+
+                //@todo - remove once FB implementation uses Http module
+                if (UEBuildConfiguration.bCompileAgainstEngine)
+                {
+                    // link against wininet (used by FBX and Facebook)
+                    InBuildTarget.GlobalLinkEnvironment.Config.AdditionalLibraries.Add("wininet.lib");
+                }
+
+                // Compile and link with Win32 API libraries.
+                InBuildTarget.GlobalLinkEnvironment.Config.AdditionalLibraries.Add("rpcrt4.lib");
+                //InBuildTarget.GlobalLinkEnvironment.Config.AdditionalLibraries.Add("wsock32.lib");
+                InBuildTarget.GlobalLinkEnvironment.Config.AdditionalLibraries.Add("ws2_32.lib");
+                InBuildTarget.GlobalLinkEnvironment.Config.AdditionalLibraries.Add("dbghelp.lib");
+                InBuildTarget.GlobalLinkEnvironment.Config.AdditionalLibraries.Add("comctl32.lib");
+                InBuildTarget.GlobalLinkEnvironment.Config.AdditionalLibraries.Add("Winmm.lib");
+                InBuildTarget.GlobalLinkEnvironment.Config.AdditionalLibraries.Add("kernel32.lib");
+                InBuildTarget.GlobalLinkEnvironment.Config.AdditionalLibraries.Add("user32.lib");
+                InBuildTarget.GlobalLinkEnvironment.Config.AdditionalLibraries.Add("gdi32.lib");
+                InBuildTarget.GlobalLinkEnvironment.Config.AdditionalLibraries.Add("winspool.lib");
+                InBuildTarget.GlobalLinkEnvironment.Config.AdditionalLibraries.Add("comdlg32.lib");
+                InBuildTarget.GlobalLinkEnvironment.Config.AdditionalLibraries.Add("advapi32.lib");
+                InBuildTarget.GlobalLinkEnvironment.Config.AdditionalLibraries.Add("shell32.lib");
+                InBuildTarget.GlobalLinkEnvironment.Config.AdditionalLibraries.Add("ole32.lib");
+                InBuildTarget.GlobalLinkEnvironment.Config.AdditionalLibraries.Add("oleaut32.lib");
+                InBuildTarget.GlobalLinkEnvironment.Config.AdditionalLibraries.Add("uuid.lib");
+                InBuildTarget.GlobalLinkEnvironment.Config.AdditionalLibraries.Add("odbc32.lib");
+                InBuildTarget.GlobalLinkEnvironment.Config.AdditionalLibraries.Add("odbccp32.lib");
+                InBuildTarget.GlobalLinkEnvironment.Config.AdditionalLibraries.Add("netapi32.lib");
+                InBuildTarget.GlobalLinkEnvironment.Config.AdditionalLibraries.Add("iphlpapi.lib");
                 InBuildTarget.GlobalLinkEnvironment.Config.AdditionalLibraries.Add("setupapi.lib"); //  Required for access monitor device enumeration
 
-				// Windows Vista/7 Desktop Windows Manager API for Slate Windows Compliance
-				if( !SupportWindowsXP )		// Windows XP does not support DWM
-				{
-					InBuildTarget.GlobalLinkEnvironment.Config.AdditionalLibraries.Add("dwmapi.lib");
-				}
-
-            	// IME
-            	InBuildTarget.GlobalLinkEnvironment.Config.AdditionalLibraries.Add("imm32.lib");
-
-				// Setup assembly path for .NET modules.  This will only be used when CLR is enabled. (CPlusPlusCLR module types)
-				InBuildTarget.GlobalCompileEnvironment.Config.SystemDotNetAssemblyPaths.Add(
-					Environment.GetFolderPath(Environment.SpecialFolder.ProgramFilesX86) +
-					"/Reference Assemblies/Microsoft/Framework/.NETFramework/v4.0");
-
-				// Setup default assemblies for .NET modules.  These will only be used when CLR is enabled. (CPlusPlusCLR module types)
-				InBuildTarget.GlobalCompileEnvironment.Config.FrameworkAssemblyDependencies.Add("System.dll");
-				InBuildTarget.GlobalCompileEnvironment.Config.FrameworkAssemblyDependencies.Add("System.Data.dll");
-				InBuildTarget.GlobalCompileEnvironment.Config.FrameworkAssemblyDependencies.Add("System.Drawing.dll");
-				InBuildTarget.GlobalCompileEnvironment.Config.FrameworkAssemblyDependencies.Add("System.Xml.dll");
-				InBuildTarget.GlobalCompileEnvironment.Config.FrameworkAssemblyDependencies.Add("System.Management.dll");
-				InBuildTarget.GlobalCompileEnvironment.Config.FrameworkAssemblyDependencies.Add("System.Windows.Forms.dll");
-				InBuildTarget.GlobalCompileEnvironment.Config.FrameworkAssemblyDependencies.Add("WindowsBase.dll");
-			}
-
-			// Disable Simplygon support if compiling against the NULL RHI.
-			if (InBuildTarget.GlobalCompileEnvironment.Config.Definitions.Contains("USE_NULL_RHI=1"))
-			{
-				UEBuildConfiguration.bCompileSimplygon = false;
-			}
-
-			// For 64-bit builds, we'll forcibly ignore a linker warning with DirectInput.  This is
-			// Microsoft's recommended solution as they don't have a fixed .lib for us.
-			if (InBuildTarget.Platform == UnrealTargetPlatform.Win64)
-			{
-				InBuildTarget.GlobalLinkEnvironment.Config.AdditionalArguments += " /ignore:4078";
-			}
-		}
-
-		/**
-		 *	Setup the configuration environment for building
-		 *	
-		 *	@param	InBuildTarget		The target being built
-		 */
-		public override void SetUpConfigurationEnvironment(UEBuildTarget InBuildTarget)
-		{
-			// Determine the C++ compile/link configuration based on the Unreal configuration.
-			CPPTargetConfiguration CompileConfiguration;
-			//@todo SAS: Add a true Debug mode!
-			UnrealTargetConfiguration CheckConfig = InBuildTarget.Configuration;
-			switch (CheckConfig)
-			{
-				default:
-				case UnrealTargetConfiguration.Debug:
-					CompileConfiguration = CPPTargetConfiguration.Debug;
-					if( BuildConfiguration.bDebugBuildsActuallyUseDebugCRT )
-					{ 
-						InBuildTarget.GlobalCompileEnvironment.Config.Definitions.Add("_DEBUG=1"); // the engine doesn't use this, but lots of 3rd party stuff does
-					}
-					else
-					{
-						InBuildTarget.GlobalCompileEnvironment.Config.Definitions.Add("NDEBUG=1"); // the engine doesn't use this, but lots of 3rd party stuff does
-					}
-					InBuildTarget.GlobalCompileEnvironment.Config.Definitions.Add("UE_BUILD_DEBUG=1");
-					break;
-				case UnrealTargetConfiguration.DebugGame:
-					// Default to Development; can be overriden by individual modules.
-				case UnrealTargetConfiguration.Development:
-					CompileConfiguration = CPPTargetConfiguration.Development;
-					InBuildTarget.GlobalCompileEnvironment.Config.Definitions.Add("NDEBUG=1"); // the engine doesn't use this, but lots of 3rd party stuff does
-					InBuildTarget.GlobalCompileEnvironment.Config.Definitions.Add("UE_BUILD_DEVELOPMENT=1");
-					break;
-				case UnrealTargetConfiguration.Shipping:
-					CompileConfiguration = CPPTargetConfiguration.Shipping;
-					InBuildTarget.GlobalCompileEnvironment.Config.Definitions.Add("NDEBUG=1"); // the engine doesn't use this, but lots of 3rd party stuff does
-					InBuildTarget.GlobalCompileEnvironment.Config.Definitions.Add("UE_BUILD_SHIPPING=1");
-					break;
-				case UnrealTargetConfiguration.Test:
-					CompileConfiguration = CPPTargetConfiguration.Shipping;
-					InBuildTarget.GlobalCompileEnvironment.Config.Definitions.Add("NDEBUG=1"); // the engine doesn't use this, but lots of 3rd party stuff does
-					InBuildTarget.GlobalCompileEnvironment.Config.Definitions.Add("UE_BUILD_TEST=1");
-					break;
-			}
-
-			// Set up the global C++ compilation and link environment.
-			InBuildTarget.GlobalCompileEnvironment.Config.Target.Configuration = CompileConfiguration;
-			InBuildTarget.GlobalLinkEnvironment.Config.Target.Configuration    = CompileConfiguration;
-
-			// Create debug info based on the heuristics specified by the user.
-			InBuildTarget.GlobalCompileEnvironment.Config.bCreateDebugInfo =
-				!BuildConfiguration.bDisableDebugInfo && ShouldCreateDebugInfo(InBuildTarget.Platform, CheckConfig);
-
-			// NOTE: Even when debug info is turned off, we currently force the linker to generate debug info
-			//       anyway on Visual C++ platforms.  This will cause a PDB file to be generated with symbols
-			//       for most of the classes and function/method names, so that crashes still yield somewhat
-			//       useful call stacks, even though compiler-generate debug info may be disabled.  This gives
-			//       us much of the build-time savings of fully-disabled debug info, without giving up call
-			//       data completely.
-			InBuildTarget.GlobalLinkEnvironment.Config.bCreateDebugInfo = true;
-		}
-
-		/**
-		 *	Whether this platform should create debug information or not
-		 *	
-		 *	@param	InPlatform			The UnrealTargetPlatform being built
-		 *	@param	InConfiguration		The UnrealTargetConfiguration being built
-		 *	
-		 *	@return	bool				true if debug info should be generated, false if not
-		 */
-		public override bool ShouldCreateDebugInfo(UnrealTargetPlatform Platform, UnrealTargetConfiguration Configuration)
-		{
-			switch (Configuration)
-			{
-				case UnrealTargetConfiguration.Development: 
-				case UnrealTargetConfiguration.Shipping: 
-				case UnrealTargetConfiguration.Test: 
-					return !BuildConfiguration.bOmitPCDebugInfoInDevelopment;
-				case UnrealTargetConfiguration.DebugGame:
-				case UnrealTargetConfiguration.Debug:
-				default: 
-					return true;
-			};
-		}
-
-		/**
-		 *	Return whether this platform has uniquely named binaries across multiple games
-		 */
-		public override bool HasUniqueBinaries()
-		{
-			// Windows applications have many shared binaries between games
-			return false;
-		}
-	}
+                // Windows Vista/7 Desktop Windows Manager API for Slate Windows Compliance
+                if( !SupportWindowsXP )		// Windows XP does not support DWM
+                {
+                    InBuildTarget.GlobalLinkEnvironment.Config.AdditionalLibraries.Add("dwmapi.lib");
+                }
+
+                // IME
+                InBuildTarget.GlobalLinkEnvironment.Config.AdditionalLibraries.Add("imm32.lib");
+
+                // Setup assembly path for .NET modules.  This will only be used when CLR is enabled. (CPlusPlusCLR module types)
+                InBuildTarget.GlobalCompileEnvironment.Config.SystemDotNetAssemblyPaths.Add(
+                    Environment.GetFolderPath(Environment.SpecialFolder.ProgramFilesX86) +
+                    "/Reference Assemblies/Microsoft/Framework/.NETFramework/v4.0");
+
+                // Setup default assemblies for .NET modules.  These will only be used when CLR is enabled. (CPlusPlusCLR module types)
+                InBuildTarget.GlobalCompileEnvironment.Config.FrameworkAssemblyDependencies.Add("System.dll");
+                InBuildTarget.GlobalCompileEnvironment.Config.FrameworkAssemblyDependencies.Add("System.Data.dll");
+                InBuildTarget.GlobalCompileEnvironment.Config.FrameworkAssemblyDependencies.Add("System.Drawing.dll");
+                InBuildTarget.GlobalCompileEnvironment.Config.FrameworkAssemblyDependencies.Add("System.Xml.dll");
+                InBuildTarget.GlobalCompileEnvironment.Config.FrameworkAssemblyDependencies.Add("System.Management.dll");
+                InBuildTarget.GlobalCompileEnvironment.Config.FrameworkAssemblyDependencies.Add("System.Windows.Forms.dll");
+                InBuildTarget.GlobalCompileEnvironment.Config.FrameworkAssemblyDependencies.Add("WindowsBase.dll");
+            }
+
+            // Disable Simplygon support if compiling against the NULL RHI.
+            if (InBuildTarget.GlobalCompileEnvironment.Config.Definitions.Contains("USE_NULL_RHI=1"))
+            {
+                UEBuildConfiguration.bCompileSimplygon = false;
+            }
+
+            // For 64-bit builds, we'll forcibly ignore a linker warning with DirectInput.  This is
+            // Microsoft's recommended solution as they don't have a fixed .lib for us.
+            if (InBuildTarget.Platform == UnrealTargetPlatform.Win64)
+            {
+                InBuildTarget.GlobalLinkEnvironment.Config.AdditionalArguments += " /ignore:4078";
+            }
+        }
+
+        /**
+         *	Setup the configuration environment for building
+         *	
+         *	@param	InBuildTarget		The target being built
+         */
+        public override void SetUpConfigurationEnvironment(UEBuildTarget InBuildTarget)
+        {
+            // Determine the C++ compile/link configuration based on the Unreal configuration.
+            CPPTargetConfiguration CompileConfiguration;
+            //@todo SAS: Add a true Debug mode!
+            UnrealTargetConfiguration CheckConfig = InBuildTarget.Configuration;
+            switch (CheckConfig)
+            {
+                default:
+                case UnrealTargetConfiguration.Debug:
+                    CompileConfiguration = CPPTargetConfiguration.Debug;
+                    if( BuildConfiguration.bDebugBuildsActuallyUseDebugCRT )
+                    { 
+                        InBuildTarget.GlobalCompileEnvironment.Config.Definitions.Add("_DEBUG=1"); // the engine doesn't use this, but lots of 3rd party stuff does
+                    }
+                    else
+                    {
+                        InBuildTarget.GlobalCompileEnvironment.Config.Definitions.Add("NDEBUG=1"); // the engine doesn't use this, but lots of 3rd party stuff does
+                    }
+                    InBuildTarget.GlobalCompileEnvironment.Config.Definitions.Add("UE_BUILD_DEBUG=1");
+                    break;
+                case UnrealTargetConfiguration.DebugGame:
+                    // Default to Development; can be overriden by individual modules.
+                case UnrealTargetConfiguration.Development:
+                    CompileConfiguration = CPPTargetConfiguration.Development;
+                    InBuildTarget.GlobalCompileEnvironment.Config.Definitions.Add("NDEBUG=1"); // the engine doesn't use this, but lots of 3rd party stuff does
+                    InBuildTarget.GlobalCompileEnvironment.Config.Definitions.Add("UE_BUILD_DEVELOPMENT=1");
+                    break;
+                case UnrealTargetConfiguration.Shipping:
+                    CompileConfiguration = CPPTargetConfiguration.Shipping;
+                    InBuildTarget.GlobalCompileEnvironment.Config.Definitions.Add("NDEBUG=1"); // the engine doesn't use this, but lots of 3rd party stuff does
+                    InBuildTarget.GlobalCompileEnvironment.Config.Definitions.Add("UE_BUILD_SHIPPING=1");
+                    break;
+                case UnrealTargetConfiguration.Test:
+                    CompileConfiguration = CPPTargetConfiguration.Shipping;
+                    InBuildTarget.GlobalCompileEnvironment.Config.Definitions.Add("NDEBUG=1"); // the engine doesn't use this, but lots of 3rd party stuff does
+                    InBuildTarget.GlobalCompileEnvironment.Config.Definitions.Add("UE_BUILD_TEST=1");
+                    break;
+            }
+
+            // Set up the global C++ compilation and link environment.
+            InBuildTarget.GlobalCompileEnvironment.Config.Target.Configuration = CompileConfiguration;
+            InBuildTarget.GlobalLinkEnvironment.Config.Target.Configuration    = CompileConfiguration;
+
+            // Create debug info based on the heuristics specified by the user.
+            InBuildTarget.GlobalCompileEnvironment.Config.bCreateDebugInfo =
+                !BuildConfiguration.bDisableDebugInfo && ShouldCreateDebugInfo(InBuildTarget.Platform, CheckConfig);
+
+            // NOTE: Even when debug info is turned off, we currently force the linker to generate debug info
+            //       anyway on Visual C++ platforms.  This will cause a PDB file to be generated with symbols
+            //       for most of the classes and function/method names, so that crashes still yield somewhat
+            //       useful call stacks, even though compiler-generate debug info may be disabled.  This gives
+            //       us much of the build-time savings of fully-disabled debug info, without giving up call
+            //       data completely.
+            InBuildTarget.GlobalLinkEnvironment.Config.bCreateDebugInfo = true;
+        }
+
+        /**
+         *	Whether this platform should create debug information or not
+         *	
+         *	@param	InPlatform			The UnrealTargetPlatform being built
+         *	@param	InConfiguration		The UnrealTargetConfiguration being built
+         *	
+         *	@return	bool				true if debug info should be generated, false if not
+         */
+        public override bool ShouldCreateDebugInfo(UnrealTargetPlatform Platform, UnrealTargetConfiguration Configuration)
+        {
+            switch (Configuration)
+            {
+                case UnrealTargetConfiguration.Development: 
+                case UnrealTargetConfiguration.Shipping: 
+                case UnrealTargetConfiguration.Test: 
+                    return !BuildConfiguration.bOmitPCDebugInfoInDevelopment;
+                case UnrealTargetConfiguration.DebugGame:
+                case UnrealTargetConfiguration.Debug:
+                default: 
+                    return true;
+            };
+        }
+
+        /**
+         *	Return whether this platform has uniquely named binaries across multiple games
+         */
+        public override bool HasUniqueBinaries()
+        {
+            // Windows applications have many shared binaries between games
+            return false;
+        }
+    }
 }
