Rolling back //depot/UE4/Engine/Source/Programs/UnrealBuildTool/UnrealBuildTool_Mono.csproj to revision 36

- Target 4.5 since we don't have 4.0 in the stripped bundled version of Mono we ship on Mac
- Remove conditions from the configuration (while xbuild supports conditions, MonoDevelop/Xamarin Studio don't).

diff --git a/Engine/Source/Programs/UnrealBuildTool/UnrealBuildTool_Mono.csproj b/Engine/Source/Programs/UnrealBuildTool/UnrealBuildTool_Mono.csproj
--- a/Engine/Source/Programs/UnrealBuildTool/UnrealBuildTool_Mono.csproj
+++ b/Engine/Source/Programs/UnrealBuildTool/UnrealBuildTool_Mono.csproj
@@ -12,7 +12,7 @@
     <AssemblyName>UnrealBuildTool</AssemblyName>
     <RunPostBuildEvent>OnOutputUpdated</RunPostBuildEvent>
     <IsWebBootstrapper>false</IsWebBootstrapper>
-    <TargetFrameworkVersion>v4.0</TargetFrameworkVersion>
+    <TargetFrameworkVersion>v4.5</TargetFrameworkVersion>
     <FileUpgradeFlags>
     </FileUpgradeFlags>
     <OldToolsVersion>3.5</OldToolsVersion>
@@ -88,110 +88,10 @@
     <Reference Include="System.Xml" />
     <Reference Include="System.Xml.Linq" />
   </ItemGroup>
-  <Choose>
-    <When Condition="Exists('XboxOne\UEBuildXboxOne.cs')">
-      <ItemGroup>
-        <Compile Include="XboxOne\UEBuildXboxOne.cs" />
-        <Compile Include="XboxOne\XboxOneToolChain.cs" />
-        <Compile Include="XboxOne\XboxOneDeploy.cs" />
-        <Compile Include="XboxOne\XboxOneProjectGenerator.cs" />
-      </ItemGroup>
-    </When>
-  </Choose>
-  <Choose>
-    <When Condition="Exists('PS4\UEBuildPS4.cs')">
-      <ItemGroup>
-        <Compile Include="PS4\PS4ProjectGenerator.cs" />
-        <Compile Include="PS4\PS4ToolChain.cs" />
-        <Compile Include="PS4\UEBuildPS4.cs" />
-      </ItemGroup>
-    </When>
-  </Choose>
-  <Choose>
-    <When Condition="Exists('Cygwin\UEBuildCygwin.cs')">
-      <ItemGroup>
-        <Compile Include="Cygwin\CygwinToolChain.cs" />
-        <Compile Include="Cygwin\UEBuildCygwin.cs" />
-        <Compile Include="Cygwin\CygwinProjectGenerator.cs" />
-      </ItemGroup>
-    </When>
-  </Choose>
-  <Choose>
-    <When Condition="Exists('Mac\UEBuildMac.cs')">
-      <ItemGroup>
-        <Compile Include="Mac\MacToolChain.cs" />
-        <Compile Include="Mac\UEBuildMac.cs" />
-        <Compile Include="Mac\MacProjectGenerator.cs" />
-        <Compile Include="Mac\UEDeployMac.cs" />
-      </ItemGroup>
-    </When>
-  </Choose>
-  <Choose>
-    <When Condition="Exists('Android\UEBuildAndroid.cs')">
-      <ItemGroup>
-        <Compile Include="Android\UEBuildAndroid.cs" />
-        <Compile Include="Android\AndroidToolChain.cs" />
-        <Compile Include="Android\AndroidProjectGenerator.cs" />
-      </ItemGroup>
-    </When>
-  </Choose>
-  <Choose>
-    <When Condition="Exists('WinRT\UEBuildWinRT.cs')">
-      <ItemGroup>
-        <Compile Include="WinRT\UEBuildWinRT.cs" />
-        <Compile Include="WinRT\WinRTToolChain.cs" />
-        <Compile Include="WinRT\WinRTDeploy.cs" />
-        <Compile Include="WinRT\WinRTProjectGenerator.cs" />
-      </ItemGroup>
-    </When>
-  </Choose>
-  <Choose>
-    <When Condition="Exists('IOS\UEBuildIOS.cs')">
-      <ItemGroup>
-        <Compile Include="IOS\IOSProjectGenerator.cs" />
-        <Compile Include="IOS\IOSToolChain.cs" />
-        <Compile Include="IOS\UEBuildIOS.cs" />
-        <Compile Include="IOS\UEDeployIOS.cs" />
-      </ItemGroup>
-    </When>
-  </Choose>
-  <Choose>
-    <When Condition="Exists('HTML5\UEBuildHTML5.cs')">
-      <ItemGroup>
-        <Compile Include="HTML5\HTML5ProjectGenerator.cs" />
-        <Compile Include="HTML5\HTML5ToolChain.cs" />
-        <Compile Include="HTML5\UEBuildHTML5.cs" />
-      </ItemGroup>
-    </When>
-  </Choose>
-  <Choose>
-    <When Condition="Exists('Linux\UEBuildLinux.cs')">
-      <ItemGroup>
-        <Compile Include="Linux\LinuxProjectGenerator.cs" />
-        <Compile Include="Linux\LinuxToolChain.cs" />
-        <Compile Include="Linux\UEBuildLinux.cs" />
-      </ItemGroup>
-    </When>
-  </Choose>
-  <Choose>
-    <When Condition="Exists('System\NotForLicensees\TelemetryProviderET.cs')">
-      <ItemGroup>
-        <Compile Include="System\NotForLicensees\TelemetryProviderET.cs" />
-      </ItemGroup>
-    </When>
-  </Choose>
-  <Choose>
-    <When Condition="Exists('ToolChain\NoRedist\EpicMachineNames.cs')">
-      <ItemGroup>
-        <Compile Include="ToolChain\NoRedist\EpicMachineNames.cs" />
-      </ItemGroup>
-    </When>
-  </Choose>
   <ItemGroup>
     <Compile Include="..\DotNETCommon\MetaData.cs">
       <Link>Properties\MetaData.cs</Link>
     </Compile>
-    <Compile Include="Android\UEDeployAndroid.cs" />
     <Compile Include="Configuration\UEBuildClient.cs" />
     <Compile Include="Configuration\UEBuildDeploy.cs" />
     <Compile Include="Configuration\UEBuildBinary.cs" />
@@ -207,6 +107,15 @@
     <Compile Include="fastJSON\JsonSerializer.cs" />
     <Compile Include="fastJSON\Reflection.cs" />
     <Compile Include="fastJSON\SafeDictionary.cs" />
+    <Compile Include="HTML5\HTML5ProjectGenerator.cs" />
+    <Compile Include="HTML5\HTML5ToolChain.cs" />
+    <Compile Include="HTML5\UEBuildHTML5.cs" />
+    <Compile Include="IOS\IOSProjectGenerator.cs" />
+    <Compile Include="IOS\IOSToolChain.cs" />
+    <Compile Include="IOS\UEBuildIOS.cs" />
+    <Compile Include="Linux\LinuxProjectGenerator.cs" />
+    <Compile Include="Linux\LinuxToolChain.cs" />
+    <Compile Include="Linux\UEBuildLinux.cs" />
     <Compile Include="Properties\AssemblyInfo.cs" />
     <Compile Include="Configuration\UEBuildConfiguration.cs" />
     <Compile Include="Configuration\UEBuildTarget.cs" />
@@ -214,9 +123,6 @@
     <Compile Include="System\ActionHistory.cs" />
     <Compile Include="System\BuildException.cs" />
     <Compile Include="Configuration\BuildConfiguration.cs" />
-    <Compile Include="System\ModuleProcessingException.cs" />
-    <Compile Include="System\MakefileGenerator.cs" />
-    <Compile Include="System\VCSolutionOptions.cs" />
     <Compile Include="System\DirectoryLookupCache.cs" />
     <Compile Include="System\CPPEnvironment.cs" />
     <Compile Include="System\CPPHeaders.cs" />
@@ -227,6 +133,7 @@
     <Compile Include="System\JunkDeleter.cs" />
     <Compile Include="System\LinkEnvironment.cs" />
     <Compile Include="System\LocalExecutor.cs" />
+    <Compile Include="System\MakefileGenerator.cs" />
     <Compile Include="System\Plugins.cs" />
     <Compile Include="System\Project.cs" />
     <Compile Include="System\ProjectFileGenerator.cs" />
@@ -250,14 +157,20 @@
     <Compile Include="System\XGE.cs" />
     <Compile Include="Utilities\CopyrightVerify.cs" />
     <Compile Include="Utilities\DictionaryExtensions.cs" />
+    <Compile Include="Utilities\XmlConfigLoader.cs" />
     <Compile Include="Utilities\GraphVisualization.cs">
       <SubType>Code</SubType>
     </Compile>
-    <Compile Include="Configuration\EngineConfiguration.cs" />
-    <Compile Include="Utilities\XmlConfigLoader.cs" />
     <Compile Include="Windows\UEBuildWindows.cs" />
     <Compile Include="Windows\VCToolChain.cs" />
     <Compile Include="Windows\WindowsProjectGenerator.cs" />
+    <Compile Include="Mac\MacToolChain.cs" />
+    <Compile Include="Mac\UEBuildMac.cs" />
+    <Compile Include="Mac\MacProjectGenerator.cs" />
+    <Compile Include="Configuration\EngineConfiguration.cs" />
+    <Compile Include="IOS\UEDeployIOS.cs" />
+    <Compile Include="System\ModuleProcessingException.cs" />
+    <Compile Include="System\VCSolutionOptions.cs" />
   </ItemGroup>
   <ItemGroup>
     <BootstrapperPackage Include="Microsoft.Net.Client.3.5">
@@ -303,4 +216,4 @@
     <PreBuildEvent>
     </PreBuildEvent>
   </PropertyGroup>
-</Project>
\ No newline at end of file
+</Project>
