Fixed monolithic link errors (regression)
- Partially backed out experimental UBT changes from CL 2223639

diff --git a/Engine/Source/Programs/UnrealBuildTool/Configuration/UEBuildModule.cs b/Engine/Source/Programs/UnrealBuildTool/Configuration/UEBuildModule.cs
--- a/Engine/Source/Programs/UnrealBuildTool/Configuration/UEBuildModule.cs
+++ b/Engine/Source/Programs/UnrealBuildTool/Configuration/UEBuildModule.cs
@@ -1062,9 +1062,6 @@ namespace UnrealBuildTool
 				return LinkInputFiles;
 			}
 
-			// Process all of the header file dependencies for this module
-			this.ProcessAllCppDependencies( ModuleCompileEnvironment );
-
 			// Throw an error if the module's source file list referenced any non-existent files.
 			if (SourceFilesToBuild.MissingFiles.Count > 0)
 			{
diff --git a/Engine/Source/Programs/UnrealBuildTool/Configuration/UEBuildTarget.cs b/Engine/Source/Programs/UnrealBuildTool/Configuration/UEBuildTarget.cs
--- a/Engine/Source/Programs/UnrealBuildTool/Configuration/UEBuildTarget.cs
+++ b/Engine/Source/Programs/UnrealBuildTool/Configuration/UEBuildTarget.cs
@@ -1514,9 +1514,27 @@ namespace UnrealBuildTool
 							throw new BuildException( "Expecting module {0} to have bIncludeInTarget set", DependencyModuleCPP.Name );
 						}
 
+						var ModuleCompileEnvironment = DependencyModuleCPP.CreateModuleCompileEnvironment(GlobalCompileEnvironment);
+						DependencyModuleCPP.ProcessAllCppDependencies(ModuleCompileEnvironment);
+
 						var UHTModuleInfo = DependencyModuleCPP.GetUHTModuleInfo();
 						if( UHTModuleInfo.PublicUObjectClassesHeaders.Count > 0 || UHTModuleInfo.PrivateUObjectHeaders.Count > 0 || UHTModuleInfo.PublicUObjectHeaders.Count > 0 )
 						{
+							UHTModuleInfo.PCH                      = "";
+							UHTModuleInfo.GeneratedCPPFilenameBase = Path.Combine( UEBuildModuleCPP.GetGeneratedCodeDirectoryForModule(this, UHTModuleInfo.ModuleDirectory, UHTModuleInfo.ModuleName), UHTModuleInfo.ModuleName ) + ".generated";
+  
+							if (DependencyModuleCPP.ProcessedDependencies.UniquePCHHeaderFile != null)
+							{
+								UHTModuleInfo.PCH = DependencyModuleCPP.ProcessedDependencies.UniquePCHHeaderFile.AbsolutePath;
+							}
+
+							// If we've got this far and there are no source files then it's likely we're running Rocket and ignoring
+							// engine files, so we don't need a .generated.cpp either
+							if (DependencyModuleCPP.SourceFilesToBuild.Count != 0)
+							{
+								DependencyModuleCPP.AutoGenerateCppInfo = new UEBuildModuleCPP.AutoGenerateCppInfoClass( UHTModuleInfo.GeneratedCPPFilenameBase + ".cpp" );
+							}
+  
 							UObjectModules.Add( UHTModuleInfo );
 							Log.TraceVerbose( "Detected UObject module: " + UHTModuleInfo.ModuleName );
 						}
@@ -1536,7 +1554,7 @@ namespace UnrealBuildTool
 					string ModuleInfoFileName = Path.GetFullPath( Path.Combine( ProjectIntermediateDirectory, "UnrealHeaderTool.manifest" ) );
 
 					ECompilationResult UHTResult = ECompilationResult.OtherCompilationError;
-					if (!ExternalExecution.ExecuteHeaderToolIfNecessary(this, GlobalCompileEnvironment, UObjectModules, ModuleInfoFileName, ref UHTResult))
+					if (!ExternalExecution.ExecuteHeaderToolIfNecessary(this, UObjectModules, ModuleInfoFileName, ref UHTResult))
 					{
 						Log.TraceInformation("UnrealHeaderTool failed for target '" + TargetName + "' (platform: " + Platform.ToString() + ", module info: " + ModuleInfoFileName + ").");
 						return UHTResult;
diff --git a/Engine/Source/Programs/UnrealBuildTool/System/ExternalExecution.cs b/Engine/Source/Programs/UnrealBuildTool/System/ExternalExecution.cs
--- a/Engine/Source/Programs/UnrealBuildTool/System/ExternalExecution.cs
+++ b/Engine/Source/Programs/UnrealBuildTool/System/ExternalExecution.cs
@@ -453,7 +453,7 @@ namespace UnrealBuildTool
 		 * Builds and runs the header tool and touches the header directories.
 		 * Performs any early outs if headers need no changes, given the UObject modules, tool path, game name, and configuration
 		 */
-		public static bool ExecuteHeaderToolIfNecessary( UEBuildTarget Target, CPPEnvironment GlobalCompileEnvironment, List<UHTModuleInfo> UObjectModules, string ModuleInfoFileName, ref ECompilationResult UHTResult )
+		public static bool ExecuteHeaderToolIfNecessary( UEBuildTarget Target, List<UHTModuleInfo> UObjectModules, string ModuleInfoFileName, ref ECompilationResult UHTResult )
 		{
 			using (ProgressWriter Progress = new ProgressWriter("Generating headers...", false))
 			{
@@ -465,36 +465,9 @@ namespace UnrealBuildTool
 				var ToolChain = UEToolChain.GetPlatformToolChain(CppPlatform);
 				var RootLocalPath  = Path.GetFullPath(ProjectFileGenerator.RootRelativePath);
 
-
-				// ensure the headers are up to date
 				bool bUHTNeedsToRun = (UEBuildConfiguration.bForceHeaderGeneration == true || AreGeneratedCodeFilesOutOfDate(Target, UObjectModules));
-				if( bUHTNeedsToRun )
-				{
-					// Since code files are definitely out of date, we'll now finish computing information about the UObject modules for UHT.  We
-					// want to save this work until we know that UHT actually needs to be run to speed up best-case iteration times.
-					foreach( var UHTModuleInfo in UObjectModules )
-					{
-						UHTModuleInfo.PCH = "";
-						UHTModuleInfo.GeneratedCPPFilenameBase = Path.Combine( UEBuildModuleCPP.GetGeneratedCodeDirectoryForModule(Target, UHTModuleInfo.ModuleDirectory, UHTModuleInfo.ModuleName), UHTModuleInfo.ModuleName ) + ".generated";
-
-						// We need to figure out which PCH header this module is including, so that UHT can inject an include statement for it into any .cpp files it is synthesizing
-						var DependencyModuleCPP = (UEBuildModuleCPP)Target.GetModuleByName( UHTModuleInfo.ModuleName );
-						var ModuleCompileEnvironment = DependencyModuleCPP.CreateModuleCompileEnvironment(GlobalCompileEnvironment);
-						DependencyModuleCPP.ProcessAllCppDependencies(ModuleCompileEnvironment);
-						if (DependencyModuleCPP.ProcessedDependencies.UniquePCHHeaderFile != null)
-						{
-							UHTModuleInfo.PCH = DependencyModuleCPP.ProcessedDependencies.UniquePCHHeaderFile.AbsolutePath;
-						}
-
-						// If we've got this far and there are no source files then it's likely we're running Rocket and ignoring
-						// engine files, so we don't need a .generated.cpp either
-						if (DependencyModuleCPP.SourceFilesToBuild.Count != 0)
-						{
-							DependencyModuleCPP.AutoGenerateCppInfo = new UEBuildModuleCPP.AutoGenerateCppInfoClass( UHTModuleInfo.GeneratedCPPFilenameBase + ".cpp" );
-						}
-					}
-				}
 
+				// ensure the headers are up to date
 				UHTManifest Manifest = new UHTManifest(Target, RootLocalPath, ToolChain.ConvertPath(RootLocalPath + '\\'), UObjectModules);
 
 				if( !bIsBuildingUHT && bUHTNeedsToRun )
