UnrealBuildTool performance: Disable verbose output in Release builds
- TRACE is no longer defined by default in Release builds.
- If you want bPrintDebugInfo to work, use a Debug build or define TRACE to 1
- Also changed .uprojectdir missing folder spew to be verbose

diff --git a/Engine/Source/Programs/UnrealBuildTool/System/UProjectInfo.cs b/Engine/Source/Programs/UnrealBuildTool/System/UProjectInfo.cs
--- a/Engine/Source/Programs/UnrealBuildTool/System/UProjectInfo.cs
+++ b/Engine/Source/Programs/UnrealBuildTool/System/UProjectInfo.cs
@@ -198,7 +198,7 @@ namespace UnrealBuildTool
 				}
 				else
 				{
-					Console.WriteLine("ProjectInfo: Skipping directory {0} from .uprojectdirs file as it doesn't exist.", DirToSearch);
+					Log.TraceVerbose("ProjectInfo: Skipping directory {0} from .uprojectdirs file as it doesn't exist.", DirToSearch);
 				}
 			}
 
diff --git a/Engine/Source/Programs/UnrealBuildTool/UnrealBuildTool.csproj b/Engine/Source/Programs/UnrealBuildTool/UnrealBuildTool.csproj
--- a/Engine/Source/Programs/UnrealBuildTool/UnrealBuildTool.csproj
+++ b/Engine/Source/Programs/UnrealBuildTool/UnrealBuildTool.csproj
@@ -38,7 +38,8 @@
     <DebugType>pdbonly</DebugType>
     <Optimize>true</Optimize>
     <OutputPath>..\..\..\Binaries\DotNET\</OutputPath>
-    <DefineConstants>TRACE</DefineConstants>
+    <DefineConstants>
+    </DefineConstants>
     <Optimize>true</Optimize>
     <DebugType>pdbonly</DebugType>
     <PlatformTarget>AnyCPU</PlatformTarget>
