enable AutoSDKs for HTML5
#ue4
#html5

diff --git a/Engine/Source/Programs/UnrealBuildTool/HTML5/HTML5ToolChain.cs b/Engine/Source/Programs/UnrealBuildTool/HTML5/HTML5ToolChain.cs
--- a/Engine/Source/Programs/UnrealBuildTool/HTML5/HTML5ToolChain.cs
+++ b/Engine/Source/Programs/UnrealBuildTool/HTML5/HTML5ToolChain.cs
@@ -32,10 +32,6 @@ namespace UnrealBuildTool
 					{
 						PythonPath = Utils.IsRunningOnMono ? "python" : "python.exe";
 					}
-					else
-					{
-						PythonPath += Utils.IsRunningOnMono ? "/python" : "/python.exe";
-					}
                     EMCCPath = "\"" + EMCCPath + "\""; 
 					// set some environment variable we'll need
 					//Environment.SetEnvironmentVariable("EMCC_DEBUG", "cache");
diff --git a/Engine/Source/Programs/UnrealBuildTool/HTML5/UEBuildHTML5.cs b/Engine/Source/Programs/UnrealBuildTool/HTML5/UEBuildHTML5.cs
--- a/Engine/Source/Programs/UnrealBuildTool/HTML5/UEBuildHTML5.cs
+++ b/Engine/Source/Programs/UnrealBuildTool/HTML5/UEBuildHTML5.cs
@@ -24,13 +24,43 @@ namespace UnrealBuildTool
 
         static private string TargetPlatformName = "HTML5";
 
-
-        public override string GetSDKTargetPlatformName()
+		// This is the SDK version we support in //CarefullyRedist.
+		// May include minor revisions or descriptions that a default install from SDK_Manager won't have.
+		// e.g. 1.600_Patch001, or 1.610.001.  The SDK_Manager always installs minor revision patches straight into
+		// the default major revision folder.
+		static private string ExpectedSDKVersion = "1.21.0";
+
+		/** 
+		 * Whether platform supports switching SDKs during runtime
+		 * 
+		 * @return true if supports
+		 */
+		protected override bool PlatformSupportsAutoSDKs()
+		{
+			return true;
+		}
+		
+		public override string GetSDKTargetPlatformName()
         {
             return TargetPlatformName;
         }
 
-        // The current architecture - affects everything about how UBT operates on HTML5
+		/** 
+		 * Returns SDK string as required by the platform 
+		 * 
+		 * @return Valid SDK string
+		 */
+		protected override string GetRequiredSDKString()
+		{
+			return ExpectedSDKVersion;
+		}
+
+		protected override String GetRequiredScriptVersionString()
+		{
+			return "1.0";
+		}
+		
+		// The current architecture - affects everything about how UBT operates on HTML5
         public override string GetActiveArchitecture()
         {
             // by default, use an empty architecture (which is really just a modifier to the platform for some paths/names)
@@ -49,8 +79,28 @@ namespace UnrealBuildTool
          */
         protected override SDKStatus HasRequiredManualSDKInternal()
         {
-            string BaseSDKPath = Environment.GetEnvironmentVariable("EMSCRIPTEN");
-            return (string.IsNullOrEmpty(BaseSDKPath) == false) ? SDKStatus.Valid : SDKStatus.Invalid;
+			// if any autosdk setup has been done then the local process environment is suspect
+			if (HasSetupAutoSDK())
+			{
+				return SDKStatus.Invalid;
+			}
+
+			string BaseSDKPath = Environment.GetEnvironmentVariable("EMSCRIPTEN");
+			if (!string.IsNullOrEmpty(BaseSDKPath))
+			{
+				// Check for the *actual* sdk version
+				// Note: This depends on keeping the installed folder name the version #!!!
+				if (BaseSDKPath.Contains(ExpectedSDKVersion))
+				{
+
+					return SDKStatus.Valid;
+				}
+				else
+				{
+					return SDKStatus.Invalid;
+				}
+			}
+			return SDKStatus.Invalid;
         }
 
         public override bool CanUseXGE()
