CORE: UBT should detect and prevent linkage that violates our EULA.

If UBT is building a non-editor configuration it throws an error if the produced binaries has non-redistributable dependencies.

User can avoid throwing error and change it to warning at the end with flag BuildConfiguration.bBreakBuildOnLicenseViolation = false.

In general module to be redistributable it must not be neither Development nor Editor type. This behavior can be overridden though. One have to just assign true or false to IsRedistributableOverride field of ModuleRules in <ModuleName>.Build.cs file.

UBT will not check for the dependencies also if the module being built is non-redistributable (e.g. some internal programs).

#ttp 335755

diff --git a/Engine/Source/Programs/UnrealBuildTool/Configuration/BuildConfiguration.cs b/Engine/Source/Programs/UnrealBuildTool/Configuration/BuildConfiguration.cs
--- a/Engine/Source/Programs/UnrealBuildTool/Configuration/BuildConfiguration.cs
+++ b/Engine/Source/Programs/UnrealBuildTool/Configuration/BuildConfiguration.cs
@@ -336,6 +336,12 @@ namespace UnrealBuildTool
 		[XmlConfig]
 		public static bool bDebugBuildsActuallyUseDebugCRT;
 
+		/// <summary>
+		/// Tells the UBT to break build if module currently being built is violating EULA.
+		/// </summary>
+		[XmlConfig]
+		public static bool bBreakBuildOnLicenseViolation;
+
 		/// <summary>
 		/// Sets the configuration back to defaults.
 		/// </summary>
@@ -438,6 +444,8 @@ namespace UnrealBuildTool
 
 			// set up some paths
 			BaseIntermediateFolder = "Intermediate/Build/";
+
+			bBreakBuildOnLicenseViolation = true;
 		}
 
 		/// <summary>
diff --git a/Engine/Source/Programs/UnrealBuildTool/Configuration/UEBuildModule.cs b/Engine/Source/Programs/UnrealBuildTool/Configuration/UEBuildModule.cs
--- a/Engine/Source/Programs/UnrealBuildTool/Configuration/UEBuildModule.cs
+++ b/Engine/Source/Programs/UnrealBuildTool/Configuration/UEBuildModule.cs
@@ -119,6 +119,21 @@ namespace UnrealBuildTool
 		/** An optional output path for the module */
 		public readonly string OutputDirectory;
 
+		/** Is this module allowed to be redistributed. */
+		private readonly bool? IsRedistributableOverride;
+
+		/**
+		 * Tells if this module can be redistributed.
+		 * 
+		 * @returns True if this module can be redistributed. False otherwise.
+		 */
+		public bool IsRedistributable()
+		{
+			return IsRedistributableOverride.HasValue
+				? IsRedistributableOverride.Value
+				: (Type != UEBuildModuleType.Developer && Type != UEBuildModuleType.Editor);
+		}
+
 		/** The binary the module will be linked into for the current target.  Only set after UEBuildBinary.BindModules is called. */
 		public UEBuildBinary Binary = null;
 
@@ -170,6 +185,7 @@ namespace UnrealBuildTool
 			UEBuildModuleType InType,
 			string InModuleDirectory,
 			string InOutputDirectory,
+			bool? InIsRedistributableOverride,
 			IEnumerable<string> InPublicDefinitions = null,
 			IEnumerable<string> InPublicIncludePaths = null,
 			IEnumerable<string> InPublicSystemIncludePaths = null,
@@ -213,6 +229,7 @@ namespace UnrealBuildTool
 			CircularlyReferencedDependentModules = new HashSet<string>( ListFromOptionalEnumerableStringParameter( InCircularlyReferencedDependentModules ) );
 			DynamicallyLoadedModuleNames = ListFromOptionalEnumerableStringParameter( InDynamicallyLoadedModuleNames );
             PlatformSpecificDynamicallyLoadedModuleNames = ListFromOptionalEnumerableStringParameter(InPlatformSpecificDynamicallyLoadedModuleNames);
+			IsRedistributableOverride = InIsRedistributableOverride;
 
 			Target.RegisterModule(this);
 		}
@@ -764,6 +781,7 @@ namespace UnrealBuildTool
 			string InName,
 			string InModuleDirectory,
 			string InOutputDirectory,
+			bool? InIsRedistributableOverride,
 			IEnumerable<string> InPublicDefinitions = null,
 			IEnumerable<string> InPublicIncludePaths = null,
 			IEnumerable<string> InPublicSystemIncludePaths = null,
@@ -782,6 +800,7 @@ namespace UnrealBuildTool
 			InName:							InName,
 			InModuleDirectory:				InModuleDirectory,
 			InOutputDirectory:				InOutputDirectory,
+			InIsRedistributableOverride:		InIsRedistributableOverride,
 			InPublicDefinitions:			InPublicDefinitions,
 			InPublicIncludePaths:			InPublicIncludePaths,
 			InPublicSystemIncludePaths:		InPublicSystemIncludePaths,
@@ -954,6 +973,7 @@ namespace UnrealBuildTool
 			UEBuildModuleType InType,
 			string InModuleDirectory,
 			string InOutputDirectory,
+			bool? InIsRedistributableOverride,
 			IntelliSenseGatherer InIntelliSenseGatherer,
 			IEnumerable<FileItem> InSourceFiles,
 			IEnumerable<string> InPublicIncludePaths,
@@ -987,7 +1007,8 @@ namespace UnrealBuildTool
 					InName, 
 					InType,
 					InModuleDirectory,
-					InOutputDirectory, 
+					InOutputDirectory,
+					InIsRedistributableOverride,
 					InDefinitions, 
 					InPublicIncludePaths,
 					InPublicSystemIncludePaths, 
@@ -1999,6 +2020,7 @@ namespace UnrealBuildTool
 			UEBuildModuleType InType,
 			string InModuleDirectory,
 			string InOutputDirectory,
+			bool? InIsRedistributableOverride,
 			IntelliSenseGatherer InIntelliSenseGatherer,
 			IEnumerable<FileItem> InSourceFiles,
 			IEnumerable<string> InPublicIncludePaths,
@@ -2029,7 +2051,7 @@ namespace UnrealBuildTool
 			bool InEnableExceptions,
 			bool bInBuildSourceFiles
 			)
-		: base(InTarget,InName,InType,InModuleDirectory,InOutputDirectory,InIntelliSenseGatherer,
+			: base(InTarget,InName,InType,InModuleDirectory,InOutputDirectory,InIsRedistributableOverride,InIntelliSenseGatherer,
 			InSourceFiles,InPublicIncludePaths,InPublicSystemIncludePaths,InDefinitions,
 			InPublicIncludePathModuleNames,InPublicDependencyModuleNames,InPublicDelayLoadDLLs,InPublicAdditionalLibraries,InPublicFrameworks,InPublicWeakFrameworks,InPublicAdditionalFrameworks,InPublicAdditionalShadowFiles,
 			InPrivateIncludePaths,InPrivateIncludePathModuleNames,InPrivateDependencyModuleNames,
diff --git a/Engine/Source/Programs/UnrealBuildTool/Configuration/UEBuildTarget.cs b/Engine/Source/Programs/UnrealBuildTool/Configuration/UEBuildTarget.cs
--- a/Engine/Source/Programs/UnrealBuildTool/Configuration/UEBuildTarget.cs
+++ b/Engine/Source/Programs/UnrealBuildTool/Configuration/UEBuildTarget.cs
@@ -1069,7 +1069,7 @@
 						string UnixLibraryFileName = Path.Combine(LibraryPath, "lib" + LibraryName + LibraryExtension);
 						if(File.Exists(UnixLibraryFileName))
 						{
-                            FileNames.Add(UnixLibraryFileName);
+							FileNames.Add(UnixLibraryFileName);
 						}
 					}
 				}
@@ -1216,187 +1216,14 @@
 		}
 
 		/** Builds the target, appending list of output files and returns building result. */
-		public ECompilationResult Build(IUEToolChain TargetToolChain, List<FileItem> OutputItems)
+		public ECompilationResult Build(IUEToolChain TargetToolChain, List<FileItem> OutputItems, out string EULAViolationWarning)
 		{
-			// Set up the global compile and link environment in GlobalCompileEnvironment and GlobalLinkEnvironment.
-			SetupGlobalEnvironment();
+			var SpecialRocketLibFilesThatAreBuildProducts = PreBuildSetup();
 
-			// Setup the target's modules.
-			SetupModules();
+			EULAViolationWarning = !ProjectFileGenerator.bGenerateProjectFiles
+				? CheckForEULAViolation()
+				: null;
 
-			// Setup the target's binaries.
-			SetupBinaries();
-
-			// Setup the target's plugins
-			SetupPlugins();
-
-            var SpecialRocketLibFilesThatAreBuildProducts = new List<string>();
-
-			// Add the enabled plugins to the build
-			foreach (PluginInfo Plugin in EnabledPlugins)
-			{
-                SpecialRocketLibFilesThatAreBuildProducts.AddRange(AddPlugin(Plugin));
-			}
-
-			// Describe what's being built.
-			Log.TraceVerbose("Building {0} - {1} - {2} - {3}", AppName, GameName, Platform, Configuration);
-
-			// Put the non-executable output files (PDB, import library, etc) in the intermediate directory.
-			GlobalLinkEnvironment.Config.IntermediateDirectory = Path.GetFullPath(GlobalCompileEnvironment.Config.OutputDirectory);
-			GlobalLinkEnvironment.Config.OutputDirectory = GlobalLinkEnvironment.Config.IntermediateDirectory;
-
-			// By default, shadow source files for this target in the root OutputDirectory
-			GlobalLinkEnvironment.Config.LocalShadowDirectory = GlobalLinkEnvironment.Config.OutputDirectory;
-
-			// Add all of the extra modules, including game modules, that need to be compiled along
-			// with this app.  These modules aren't necessarily statically linked against, but may
-			// still be required at runtime in order for the application to load and function properly!
-			AddExtraModules();
-
-			// Bind modules for all app binaries.
-			foreach (var Binary in AppBinaries)
-			{
-				Binary.BindModules();
-			}
-
-			// Process all referenced modules and create new binaries for DLL dependencies if needed
-			var NewBinaries = new List<UEBuildBinary>();
-			foreach (var Binary in AppBinaries)
-			{
-				// Add binaries for all of our dependent modules
-				var FoundBinaries = Binary.ProcessUnboundModules();
-				if (FoundBinaries != null)
-				{
-					NewBinaries.AddRange(FoundBinaries);
-				}
-			}
-			AppBinaries.AddRange( NewBinaries );
-
-			// On Mac AppBinaries paths for non-console targets need to be adjusted to be inside the app bundle
-			if (GlobalLinkEnvironment.Config.Target.Platform == CPPTargetPlatform.Mac && !GlobalLinkEnvironment.Config.bIsBuildingConsoleApplication)
-			{
-				MacToolChain.FixBundleBinariesPaths( this, AppBinaries );
-			}
-
-			// If we're building a single module, then find the binary for that module and add it to our target
-			if (OnlyModules.Count > 0)
-			{
-				var FilteredBinaries = new List<UEBuildBinary>();
-
-				var AnyBinariesAdded = false;
-				foreach (var DLLBinary in AppBinaries)
-				{
-					var FoundOnlyModule = DLLBinary.FindOnlyModule(OnlyModules);
-					if (FoundOnlyModule != null)
-					{
-						FilteredBinaries.Add( DLLBinary );
-						AnyBinariesAdded = true;
-
-						if (!String.IsNullOrEmpty(FoundOnlyModule.OnlyModuleSuffix))
-						{
-							var Appendage = "-" + FoundOnlyModule.OnlyModuleSuffix;
-
-							var MatchPos = DLLBinary.Config.OutputFilePath.LastIndexOf(FoundOnlyModule.OnlyModuleName, StringComparison.InvariantCultureIgnoreCase);
-							if (MatchPos < 0)
-							{
-								throw new BuildException("Failed to find module name \"{0}\" specified on the command line inside of the output filename \"{1}\" to add appendage.", FoundOnlyModule.OnlyModuleName, DLLBinary.Config.OutputFilePath);
-							}
-							DLLBinary.Config.OriginalOutputFilePath = DLLBinary.Config.OutputFilePath;
-							DLLBinary.Config.OutputFilePath = DLLBinary.Config.OutputFilePath.Insert(MatchPos + FoundOnlyModule.OnlyModuleName.Length, Appendage);
-						}
-					}
-				}
-
-				// Copy the result into the final list
-				AppBinaries = FilteredBinaries;
-
-				if (!AnyBinariesAdded)
-				{
-					throw new BuildException("One or more of the modules specified using the '-module' argument could not be found.");
-				}
-			}
-
-			// Filter out binaries that were already built and are just used for linking. We will not build these binaries but we need them for link information
-			{
-				var FilteredBinaries = new List<UEBuildBinary>();
-
-				foreach (var DLLBinary in AppBinaries)
-				{
-					if ( DLLBinary.Config.bAllowCompilation )
-					{
-						FilteredBinaries.Add(DLLBinary);
-					}
-				}
-
-				// Copy the result into the final list
-				AppBinaries = FilteredBinaries;
-
-				if (AppBinaries.Count == 0)
-				{
-					throw new BuildException("No modules found to build. All requested binaries were already built.");
-				}
-			}
-
-            // If we are only interested in platform specific binaries, filter everything else out now
-            if (UnrealBuildTool.GetOnlyPlatformSpecificFor() != UnrealTargetPlatform.Unknown)
-            {
-                var FilteredBinaries = new List<UEBuildBinary>();
-
-                var OtherThingsWeNeedToBuild = new List<OnlyModule>();
-
-                foreach (var DLLBinary in AppBinaries)
-                {
-                    if (DLLBinary.Config.bIsCrossTarget)
-                    {
-                        FilteredBinaries.Add(DLLBinary);
-                        bool bIncludeDynamicallyLoaded = false;
-                        var AllReferencedModules = DLLBinary.GetAllDependencyModules(bIncludeDynamicallyLoaded, bForceCircular: true);
-                        foreach (var Other in AllReferencedModules)
-                        {
-                            OtherThingsWeNeedToBuild.Add(new OnlyModule(Other.Name));
-                        }
-                    }
-                }
-                foreach (var DLLBinary in AppBinaries)
-                {
-                    if (!FilteredBinaries.Contains(DLLBinary) && DLLBinary.FindOnlyModule(OtherThingsWeNeedToBuild) != null)
-                    {
-                        if (!File.Exists(DLLBinary.Config.OutputFilePath))
-                        {
-                            throw new BuildException("Module {0} is potentially needed for the {1} platform to work, but it isn't actually compiled. This either needs to be marked as platform specific or made a dynamic dependency of something compiled with the base editor.", DLLBinary.Config.OutputFilePath, UnrealBuildTool.GetOnlyPlatformSpecificFor().ToString());
-                        }
-                    }
-                }
-
-                // Copy the result into the final list
-                AppBinaries = FilteredBinaries;
-            }
-
-			[REDACTED_EMAIL]: Will users be able to rebuild UnrealHeaderTool? NO
-			if (UnrealBuildTool.RunningRocket() && AppName != "UnrealHeaderTool")
-			{
-				var FilteredBinaries = new List<UEBuildBinary>();
-
-				// Have to do absolute here as this could be a project that is under the root
-				string FullUProjectPath = Path.GetFullPath(UnrealBuildTool.GetUProjectPath());
-
-				// We only want to build rocket projects...
-				foreach (var DLLBinary in AppBinaries)
-				{
-					if (Utils.IsFileUnderDirectory( DLLBinary.Config.OutputFilePath, FullUProjectPath ))
-					{
-						FilteredBinaries.Add(DLLBinary);
-					}
-				}
-
-				// Copy the result into the final list
-				AppBinaries = FilteredBinaries;
-
-				if (AppBinaries.Count == 0)
-				{
-					throw new BuildException("Rocket: No modules found to build?");
-				}
-			}
 			// If we're compiling monolithic, make sure the executable knows about all referenced modules
 			if (ShouldCompileMonolithic())
 			{
@@ -1577,5 +1404,253 @@
 		}
 
 		/// <summary>
+		/// Check for EULA violation dependency issues.
+		/// </summary>
+		/// <returns>EULAViolationWarning if error hasn't been thrown earlier. This behavior is controlled using BuildConfiguration.bBreakBuildOnLicenseViolation.</returns>
+		private string CheckForEULAViolation()
+		{
+			string EULAViolationWarning = null;
+
+			if (Rules.Type != TargetRules.TargetType.Editor && IsAppRedistributable(AppName))
+			{
+				var RedistributionErrorMessageBuilder = new StringBuilder();
+				foreach (var Binary in AppBinaries)
+				{
+					var NonRedistModules = Binary.GetAllDependencyModules(true, false).Where((DependencyModule) =>
+							!DependencyModule.IsRedistributable() && DependencyModule.Name != Binary.Target.AppName
+						);
+
+					if (NonRedistModules.Count() != 0)
+					{
+						RedistributionErrorMessageBuilder.Append(string.Format(
+							"{2}Binary {0} depends on: {1}.", Binary.ToString(), string.Join(", ", NonRedistModules),
+							RedistributionErrorMessageBuilder.Length > 0 ? "\n" : ""
+						));
+					}
+				}
+
+				if (RedistributionErrorMessageBuilder.Length > 0)
+				{
+					EULAViolationWarning = string.Format("Non-editor build cannot depend on non-redistributable modules. Details:\n{0}", RedistributionErrorMessageBuilder.ToString());
+
+					if (BuildConfiguration.bBreakBuildOnLicenseViolation)
+					{
+						throw new BuildException(EULAViolationWarning);
+					}
+				}
+			}
+
+			return EULAViolationWarning;
+		}
+
+		/// <summary>
+		/// Tells if app with given app name is redistributable.
+		/// </summary>
+		/// <param name="AppName"></param>
+		/// <returns></returns>
+		private bool IsAppRedistributable(string AppName)
+		{
+			if(AppName == "UE4Client" || AppName == "UE4Server")
+			{
+				return true;
+			}
+			
+			if(AppName == "UE4Editor")
+			{
+				return false;
+			}
+			
+			return FindOrCreateModuleByName(AppName).IsRedistributable();
+		}
+
+		/// <summary>
+		/// Setup target before build. This method finds dependencies, sets up global environment etc.
+		/// </summary>
+		/// <returns>Special Rocket lib files that are build products.</returns>
+		public List<string> PreBuildSetup()
+		{
+			// Set up the global compile and link environment in GlobalCompileEnvironment and GlobalLinkEnvironment.
+			SetupGlobalEnvironment();
+
+			// Setup the target's modules.
+			SetupModules();
+
+			// Setup the target's binaries.
+			SetupBinaries();
+
+			// Setup the target's plugins
+			SetupPlugins();
+
+			var SpecialRocketLibFilesThatAreBuildProducts = new List<string>();
+
+			// Add the enabled plugins to the build
+			foreach (PluginInfo Plugin in EnabledPlugins)
+			{
+				SpecialRocketLibFilesThatAreBuildProducts.AddRange(AddPlugin(Plugin));
+			}
+
+			// Describe what's being built.
+			Log.TraceVerbose("Building {0} - {1} - {2} - {3}", AppName, GameName, Platform, Configuration);
+
+			// Put the non-executable output files (PDB, import library, etc) in the intermediate directory.
+			GlobalLinkEnvironment.Config.IntermediateDirectory = Path.GetFullPath(GlobalCompileEnvironment.Config.OutputDirectory);
+			GlobalLinkEnvironment.Config.OutputDirectory = GlobalLinkEnvironment.Config.IntermediateDirectory;
+
+			// By default, shadow source files for this target in the root OutputDirectory
+			GlobalLinkEnvironment.Config.LocalShadowDirectory = GlobalLinkEnvironment.Config.OutputDirectory;
+
+			// Add all of the extra modules, including game modules, that need to be compiled along
+			// with this app.  These modules aren't necessarily statically linked against, but may
+			// still be required at runtime in order for the application to load and function properly!
+			AddExtraModules();
+
+			// Bind modules for all app binaries.
+			foreach (var Binary in AppBinaries)
+			{
+				Binary.BindModules();
+			}
+
+			// Process all referenced modules and create new binaries for DLL dependencies if needed
+			var NewBinaries = new List<UEBuildBinary>();
+			foreach (var Binary in AppBinaries)
+			{
+				// Add binaries for all of our dependent modules
+				var FoundBinaries = Binary.ProcessUnboundModules();
+				if (FoundBinaries != null)
+				{
+					NewBinaries.AddRange(FoundBinaries);
+				}
+			}
+			AppBinaries.AddRange(NewBinaries);
+
+			// On Mac AppBinaries paths for non-console targets need to be adjusted to be inside the app bundle
+			if (GlobalLinkEnvironment.Config.Target.Platform == CPPTargetPlatform.Mac && !GlobalLinkEnvironment.Config.bIsBuildingConsoleApplication)
+			{
+				MacToolChain.FixBundleBinariesPaths(this, AppBinaries);
+			}
+
+			// If we're building a single module, then find the binary for that module and add it to our target
+			if (OnlyModules.Count > 0)
+			{
+				var FilteredBinaries = new List<UEBuildBinary>();
+
+				var AnyBinariesAdded = false;
+				foreach (var DLLBinary in AppBinaries)
+				{
+					var FoundOnlyModule = DLLBinary.FindOnlyModule(OnlyModules);
+					if (FoundOnlyModule != null)
+					{
+						FilteredBinaries.Add(DLLBinary);
+						AnyBinariesAdded = true;
+
+						if (!String.IsNullOrEmpty(FoundOnlyModule.OnlyModuleSuffix))
+						{
+							var Appendage = "-" + FoundOnlyModule.OnlyModuleSuffix;
+
+							var MatchPos = DLLBinary.Config.OutputFilePath.LastIndexOf(FoundOnlyModule.OnlyModuleName, StringComparison.InvariantCultureIgnoreCase);
+							if (MatchPos < 0)
+							{
+								throw new BuildException("Failed to find module name \"{0}\" specified on the command line inside of the output filename \"{1}\" to add appendage.", FoundOnlyModule.OnlyModuleName, DLLBinary.Config.OutputFilePath);
+							}
+							DLLBinary.Config.OriginalOutputFilePath = DLLBinary.Config.OutputFilePath;
+							DLLBinary.Config.OutputFilePath = DLLBinary.Config.OutputFilePath.Insert(MatchPos + FoundOnlyModule.OnlyModuleName.Length, Appendage);
+						}
+					}
+				}
+
+				// Copy the result into the final list
+				AppBinaries = FilteredBinaries;
+
+				if (!AnyBinariesAdded)
+				{
+					throw new BuildException("One or more of the modules specified using the '-module' argument could not be found.");
+				}
+			}
+
+			// Filter out binaries that were already built and are just used for linking. We will not build these binaries but we need them for link information
+			{
+				var FilteredBinaries = new List<UEBuildBinary>();
+
+				foreach (var DLLBinary in AppBinaries)
+				{
+					if (DLLBinary.Config.bAllowCompilation)
+					{
+						FilteredBinaries.Add(DLLBinary);
+					}
+				}
+
+				// Copy the result into the final list
+				AppBinaries = FilteredBinaries;
+
+				if (AppBinaries.Count == 0)
+				{
+					throw new BuildException("No modules found to build. All requested binaries were already built.");
+				}
+			}
+
+			// If we are only interested in platform specific binaries, filter everything else out now
+			if (UnrealBuildTool.GetOnlyPlatformSpecificFor() != UnrealTargetPlatform.Unknown)
+			{
+				var FilteredBinaries = new List<UEBuildBinary>();
+
+				var OtherThingsWeNeedToBuild = new List<OnlyModule>();
+
+				foreach (var DLLBinary in AppBinaries)
+				{
+					if (DLLBinary.Config.bIsCrossTarget)
+					{
+						FilteredBinaries.Add(DLLBinary);
+						bool bIncludeDynamicallyLoaded = false;
+						var AllReferencedModules = DLLBinary.GetAllDependencyModules(bIncludeDynamicallyLoaded, bForceCircular: true);
+						foreach (var Other in AllReferencedModules)
+						{
+							OtherThingsWeNeedToBuild.Add(new OnlyModule(Other.Name));
+						}
+					}
+				}
+				foreach (var DLLBinary in AppBinaries)
+				{
+					if (!FilteredBinaries.Contains(DLLBinary) && DLLBinary.FindOnlyModule(OtherThingsWeNeedToBuild) != null)
+					{
+						if (!File.Exists(DLLBinary.Config.OutputFilePath))
+						{
+							throw new BuildException("Module {0} is potentially needed for the {1} platform to work, but it isn't actually compiled. This either needs to be marked as platform specific or made a dynamic dependency of something compiled with the base editor.", DLLBinary.Config.OutputFilePath, UnrealBuildTool.GetOnlyPlatformSpecificFor().ToString());
+						}
+					}
+				}
+
+				// Copy the result into the final list
+				AppBinaries = FilteredBinaries;
+			}
+
+			[REDACTED_EMAIL]: Will users be able to rebuild UnrealHeaderTool? NO
+			if (UnrealBuildTool.RunningRocket() && AppName != "UnrealHeaderTool")
+			{
+				var FilteredBinaries = new List<UEBuildBinary>();
+
+				// Have to do absolute here as this could be a project that is under the root
+				string FullUProjectPath = Path.GetFullPath(UnrealBuildTool.GetUProjectPath());
+
+				// We only want to build rocket projects...
+				foreach (var DLLBinary in AppBinaries)
+				{
+					if (Utils.IsFileUnderDirectory(DLLBinary.Config.OutputFilePath, FullUProjectPath))
+					{
+						FilteredBinaries.Add(DLLBinary);
+					}
+				}
+
+				// Copy the result into the final list
+				AppBinaries = FilteredBinaries;
+
+				if (AppBinaries.Count == 0)
+				{
+					throw new BuildException("Rocket: No modules found to build?");
+				}
+			}
+
+			return SpecialRocketLibFilesThatAreBuildProducts;
+		}
+		/// <summary>
 		/// Writes target info.
 		/// </summary>
@@ -1708,6 +1783,7 @@
 					InType: UEBuildModuleType.Game,
 					InModuleDirectory:FakeModuleDirectory,
 					InOutputDirectory: GlobalCompileEnvironment.Config.OutputDirectory,
+					InIsRedistributableOverride: null,
 					InIntelliSenseGatherer: null,
 					InSourceFiles: SourceFiles,
 					InPublicIncludePaths: new List<string>(),
@@ -2658,6 +2734,7 @@
 							InType: ModuleType,
 							InModuleDirectory: ModuleDirectory,
 							InOutputDirectory: ApplicationOutputPath,
+							InIsRedistributableOverride: RulesObject.IsRedistributableOverride,
 							InIntelliSenseGatherer: IntelliSenseGatherer,
 							InSourceFiles: ModuleSourceFiles,
 							InPublicSystemIncludePaths: RulesObject.PublicSystemIncludePaths,
@@ -2695,6 +2772,7 @@
 							InType: ModuleType,
 							InModuleDirectory: ModuleDirectory,
 							InOutputDirectory: ApplicationOutputPath,
+							InIsRedistributableOverride: RulesObject.IsRedistributableOverride,
 							InIntelliSenseGatherer: IntelliSenseGatherer,
 							InSourceFiles: ModuleSourceFiles,
 							InDefinitions: RulesObject.Definitions,
@@ -2733,6 +2811,7 @@
 							InType: ModuleType,
 							InModuleDirectory: ModuleDirectory,
 							InOutputDirectory: ApplicationOutputPath,
+							InIsRedistributableOverride: RulesObject.IsRedistributableOverride,
 							InPublicDefinitions: RulesObject.Definitions,
 							InPublicSystemIncludePaths: RulesObject.PublicSystemIncludePaths,
 							InPublicIncludePaths: RulesObject.PublicIncludePaths,
diff --git a/Engine/Source/Programs/UnrealBuildTool/System/RulesCompiler.cs b/Engine/Source/Programs/UnrealBuildTool/System/RulesCompiler.cs
--- a/Engine/Source/Programs/UnrealBuildTool/System/RulesCompiler.cs
+++ b/Engine/Source/Programs/UnrealBuildTool/System/RulesCompiler.cs
@@ -300,6 +300,9 @@ namespace UnrealBuildTool
 			// Box2D included define (required because pointer types may be in public exported structures)
 			Definitions.Add(string.Format("WITH_BOX2D={0}", bSupported ? 1 : 0));
 		}
+
+		/** Redistribution override flag for this module. */
+		public bool? IsRedistributableOverride { get; set; }
 	}
 
 	/// <summary>
diff --git a/Engine/Source/Programs/UnrealBuildTool/System/UnrealBuildTool.cs b/Engine/Source/Programs/UnrealBuildTool/System/UnrealBuildTool.cs
--- a/Engine/Source/Programs/UnrealBuildTool/System/UnrealBuildTool.cs
+++ b/Engine/Source/Programs/UnrealBuildTool/System/UnrealBuildTool.cs
@@ -1050,6 +1050,8 @@ namespace UnrealBuildTool
 
 			var ToolChain = UEToolChain.GetPlatformToolChain(BuildPlatform.GetCPPTargetPlatform(ResetPlatform));
 
+			string EULAViolationWarning = null;
+
 			try
 			{
 				List<string[]> TargetSettings = ParseCommandLineFlags(Arguments);
@@ -1089,7 +1091,7 @@ namespace UnrealBuildTool
 					}
 
 					var TargetOutputItems = new List<FileItem>();
-					BuildResult = Target.Build(ToolChain, TargetOutputItems);
+					BuildResult = Target.Build(ToolChain, TargetOutputItems, out EULAViolationWarning);
 
 					if(BuildResult != ECompilationResult.Succeeded)
 					{
@@ -1196,6 +1198,11 @@ namespace UnrealBuildTool
 				BuildResult = ECompilationResult.OtherCompilationError;
 			}
 
+			if(EULAViolationWarning != null)
+			{
+				Log.TraceWarning("WARNING: {0}", EULAViolationWarning);
+			}
+
 			// Figure out how long we took to execute.
 			double BuildDuration = (DateTime.UtcNow - StartTime - MutexWaitTime).TotalSeconds;
 			if (ExecutorName == "Local")
