CIS fix.
Not every AppName is represented as module, hence this additional check cannot be made.

diff --git a/Engine/Source/Programs/UnrealBuildTool/Configuration/UEBuildTarget.cs b/Engine/Source/Programs/UnrealBuildTool/Configuration/UEBuildTarget.cs
--- a/Engine/Source/Programs/UnrealBuildTool/Configuration/UEBuildTarget.cs
+++ b/Engine/Source/Programs/UnrealBuildTool/Configuration/UEBuildTarget.cs
@@ -1411,7 +1411,7 @@ namespace UnrealBuildTool
 		{
 			string EULAViolationWarning = null;
 
-			if (Rules.Type != TargetRules.TargetType.Editor && IsAppRedistributable(AppName))
+			if (Rules.Type != TargetRules.TargetType.Editor)
 			{
 				var RedistributionErrorMessageBuilder = new StringBuilder();
 				foreach (var Binary in AppBinaries)
@@ -1443,26 +1443,6 @@ namespace UnrealBuildTool
 			return EULAViolationWarning;
 		}
 
-		/// <summary>
-		/// Tells if app with given app name is redistributable.
-		/// </summary>
-		/// <param name="AppName"></param>
-		/// <returns></returns>
-		private bool IsAppRedistributable(string AppName)
-		{
-			if(AppName == "UE4Client" || AppName == "UE4Server")
-			{
-				return true;
-			}
-			
-			if(AppName == "UE4Editor")
-			{
-				return false;
-			}
-			
-			return FindOrCreateModuleByName(AppName).IsRedistributable();
-		}
-
 		/// <summary>
 		/// Setup target before build. This method finds dependencies, sets up global environment etc.
 		/// </summary>
