UBT: EULA check should omit Programs as well.

diff --git a/Engine/Source/Programs/UnrealBuildTool/Configuration/UEBuildTarget.cs b/Engine/Source/Programs/UnrealBuildTool/Configuration/UEBuildTarget.cs
--- a/Engine/Source/Programs/UnrealBuildTool/Configuration/UEBuildTarget.cs
+++ b/Engine/Source/Programs/UnrealBuildTool/Configuration/UEBuildTarget.cs
@@ -1411,7 +1411,7 @@ namespace UnrealBuildTool
 		{
 			string EULAViolationWarning = null;
 
-			if (Rules.Type != TargetRules.TargetType.Editor)
+			if (Rules.Type != TargetRules.TargetType.Editor && Rules.Type != TargetRules.TargetType.Program)
 			{
 				var RedistributionErrorMessageBuilder = new StringBuilder();
 				foreach (var Binary in AppBinaries)
