#HTML5
     #TTP 338961
             provide correct version checks for emscripten sdk and generate appropriate errors for the user.

diff --git a/Engine/Source/Programs/UnrealBuildTool/HTML5/UEBuildHTML5.cs b/Engine/Source/Programs/UnrealBuildTool/HTML5/UEBuildHTML5.cs
--- a/Engine/Source/Programs/UnrealBuildTool/HTML5/UEBuildHTML5.cs
+++ b/Engine/Source/Programs/UnrealBuildTool/HTML5/UEBuildHTML5.cs
@@ -84,16 +84,25 @@ namespace UnrealBuildTool
 			string BaseSDKPath = Environment.GetEnvironmentVariable("EMSCRIPTEN");
 			if (!string.IsNullOrEmpty(BaseSDKPath))
 			{
-				// Check for the *actual* sdk version
-				// Note: This depends on keeping the installed folder name the version #!!!
-				if (BaseSDKPath.Contains(ExpectedSDKVersion))
-				{
 
-					return SDKStatus.Valid;
-				}
-				else
+				try 
+				{ 
+					// Check for the *actual* sdk version
+					string VersionInfo = File.ReadAllText( Path.Combine(BaseSDKPath,"emscripten-version.txt")); 
+					if (VersionInfo.Contains(ExpectedSDKVersion))
+					{
+						return SDKStatus.Valid;
+					}
+					else
+					{
+						Console.WriteLine( "EMSCRIPTEN sdk found but of unexpected version, Please install version " + ExpectedSDKVersion); 
+						return SDKStatus.Invalid;
+					}
+				} 
+				catch (System.Exception) 
 				{
-					return SDKStatus.Invalid;
+					Console.WriteLine( "Please check you emscripten installation. Incorrectly set EMSCRIPTEN sdk at " + BaseSDKPath );
+					return SDKStatus.Invalid; 
 				}
 			}
 			return SDKStatus.Invalid;
