Temporarily switching off EULA violation warnings

diff --git a/Engine/Source/Programs/UnrealBuildTool/Configuration/BuildConfiguration.cs b/Engine/Source/Programs/UnrealBuildTool/Configuration/BuildConfiguration.cs
--- a/Engine/Source/Programs/UnrealBuildTool/Configuration/BuildConfiguration.cs
+++ b/Engine/Source/Programs/UnrealBuildTool/Configuration/BuildConfiguration.cs
@@ -336,6 +336,12 @@ namespace UnrealBuildTool
 		[XmlConfig]
 		public static bool bDebugBuildsActuallyUseDebugCRT;
 
+		/// <summary>
+		/// Tells the UBT to check if module currently being built is violating EULA.
+		/// </summary>
+		[XmlConfig]
+		public static bool bCheckLicenseViolations;
+
 		/// <summary>
 		/// Tells the UBT to break build if module currently being built is violating EULA.
 		/// </summary>
@@ -445,6 +451,8 @@ namespace UnrealBuildTool
 			// set up some paths
 			BaseIntermediateFolder = "Intermediate/Build/";
 
+			// By default check and stop the build on EULA violation
+			bCheckLicenseViolations = true;
 			bBreakBuildOnLicenseViolation = true;
 		}
 
diff --git a/Engine/Source/Programs/UnrealBuildTool/Configuration/UEBuildTarget.cs b/Engine/Source/Programs/UnrealBuildTool/Configuration/UEBuildTarget.cs
--- a/Engine/Source/Programs/UnrealBuildTool/Configuration/UEBuildTarget.cs
+++ b/Engine/Source/Programs/UnrealBuildTool/Configuration/UEBuildTarget.cs
@@ -1411,7 +1411,8 @@ namespace UnrealBuildTool
 		{
 			string EULAViolationWarning = null;
 
-			if (Rules.Type != TargetRules.TargetType.Editor && Rules.Type != TargetRules.TargetType.Program)
+			if (Rules.Type != TargetRules.TargetType.Editor && Rules.Type != TargetRules.TargetType.Program &&
+				BuildConfiguration.bCheckLicenseViolations)
 			{
 				var RedistributionErrorMessageBuilder = new StringBuilder();
 				foreach (var Binary in AppBinaries)
