Temporarily remove obsolete until I find out why it fails the build

diff --git a/Engine/Source/Programs/UnrealBuildTool/System/RulesCompiler.cs b/Engine/Source/Programs/UnrealBuildTool/System/RulesCompiler.cs
--- a/Engine/Source/Programs/UnrealBuildTool/System/RulesCompiler.cs
+++ b/Engine/Source/Programs/UnrealBuildTool/System/RulesCompiler.cs
@@ -604,8 +604,7 @@ namespace UnrealBuildTool
         /// Return true if this target should always be built with the tools. Usually programs like unrealpak.
         /// <param name="SeparateNode">If this is set to true, the program will get its own node</param>
         /// </summary>
-        /// <returns>true if this target should always be built with the base editor.</returns>
-        [Obsolete]
+        /// <returns>true if this target should always be built with the base editor.</returns>        
         public virtual bool GUBP_AlwaysBuildWithTools(UnrealTargetPlatform InHostPlatform, out bool bInternalToolOnly, out bool SeparateNode)
         {
             bInternalToolOnly = false;
