Changed UBT's WinRTPlatform.HasRequiredManualSDKInternal() to always return SDKStatus.Invalid on Mono

diff --git a/Engine/Source/Programs/UnrealBuildTool/WinRT/UEBuildWinRT.cs b/Engine/Source/Programs/UnrealBuildTool/WinRT/UEBuildWinRT.cs
--- a/Engine/Source/Programs/UnrealBuildTool/WinRT/UEBuildWinRT.cs
+++ b/Engine/Source/Programs/UnrealBuildTool/WinRT/UEBuildWinRT.cs
@@ -45,7 +45,7 @@ namespace UnrealBuildTool
          */
         protected override SDKStatus HasRequiredManualSDKInternal()
         {
-            return IsVisualStudioInstalled() ? SDKStatus.Valid : SDKStatus.Invalid;
+            return IsVisualStudioInstalled() && !Utils.IsRunningOnMono ? SDKStatus.Valid : SDKStatus.Invalid;
         }
 
         /**
