Make GUBP function obsolete and remove last reference to the obsolete call.

diff --git a/Engine/Source/Programs/AutomationTool/ProjectUtils.cs b/Engine/Source/Programs/AutomationTool/ProjectUtils.cs
--- a/Engine/Source/Programs/AutomationTool/ProjectUtils.cs
+++ b/Engine/Source/Programs/AutomationTool/ProjectUtils.cs
@@ -551,7 +551,7 @@ namespace AutomationTool
                     {
                         bool bInternalToolOnly;
                         bool SeparateNode;
-                        bool Tool = ThisTarget.Rules.GUBP_AlwaysBuildWithTools(HostPlatform, out bInternalToolOnly, out SeparateNode);
+                        bool Tool = ThisTarget.Rules.GUBP_AlwaysBuildWithTools(HostPlatform, false, out bInternalToolOnly, out SeparateNode);
 
                         CommandUtils.Log("            TargetName                    : " + ThisTarget.TargetName);
                         CommandUtils.Log("              Build With Editor           : " + (ThisTarget.Rules.GUBP_AlwaysBuildWithBaseEditor() ? "YES" : "NO"));
diff --git a/Engine/Source/Programs/UnrealBuildTool/System/RulesCompiler.cs b/Engine/Source/Programs/UnrealBuildTool/System/RulesCompiler.cs
--- a/Engine/Source/Programs/UnrealBuildTool/System/RulesCompiler.cs
+++ b/Engine/Source/Programs/UnrealBuildTool/System/RulesCompiler.cs
@@ -604,7 +604,8 @@ namespace UnrealBuildTool
         /// Return true if this target should always be built with the tools. Usually programs like unrealpak.
         /// <param name="SeparateNode">If this is set to true, the program will get its own node</param>
         /// </summary>
-        /// <returns>true if this target should always be built with the base editor.</returns>        
+        /// <returns>true if this target should always be built with the base editor.</returns>
+        [Obsolete]
         public virtual bool GUBP_AlwaysBuildWithTools(UnrealTargetPlatform InHostPlatform, out bool bInternalToolOnly, out bool SeparateNode)
         {
             bInternalToolOnly = false;
