Changed MacToolChain.CompileCPPFiles to use StringBuilder and not use Path.GetFullPath and ConvertPath when it's not needed. This saved about 9 seconds of the time it takes UBT to determine if targer is up to date. Other platforms may want similar changes.

diff --git a/Engine/Source/Programs/UnrealBuildTool/Mac/MacToolChain.cs b/Engine/Source/Programs/UnrealBuildTool/Mac/MacToolChain.cs
--- a/Engine/Source/Programs/UnrealBuildTool/Mac/MacToolChain.cs
+++ b/Engine/Source/Programs/UnrealBuildTool/Mac/MacToolChain.cs
@@ -5,6 +5,7 @@ using System.Collections.Generic;
 using System.IO;
 using System.Diagnostics;
 using System.Security.AccessControl;
+using System.Text;
 
 namespace UnrealBuildTool
 {
@@ -221,12 +222,14 @@ namespace UnrealBuildTool
 
 		public override CPPOutput CompileCPPFiles(CPPEnvironment CompileEnvironment, List<FileItem> SourceFiles, string ModuleName)
 		{
-			string Arguments = GetCompileArguments_Global(CompileEnvironment);
-			string PCHArguments = "";
+			var Arguments = new StringBuilder();
+			var PCHArguments = new StringBuilder();
+
+			Arguments.Append(GetCompileArguments_Global(CompileEnvironment));
 
 			if (CompileEnvironment.Config.PrecompiledHeaderAction != PrecompiledHeaderAction.Create)
 			{
-				Arguments += " -Werror";
+				Arguments.Append(" -Werror");
 			}
 
 			if (CompileEnvironment.Config.PrecompiledHeaderAction == PrecompiledHeaderAction.Include)
@@ -234,7 +237,9 @@ namespace UnrealBuildTool
 				// Add the precompiled header file's path to the include path so GCC can find it.
 				// This needs to be before the other include paths to ensure GCC uses it instead of the source header file.
 				var PrecompiledFileExtension = UEBuildPlatform.BuildPlatformDictionary[UnrealTargetPlatform.Mac].GetBinaryExtension(UEBuildBinaryType.PrecompiledHeader);
-				PCHArguments += string.Format(" -include \"{0}\"", CompileEnvironment.PrecompiledHeaderFile.AbsolutePath.Replace(PrecompiledFileExtension, ""));
+				PCHArguments.Append(" -include \"");
+				PCHArguments.Append(CompileEnvironment.PrecompiledHeaderFile.AbsolutePath.Replace(PrecompiledFileExtension, ""));
+				PCHArguments.Append("\"");
 			}
 
 			// Add include paths to the argument list.
@@ -242,10 +247,12 @@ namespace UnrealBuildTool
 			AllIncludes.AddRange(CompileEnvironment.Config.SystemIncludePaths);
 			foreach (string IncludePath in AllIncludes)
 			{
-				Arguments += string.Format(" -I\"{0}\"", ConvertPath(Path.GetFullPath(IncludePath)));
+				Arguments.Append(" -I\"");
 
 				if (ExternalExecution.GetRuntimePlatform() != UnrealTargetPlatform.Mac)
 				{
+					Arguments.Append(ConvertPath(Path.GetFullPath(IncludePath)));
+
 					// sync any third party headers we may need
 					if (IncludePath.Contains("ThirdParty"))
 					{
@@ -264,11 +271,19 @@ namespace UnrealBuildTool
 						}
 					}
 				}
+				else
+				{
+					Arguments.Append(IncludePath);
+				}
+
+				Arguments.Append("\"");
 			}
 
 			foreach (string Definition in CompileEnvironment.Config.Definitions)
 			{
-				Arguments += string.Format(" -D\"{0}\"", Definition);
+				Arguments.Append(" -D\"");
+				Arguments.Append(Definition);
+				Arguments.Append("\"");
 			}
 
 			CPPOutput Result = new CPPOutput();
@@ -310,7 +325,7 @@ namespace UnrealBuildTool
 					FileArguments += GetCompileArguments_CPP();
 
 					// only use PCH for .cpp files
-					FileArguments += PCHArguments;
+					FileArguments += PCHArguments.ToString();
 				}
 
 				// Add the C++ source file and its included files to the prerequisite item list.
