Non monolithic projects will load AutoSDK environment properly.

diff --git a/Engine/Source/Developer/TargetPlatform/Private/TargetPlatformManagerModule.cpp b/Engine/Source/Developer/TargetPlatform/Private/TargetPlatformManagerModule.cpp
--- a/Engine/Source/Developer/TargetPlatform/Private/TargetPlatformManagerModule.cpp
+++ b/Engine/Source/Developer/TargetPlatform/Private/TargetPlatformManagerModule.cpp
@@ -6,9 +6,14 @@
 
 #include "TargetPlatformPrivatePCH.h"
 
-
 DEFINE_LOG_CATEGORY_STATIC(LogTargetPlatformManager, Log, All);
 
+// autosdks only function properly on windows right now.
+#if !IS_MONOLITHIC && (PLATFORM_WINDOWS)
+#define AUTOSDKS_ENABLED 1
+#else
+#define AUTOSDKS_ENABLED 0
+#endif
 
 /**
  * Module for the target platform manager
@@ -24,6 +29,18 @@ public:
 	FTargetPlatformManagerModule()
 		: bRestrictFormatsToRuntimeOnly(false), bForceCacheUpdate(true)
 	{
+
+#if AUTOSDKS_ENABLED		
+		// amortize UBT cost by calling it once for all platforms, rather than once per platform.
+		FString UBTParams(TEXT("-autosdkonly"));
+		int32 UBTReturnCode = -1;
+		FString UBTOutput;
+		if (!FUBTInvoker::InvokeUnrealBuildToolSync(UBTParams, *GLog, true, UBTReturnCode, UBTOutput))
+		{
+			UE_LOG(LogTargetPlatformManager, Fatal, TEXT("Failed to run UBT to check SDK status!"));
+		}		
+#endif
+
 		GetTargetPlatforms();
 		GetActiveTargetPlatforms();
 		GetAudioFormats();
@@ -33,6 +50,8 @@ public:
 		bForceCacheUpdate = false;
 
 		FModuleManager::Get().OnModulesChanged().AddRaw(this, &FTargetPlatformManagerModule::ModulesChangesCallback);
+
+
 	}
 
 	virtual ~FTargetPlatformManagerModule()
@@ -148,8 +167,8 @@ public:
 					for (int32 Index = 0; Index < TargetPlatforms.Num(); Index++)
 					{
 						if (PlatformNames.Contains(TargetPlatforms[Index]->PlatformName()))
-						{
-							Results.Add(TargetPlatforms[Index]);
+						{							
+							Results.Add(TargetPlatforms[Index]);						
 						}
 					}
 
@@ -169,8 +188,8 @@ public:
 				for (int32 Index = 0; Index < TargetPlatforms.Num(); Index++)
 				{
 					if (TargetPlatforms[Index]->IsRunningPlatform())
-					{
-						Results.Add(TargetPlatforms[Index]);
+					{						
+						Results.Add(TargetPlatforms[Index]);					
 					}
 				}
 			}
@@ -508,10 +527,250 @@ protected:
 				ITargetPlatform* Platform = Module->GetTargetPlatform();
 				if (Platform != NULL)
 				{
-					Platforms.Add(Platform);
+					// would like to move this check to GetActiveTargetPlatforms, but too many things cache this result
+					// this setup will become faster after TTP 341897 is complete.
+					if (SetupAndValidateAutoSDK(*Platform))
+					{
+						Platforms.Add(Platform);
+					}
+				}
+			}
+		}
+	}
+
+	bool SetupAndValidateAutoSDK(ITargetPlatform& Platform)
+	{
+		bool bValidSDK = false;
+		// Invoke UBT to perform SDK switching, or detect that a proper manual SDK is already setup.
+		FString PlatformName = Platform.IniPlatformName();
+		FName PlatformFName(*PlatformName);
+
+		// cache result of the last setup attempt to avoid calling UBT all the time.
+		bool* bPreviousSetupSuccessful = PlatformsSetup.Find(PlatformFName);
+		if (bPreviousSetupSuccessful)
+		{
+			bValidSDK = *bPreviousSetupSuccessful;
+		}
+		else
+		{
+			bValidSDK = SetupEnvironmentFromAutoSDK(Platform);
+			PlatformsSetup.Add(PlatformFName, bValidSDK);
+		}
+		return bValidSDK;
+	}
+	
+	bool SetupEnvironmentFromAutoSDK(ITargetPlatform& Platform)
+	{						
+#if AUTOSDKS_ENABLED
+		static const FString SDKRootEnvFar(TEXT("UE_SDKS_ROOT"));		
+		const int32 MaxPathSize = 32768;
+		TCHAR SDKPath[MaxPathSize] = { 0 };
+		FPlatformMisc::GetEnvironmentVariable(*SDKRootEnvFar, SDKPath, MaxPathSize);
+
+		// AutoSDKs only enabled if UE_SDKS_ROOT is set.
+		if (SDKPath[0] == 0)
+		{
+			return true;
+		}
+
+		// if this platform hasn't been setup with AutoSDK handling, then we just assume it's setup
+		// properly.
+		if (!Platform.SupportsAutoSDK())
+		{
+			return true;
+		}
+
+		// Invoke UBT to perform SDK switching, or detect that a proper manual SDK is already setup.
+		FString PlatformName = Platform.IniPlatformName();
+		FName PlatformFName(*PlatformName);		
+
+#if PLATFORM_WINDOWS
+		FString HostPlatform(TEXT("HostWin64"));
+#else
+#error Fill in your host platform directory
+#endif
+
+		if (PlatformName.Find(TEXT("Linux"), ESearchCase::IgnoreCase) != -1)
+		{
+			PlatformName = TEXT("Linux_x64");
+		}
+
+		FString TargetSDKRoot = FPaths::Combine(SDKPath, *HostPlatform, *PlatformName);
+		static const FString SDKInstallManifestFileName(TEXT("CurrentlyInstalled.txt"));
+		FString SDKInstallManifestFilePath = FPaths::Combine(*TargetSDKRoot, *SDKInstallManifestFileName);
+
+		// If we are using a manual install, then it is valid for there to be no OutputEnvVars file.
+		TAutoPtr<FArchive> InstallManifestFile(IFileManager::Get().CreateFileReader(*SDKInstallManifestFilePath));
+		if (InstallManifestFile.IsValid())
+		{
+			TArray<FString> FileLines;
+			int64 FileSize = InstallManifestFile->TotalSize();
+			int64 MemSize = FileSize + 1;
+			void* FileMem = FMemory::Malloc(MemSize);
+			FMemory::Memset(FileMem, 0, MemSize);
+
+			InstallManifestFile->Serialize(FileMem, FileSize);
+
+			FString FileAsString(ANSI_TO_TCHAR(FileMem));
+			FileAsString.ParseIntoArrayLines(&FileLines);
+
+			FMemory::Free(FileMem);
+			InstallManifestFile->Close();
+
+			
+			if (FileLines.Num() != 2)
+			{
+				UE_LOG(LogTargetPlatformManager, Error, TEXT("Malformed install manifest file for Platform %s"), *PlatformName);				
+				return false;
+			}
+
+			static const FString ManualSDKString(TEXT("ManualSDK"));
+			if (FileLines[1].Compare(ManualSDKString, ESearchCase::IgnoreCase) == 0)
+			{
+				UE_LOG(LogTargetPlatformManager, Verbose, TEXT("Platform %s has manual sdk install"), *PlatformName);				
+				return true;
+			}
+		}
+		else
+		{	
+			UE_LOG(LogTargetPlatformManager, Error, TEXT("install manifest file for Platform %s not found."), *PlatformName);			
+			return false;			
+		}		
+
+		static const FString SDKEnvironmentVarsFile(TEXT("OutputEnvVars.txt"));
+		FString EnvVarFileName = FPaths::Combine(*TargetSDKRoot, *SDKEnvironmentVarsFile);		
+
+		// If we are using a manual install, then it is valid for there to be no OutputEnvVars file.
+		TAutoPtr<FArchive> EnvVarFile(IFileManager::Get().CreateFileReader(*EnvVarFileName));
+		if (EnvVarFile.IsValid())
+		{
+			TArray<FString> FileLines;
+			{
+				int64 FileSize = EnvVarFile->TotalSize();
+				int64 MemSize = FileSize + 1;
+				void* FileMem = FMemory::Malloc(MemSize);
+				FMemory::Memset(FileMem, 0, MemSize);
+
+				EnvVarFile->Serialize(FileMem, FileSize);
+
+				FString FileAsString(ANSI_TO_TCHAR(FileMem));
+				FileAsString.ParseIntoArrayLines(&FileLines);
+
+				FMemory::Free(FileMem);
+				EnvVarFile->Close();				
+			}			
+
+			TArray<FString> PathAdds;
+			TArray<FString> PathRemoves;
+			TArray<FString> EnvVarNames;
+			TArray<FString> EnvVarValues;
+
+			const FString VariableSplit(TEXT("="));
+			for (int32 i = 0; i < FileLines.Num(); ++i)
+			{				
+				const FString& VariableString = FileLines[i];
+
+				FString Left;
+				FString Right;
+				VariableString.Split(VariableSplit, &Left, &Right);
+				
+				if (Left.Compare(TEXT("strippath"), ESearchCase::IgnoreCase) == 0)
+				{
+					PathRemoves.Add(Right);
+				}
+				else if (Left.Compare(TEXT("addpath"), ESearchCase::IgnoreCase) == 0)
+				{
+					PathAdds.Add(Right);
+				}
+				else
+				{
+					// convenience for setup.bat writers.  Trim any accidental whitespace from var names/values.
+					EnvVarNames.Add(Left.Trim().TrimTrailing());
+					EnvVarValues.Add(Right.Trim().TrimTrailing());
+				}
+			}
+
+			// don't actually set anything until we successfully validate and read all values in.
+			// we don't want to set a few vars, return a failure, and then have a platform try to
+			// build against a manually installed SDK with half-set env vars.
+			for (int i = 0; i < EnvVarNames.Num(); ++i)
+			{
+				const FString& EnvVarName = EnvVarNames[i];
+				const FString& EnvVarValue = EnvVarValues[i];
+				UE_LOG(LogTargetPlatformManager, Verbose, TEXT("Setting variable '%s' to '%s'."), *EnvVarName, *EnvVarValue);
+				FPlatformMisc::SetEnvironmentVar(*EnvVarName, *EnvVarValue);
+			}
+
+			const int32 MaxPathVarLen = 32768;
+			TCHAR OrigPathVarMem[MaxPathVarLen];
+			FPlatformMisc::GetEnvironmentVariable(TEXT("PATH"), OrigPathVarMem, MaxPathVarLen);
+
+			// actually perform the PATH stripping / adding.
+			FString OrigPathVar(OrigPathVarMem);
+
+			const TCHAR* PathDelimiter = FPlatformMisc::GetPathVarDelimiter();
+			TArray<FString> PathVars;
+			OrigPathVar.ParseIntoArray(&PathVars, PathDelimiter, true);
+
+			TArray<FString> ModifiedPathVars;
+			ModifiedPathVars = PathVars;
+
+			// perform removes first, in case they overlap with any adds.
+			for (int32 PathRemoveIndex = 0; PathRemoveIndex < PathRemoves.Num(); ++PathRemoveIndex)
+			{
+				const FString& PathRemove = PathRemoves[PathRemoveIndex];
+				for (int32 PathVarIndex = 0; PathVarIndex < PathVars.Num(); ++PathVarIndex)
+				{
+					const FString& PathVar = PathVars[PathVarIndex];
+					if (PathVar.Find(PathRemove, ESearchCase::IgnoreCase) >= 0)
+					{
+						UE_LOG(LogTargetPlatformManager, Verbose, TEXT("Removing Path: '%s'"), *PathVar);
+						ModifiedPathVars.Remove(PathVar);
+					}
+				}
+			}
+
+			// remove all the of ADDs so that if this function is executed multiple times, the paths will be guarateed to be in the same order after each run.
+			// If we did not do this, a 'remove' that matched some, but not all, of our 'adds' would cause the order to change.
+			for (int32 PathAddIndex = 0; PathAddIndex < PathAdds.Num(); ++PathAddIndex)			
+			{
+				const FString& PathAdd = PathAdds[PathAddIndex];				
+				for (int32 PathVarIndex = 0; PathVarIndex < PathVars.Num(); ++PathVarIndex)
+				{
+					const FString& PathVar = PathVars[PathVarIndex];
+					if (PathVar.Find(PathAdd, ESearchCase::IgnoreCase) >= 0)
+					{
+						UE_LOG(LogTargetPlatformManager, Verbose, TEXT("Removing Path: '%s'"), *PathVar);
+						ModifiedPathVars.Remove(PathVar);
+					}
 				}
 			}
+
+			// perform adds, but don't add duplicates
+			for (int32 PathAddIndex = 0; PathAddIndex < PathAdds.Num(); ++PathAddIndex)
+			{
+				const FString& PathAdd = PathAdds[PathAddIndex];
+				if (!ModifiedPathVars.Contains(PathAdd))
+				{
+					UE_LOG(LogTargetPlatformManager, Verbose, TEXT("Adding Path: '%s'"), *PathAdd);					
+					ModifiedPathVars.Add(PathAdd);
+				}
+			}
+
+			FString ModifiedPath = FString::Join(ModifiedPathVars, PathDelimiter);
+			FPlatformMisc::SetEnvironmentVar(TEXT("PATH"), *ModifiedPath);			
+		}
+		else
+		{
+			UE_LOG(LogTargetPlatformManager, Error, TEXT("OutputEnvVars.txt not found for platform: '%s'"), *PlatformName);			
+			return false;
 		}
+
+		UE_LOG(LogTargetPlatformManager, Verbose, TEXT("Platform %s has auto sdk install"), *PlatformName);		
+		return true;
+#else
+		return true;
+#endif
 	}
 
 
@@ -536,6 +795,11 @@ private:
 
 	// Holds the list of discovered platforms.
 	TArray<ITargetPlatform*> Platforms;
+
+#if AUTOSDKS_ENABLED
+	// holds the list of Platforms that have attempted setup.
+	TMap<FName, bool> PlatformsSetup;
+#endif
 };
 
 
diff --git a/Engine/Source/Developer/TargetPlatform/Public/Common/TargetPlatformBase.h b/Engine/Source/Developer/TargetPlatform/Public/Common/TargetPlatformBase.h
--- a/Engine/Source/Developer/TargetPlatform/Public/Common/TargetPlatformBase.h
+++ b/Engine/Source/Developer/TargetPlatform/Public/Common/TargetPlatformBase.h
@@ -259,6 +259,11 @@ public:
 		return TPlatformProperties::SupportsBuildTarget(BuildTarget);
 	}
 
+	virtual bool SupportsAutoSDK() const override
+	{
+		return TPlatformProperties::SupportsAutoSDK();
+	}
+
 	virtual bool SupportsFeature( ETargetPlatformFeatures::Type Feature ) const override
 	{
 		switch (Feature)
diff --git a/Engine/Source/Developer/TargetPlatform/Public/Interfaces/ITargetPlatform.h b/Engine/Source/Developer/TargetPlatform/Public/Interfaces/ITargetPlatform.h
--- a/Engine/Source/Developer/TargetPlatform/Public/Interfaces/ITargetPlatform.h
+++ b/Engine/Source/Developer/TargetPlatform/Public/Interfaces/ITargetPlatform.h
@@ -272,6 +272,11 @@ public:
 	 */
 	virtual bool SupportsBuildTarget( EBuildTargets::Type BuildTarget ) const = 0;
 
+	/**
+	 * Returns true if the platform supports the AutoSDK system
+	 */
+	virtual bool SupportsAutoSDK() const = 0;
+
 	/**
 	 * Checks whether the target platform supports the specified feature.
 	 *
diff --git a/Engine/Source/Programs/UnrealBuildTool/Configuration/UEBuildPlatform.cs b/Engine/Source/Programs/UnrealBuildTool/Configuration/UEBuildPlatform.cs
--- a/Engine/Source/Programs/UnrealBuildTool/Configuration/UEBuildPlatform.cs
+++ b/Engine/Source/Programs/UnrealBuildTool/Configuration/UEBuildPlatform.cs
@@ -30,6 +30,11 @@ namespace UnrealBuildTool
 		 * Register the platform with the UEBuildPlatform class
 		 */
 		void RegisterBuildPlatform();
+
+        /**
+         * Attempt to set up AutoSDK for this paltform
+         */
+        void ManageAndValidateSDK();
 		
 		/**
 		 * Retrieve the CPPTargetPlatform for the given UnrealTargetPlatform
@@ -911,10 +916,16 @@ namespace UnrealBuildTool
                     using (StreamReader Reader = new StreamReader(VersionFilename))
                     {
                         string Version = Reader.ReadLine();
-                        if (Version != null)
+                        string Type = Reader.ReadLine();
+
+                        // don't allow ManualSDK installs to count as an AutoSDK install version.
+                        if (Type != null && Type == "AutoSDK")
                         {
-                            OutInstalledSDKVersionString = Version;
-                            return true;
+                            if (Version != null)
+                            {
+                                OutInstalledSDKVersionString = Version;
+                                return true;
+                            }
                         }
                     }
                 }
@@ -963,8 +974,9 @@ namespace UnrealBuildTool
          * 
          * @return true if was able to set it
          */
-        protected bool SetCurrentlyInstalledAutoSDKString(string PlatformSDKRoot, string InstalledSDKVersionString)
+        protected bool SetCurrentlyInstalledAutoSDKString(String InstalledSDKVersionString)
         {
+            String PlatformSDKRoot = GetPathToPlatformAutoSDKs();
             if (Directory.Exists(PlatformSDKRoot))
             {
                 string VersionFilename = Path.Combine(PlatformSDKRoot, CurrentlyInstalledSDKStringManifest);
@@ -976,6 +988,7 @@ namespace UnrealBuildTool
                 using (StreamWriter Writer = File.CreateText(VersionFilename))
                 {
                     Writer.WriteLine(InstalledSDKVersionString);
+                    Writer.WriteLine("AutoSDK");
                     return true;
                 }
             }
@@ -983,8 +996,38 @@ namespace UnrealBuildTool
             return false;
         }
 
-        protected bool SetLastRunAutoSDKScriptVersion(string PlatformSDKRoot, string LastRunScriptVersion)
+        protected void SetupManualSDK()
         {
+            if (PlatformSupportsAutoSDKs() && HasAutoSDKSystemEnabled())
+            {
+                String InstalledSDKVersionString = GetRequiredSDKString();
+                String PlatformSDKRoot = GetPathToPlatformAutoSDKs();
+                if (Directory.Exists(PlatformSDKRoot))
+                {
+                    string VersionFilename = Path.Combine(PlatformSDKRoot, CurrentlyInstalledSDKStringManifest);
+                    if (File.Exists(VersionFilename))
+                    {
+                        File.Delete(VersionFilename);
+                    }
+
+                    string EnvVarFile = Path.Combine(PlatformSDKRoot, SDKEnvironmentVarsFile);
+                    if (File.Exists(EnvVarFile))
+                    {
+                        File.Delete(EnvVarFile);
+                    }
+
+                    using (StreamWriter Writer = File.CreateText(VersionFilename))
+                    {
+                        Writer.WriteLine(InstalledSDKVersionString);
+                        Writer.WriteLine("ManualSDK");                        
+                    }
+                }                
+            }
+        }
+
+        protected bool SetLastRunAutoSDKScriptVersion(string LastRunScriptVersion)
+        {
+            String PlatformSDKRoot = GetPathToPlatformAutoSDKs();
             if (Directory.Exists(PlatformSDKRoot))
             {
                 string VersionFilename = Path.Combine(PlatformSDKRoot, LastRunScriptVersionManifest);
@@ -999,7 +1042,6 @@ namespace UnrealBuildTool
                     return true;
                 }
             }
-
             return false;
         }
 
@@ -1049,7 +1091,7 @@ namespace UnrealBuildTool
                     Console.WriteLine("Running {0} hook {1}", Hook, HookExe);
 
                     // run it
-                    Process HookProcess = new Process();
+                    Process HookProcess = new Process();                                        
                     HookProcess.StartInfo.WorkingDirectory = SDKDirectory;
                     HookProcess.StartInfo.FileName = HookExe;
                     HookProcess.StartInfo.Arguments = "";
@@ -1062,7 +1104,7 @@ namespace UnrealBuildTool
 
                     if (HookProcess.ExitCode != 0)
                     {
-                        Console.WriteLine("Hook exited uncleanly (returned {0}), considering it failed.", HookProcess.ExitCode);
+                        Console.WriteLine("Hook exited uncleanly (returned {0}), considering it failed.", HookProcess.ExitCode);                        
                         return false;
                     }
 
@@ -1075,6 +1117,8 @@ namespace UnrealBuildTool
 
         /**
          * Loads environment variables from SDK
+         * If any commands are added or removed the handling needs to be duplicated in
+         * TargetPlatformManagerModule.cpp
          * 
          * @param PlatformSDKRoot absolute path to platform SDK         
          * 
@@ -1206,10 +1250,29 @@ namespace UnrealBuildTool
             return false;
         }
 
-        protected void InvalidateCurrentlyInstalledAutoSDK(String PlatformSDKRoot)
+        protected void InvalidateCurrentlyInstalledAutoSDK()
         {
-            SetCurrentlyInstalledAutoSDKString(PlatformSDKRoot, "");
-            SetLastRunAutoSDKScriptVersion(PlatformSDKRoot, "");
+            String PlatformSDKRoot = GetPathToPlatformAutoSDKs();
+            if (Directory.Exists(PlatformSDKRoot))
+            {
+                string SDKFilename = Path.Combine(PlatformSDKRoot, CurrentlyInstalledSDKStringManifest);
+                if (File.Exists(SDKFilename))
+                {
+                    File.Delete(SDKFilename);
+                }
+
+                string VersionFilename = Path.Combine(PlatformSDKRoot, LastRunScriptVersionManifest);
+                if (File.Exists(VersionFilename))
+                {
+                    File.Delete(VersionFilename);
+                }
+
+                string EnvVarFile = Path.Combine(PlatformSDKRoot, SDKEnvironmentVarsFile);
+                if (File.Exists(EnvVarFile))
+                {
+                    File.Delete(EnvVarFile);
+                }
+            }            
         }
 
         /** 
@@ -1231,8 +1294,12 @@ namespace UnrealBuildTool
                         bScriptVersionMatches = true;
                     }
 
+                    // check to make sure OutputEnvVars doesn't need regenerating
+                    string EnvVarFile = Path.Combine(AutoSDKRoot, SDKEnvironmentVarsFile);
+                    bool bEnvVarFileExists = File.Exists(EnvVarFile);
+                    
                     string CurrentSDKString;
-                    if (GetCurrentlyInstalledSDKString(AutoSDKRoot, out CurrentSDKString) && CurrentSDKString == GetRequiredSDKString() && bScriptVersionMatches)
+                    if (bEnvVarFileExists && GetCurrentlyInstalledSDKString(AutoSDKRoot, out CurrentSDKString) && CurrentSDKString == GetRequiredSDKString() && bScriptVersionMatches)
                     {
                         return SDKStatus.Valid;
                     }
@@ -1316,11 +1383,11 @@ namespace UnrealBuildTool
                     if (!RunAutoSDKHooks(AutoSDKRoot, CurrentSDKString, SDKHookType.Uninstall))
                     {
                         Console.WriteLine("Failed to uninstall currently installed SDK {0}", CurrentSDKString);
-                        InvalidateCurrentlyInstalledAutoSDK(AutoSDKRoot);
+                        InvalidateCurrentlyInstalledAutoSDK();
                         return;
                     }
                     // delete Manifest file to avoid multiple uninstalls
-                    InvalidateCurrentlyInstalledAutoSDK(AutoSDKRoot);
+                    InvalidateCurrentlyInstalledAutoSDK();
 
                     if (!RunAutoSDKHooks(AutoSDKRoot, GetRequiredSDKString(), SDKHookType.Install, false))
                     {
@@ -1328,8 +1395,8 @@ namespace UnrealBuildTool
                         RunAutoSDKHooks(AutoSDKRoot, GetRequiredSDKString(), SDKHookType.Uninstall, false);
                         return;
                     }
-                    SetCurrentlyInstalledAutoSDKString(AutoSDKRoot, GetRequiredSDKString());
-                    SetLastRunAutoSDKScriptVersion(AutoSDKRoot, GetRequiredScriptVersionString());
+                    SetCurrentlyInstalledAutoSDKString(GetRequiredSDKString());
+                    SetLastRunAutoSDKScriptVersion(GetRequiredScriptVersionString());
                 }
 
                 // fixup process environment to match autosdk
@@ -1369,7 +1436,7 @@ namespace UnrealBuildTool
             if (!SetupEnvironmentFromAutoSDK(PlatformSDKRoot))
             {
                 Console.WriteLine("Failed to load environment from required SDK {0}", GetRequiredSDKString());
-                InvalidateCurrentlyInstalledAutoSDK(PlatformSDKRoot);
+                InvalidateCurrentlyInstalledAutoSDK();
                 return SDKStatus.Invalid;
             }            
             return SDKStatus.Valid;
@@ -1395,16 +1462,27 @@ namespace UnrealBuildTool
 
         // Arbitrates between manual SDKs and setting up AutoSDK based on program options and platform preferences.
         public void ManageAndValidateSDK()
-        {
-            bool bHasRequiredManualSDK = HasRequiredManualSDK() == SDKStatus.Valid;
-            if (bAllowAutoSDKSwitching && IsAutoSDKSafe() && (PreferAutoSDK() || !bHasRequiredManualSDK))
+        {            
+            if (bAllowAutoSDKSwitching)
             {
-                SetupAutoSDK();
-            }
+                bool bHasRequiredManualSDK = HasRequiredManualSDK() == SDKStatus.Valid;
+                if (IsAutoSDKSafe() && (PreferAutoSDK() || !bHasRequiredManualSDK))
+                {
+                    SetupAutoSDK();
+                }
+                else if (bHasRequiredManualSDK)
+                {
+                    SetupManualSDK();
+                }
+                else
+                {
+                    InvalidateCurrentlyInstalledAutoSDK();
+                }
 
-            if (BuildConfiguration.bPrintDebugInfo)
-            {
-                PrintSDKInfo();
+                if (BuildConfiguration.bPrintDebugInfo)
+                {
+                    PrintSDKInfo();
+                }
             }
         }
 
diff --git a/Engine/Source/Programs/UnrealBuildTool/System/UnrealBuildTool.cs b/Engine/Source/Programs/UnrealBuildTool/System/UnrealBuildTool.cs
--- a/Engine/Source/Programs/UnrealBuildTool/System/UnrealBuildTool.cs
+++ b/Engine/Source/Programs/UnrealBuildTool/System/UnrealBuildTool.cs
@@ -11,36 +11,36 @@ using System.Linq;
 
 namespace UnrealBuildTool
 {
-	public partial class UnrealBuildTool
-	{
-		/** Time at which control fist passes to UBT. */
-		static public DateTime StartTime;
+    public partial class UnrealBuildTool
+    {
+        /** Time at which control fist passes to UBT. */
+        static public DateTime StartTime;
 
-		/** Time we waited for the mutex. Tracked for proper active execution time collection. */
-		static public TimeSpan MutexWaitTime = TimeSpan.Zero;
+        /** Time we waited for the mutex. Tracked for proper active execution time collection. */
+        static public TimeSpan MutexWaitTime = TimeSpan.Zero;
 
-		static public string BuildGuid = Guid.NewGuid().ToString();
+        static public string BuildGuid = Guid.NewGuid().ToString();
 
-		/** Total time spent building projects. */
-		static public double TotalBuildProjectTime = 0;
+        /** Total time spent building projects. */
+        static public double TotalBuildProjectTime = 0;
 
-		/** Total time spent compiling. */
-		static public double TotalCompileTime = 0;
+        /** Total time spent compiling. */
+        static public double TotalCompileTime = 0;
 
-		/** Total time spent creating app bundles. */
-		static public double TotalCreateAppBundleTime = 0;
+        /** Total time spent creating app bundles. */
+        static public double TotalCreateAppBundleTime = 0;
 
-		/** Total time spent generating debug information. */
-		static public double TotalGenerateDebugInfoTime = 0;
+        /** Total time spent generating debug information. */
+        static public double TotalGenerateDebugInfoTime = 0;
 
-		/** Total time spent linking. */
-		static public double TotalLinkTime = 0;
+        /** Total time spent linking. */
+        static public double TotalLinkTime = 0;
 
-		/** Total time spent on other actions. */
-		static public double TotalOtherActionsTime = 0;
+        /** Total time spent on other actions. */
+        static public double TotalOtherActionsTime = 0;
 
-		/** The command line */
-		static public List<string> CmdLine = new List<string>();
+        /** The command line */
+        static public List<string> CmdLine = new List<string>();
 
         /** The environment at boot time. */
         static public System.Collections.IDictionary InitialEnvironment;
@@ -51,19 +51,19 @@ namespace UnrealBuildTool
         static UnrealTargetPlatform OnlyPlatformSpecificFor = UnrealTargetPlatform.Unknown;
 
         /** Are we running for Rocket */
-		static public bool bRunningRocket = false;
+        static public bool bRunningRocket = false;
 
-		/** Are we building Rocket */
-		static bool bBuildingRocket = false;
+        /** Are we building Rocket */
+        static bool bBuildingRocket = false;
 
-		/** The project file */
-		static string UProjectFile = "";
+        /** The project file */
+        static string UProjectFile = "";
 
-		/** The Unreal project file path.  This should always be valid when compiling game projects.  Engine targets and program targets may not have a UProject set. */
-		static bool bHasUProjectFile = false;
+        /** The Unreal project file path.  This should always be valid when compiling game projects.  Engine targets and program targets may not have a UProject set. */
+        static bool bHasUProjectFile = false;
 
-		/** The Unreal project directory.  This should always be valid when compiling game projects.  Engine targets and program targets may not have a UProject set. */
-		static string UProjectPath = "";
+        /** The Unreal project directory.  This should always be valid when compiling game projects.  Engine targets and program targets may not have a UProject set. */
+        static string UProjectPath = "";
 
         /// <summary>
         /// Returns true if UnrealBuildTool should skip everything other than the host platform as far as finding other modules that need to be compiled.
@@ -83,103 +83,103 @@ namespace UnrealBuildTool
             return OnlyPlatformSpecificFor;
         }
 
-		/// <summary>
-		/// Returns true if UnrealBuildTool is running with the "-rocket" option (Rocket mode)
-		/// </summary>
-		/// <returns>True if running with "-rocket"</returns>
-		static public bool RunningRocket()
-		{
-			return bRunningRocket;
-		}
-
-		/// <summary>
-		/// Checks to see if we're building a "Rocket" target
-		/// </summary>
-		/// <returns>True if building a "Rocket" target, such as RocketEditor (not running with "-rocket" though)</returns>
-		static public bool BuildingRocket()
-		{
-			return bBuildingRocket;
-		}
-
-		/// <summary>
-		/// Returns true if a UProject file is available.  This should always be the case when working with game projects.  For engine or program targets, a UProject will not be available.
-		/// </summary>
-		/// <returns></returns>
-		static public bool HasUProjectFile()
-		{
-			return bHasUProjectFile;
-		}
-
-		/// <summary>
-		/// The Unreal project file path.  This should always be valid when compiling game projects.  Engine targets and program targets may not have a UProject set.
-		/// </summary>
-		/// <returns>The path to the file</returns>
-		static public string GetUProjectFile()
-		{
-			return UProjectFile;
-		}
-
-		/// <summary>
-		/// The Unreal project directory.  This should always be valid when compiling game projects.  Engine targets and program targets may not have a UProject set. */
-		/// </summary>
-		/// <returns>The directory path</returns>
-		static public string GetUProjectPath()
-		{
-			return UProjectPath;
-		}
-
-		/// <summary>
-		/// The Unreal project source directory.  This should always be valid when compiling game projects.  Engine targets and program targets may not have a UProject set. */
-		/// </summary>
-		/// <returns>Relative path to the project's source directory</returns>
-		static public string GetUProjectSourcePath()
-		{
-			string ProjectPath = GetUProjectPath();
-			if (ProjectPath != null)
-			{
-				ProjectPath = Path.Combine(ProjectPath, "Source");
-			}
-			return ProjectPath;
-		}
-
-		static public string GetUBTPath()
-		{
-			string UnrealBuildToolPath = Path.Combine(Utils.GetExecutingAssemblyDirectory(), "UnrealBuildTool.exe");
-			return UnrealBuildToolPath;
-		}
-
-		// @todo projectfiles: Move this into the ProjectPlatformGeneration class?
-		/// <summary>
-		/// IsDesktopPlatform
-		/// </summary>
-		/// <param name="InPlatform">The platform of interest</param>
-		/// <returns>True if the given platform is a 'desktop' platform</returns>
-		static public bool IsDesktopPlatform(UnrealTargetPlatform InPlatform)
-		{
-			// Windows and Mac are desktop platforms.
-			return (
-				(InPlatform == UnrealTargetPlatform.Win64) ||
-				(InPlatform == UnrealTargetPlatform.Win32) ||
+        /// <summary>
+        /// Returns true if UnrealBuildTool is running with the "-rocket" option (Rocket mode)
+        /// </summary>
+        /// <returns>True if running with "-rocket"</returns>
+        static public bool RunningRocket()
+        {
+            return bRunningRocket;
+        }
+
+        /// <summary>
+        /// Checks to see if we're building a "Rocket" target
+        /// </summary>
+        /// <returns>True if building a "Rocket" target, such as RocketEditor (not running with "-rocket" though)</returns>
+        static public bool BuildingRocket()
+        {
+            return bBuildingRocket;
+        }
+
+        /// <summary>
+        /// Returns true if a UProject file is available.  This should always be the case when working with game projects.  For engine or program targets, a UProject will not be available.
+        /// </summary>
+        /// <returns></returns>
+        static public bool HasUProjectFile()
+        {
+            return bHasUProjectFile;
+        }
+
+        /// <summary>
+        /// The Unreal project file path.  This should always be valid when compiling game projects.  Engine targets and program targets may not have a UProject set.
+        /// </summary>
+        /// <returns>The path to the file</returns>
+        static public string GetUProjectFile()
+        {
+            return UProjectFile;
+        }
+
+        /// <summary>
+        /// The Unreal project directory.  This should always be valid when compiling game projects.  Engine targets and program targets may not have a UProject set. */
+        /// </summary>
+        /// <returns>The directory path</returns>
+        static public string GetUProjectPath()
+        {
+            return UProjectPath;
+        }
+
+        /// <summary>
+        /// The Unreal project source directory.  This should always be valid when compiling game projects.  Engine targets and program targets may not have a UProject set. */
+        /// </summary>
+        /// <returns>Relative path to the project's source directory</returns>
+        static public string GetUProjectSourcePath()
+        {
+            string ProjectPath = GetUProjectPath();
+            if (ProjectPath != null)
+            {
+                ProjectPath = Path.Combine(ProjectPath, "Source");
+            }
+            return ProjectPath;
+        }
+
+        static public string GetUBTPath()
+        {
+            string UnrealBuildToolPath = Path.Combine(Utils.GetExecutingAssemblyDirectory(), "UnrealBuildTool.exe");
+            return UnrealBuildToolPath;
+        }
+
+        // @todo projectfiles: Move this into the ProjectPlatformGeneration class?
+        /// <summary>
+        /// IsDesktopPlatform
+        /// </summary>
+        /// <param name="InPlatform">The platform of interest</param>
+        /// <returns>True if the given platform is a 'desktop' platform</returns>
+        static public bool IsDesktopPlatform(UnrealTargetPlatform InPlatform)
+        {
+            // Windows and Mac are desktop platforms.
+            return (
+                (InPlatform == UnrealTargetPlatform.Win64) ||
+                (InPlatform == UnrealTargetPlatform.Win32) ||
+                (InPlatform == UnrealTargetPlatform.Linux) ||
+                (InPlatform == UnrealTargetPlatform.Mac)
+                );
+        }
+
+        // @todo projectfiles: Move this into the ProjectPlatformGeneration class?
+        /// <summary>
+        /// IsEditorPlatform
+        /// </summary>
+        /// <param name="InPlatform">The platform of interest</param>
+        /// <returns>True if the given platform is a platform that supports the editor.</returns>
+        static public bool IsEditorPlatform(UnrealTargetPlatform InPlatform)
+        {
+            // Windows 64, Linux and Mac are editor platforms.
+            return (
+                (InPlatform == UnrealTargetPlatform.Win64) ||
                 (InPlatform == UnrealTargetPlatform.Linux) ||
-				(InPlatform == UnrealTargetPlatform.Mac)
-				);
-		}
-
-		// @todo projectfiles: Move this into the ProjectPlatformGeneration class?
-		/// <summary>
-		/// IsEditorPlatform
-		/// </summary>
-		/// <param name="InPlatform">The platform of interest</param>
-		/// <returns>True if the given platform is a platform that supports the editor.</returns>
-		static public bool IsEditorPlatform(UnrealTargetPlatform InPlatform)
-		{
-			// Windows 64, Linux and Mac are editor platforms.
-			return (
-				(InPlatform == UnrealTargetPlatform.Win64) ||
-				(InPlatform == UnrealTargetPlatform.Linux) ||
-				(InPlatform == UnrealTargetPlatform.Mac)
-				);
-		}
+                (InPlatform == UnrealTargetPlatform.Mac)
+                );
+        }
 
         // @todo projectfiles: Move this into the ProjectPlatformGeneration class?
         /// <summary>
@@ -215,183 +215,183 @@ namespace UnrealBuildTool
                 );
         }
 
-		/// <summary>
-		/// Gets all platforms that satisfies the predicate.
-		/// </summary>
-		/// <param name="OutPlatforms">The list of platform to fill in.</param>
-		/// <param name="Predicate">The predicate to filter the platforms.</param>
-		/// <param name="bCheckValidity">If true, ensure it is a valid platform.</param>
-		/// <returns>True if successful and platforms are found, false if not.</returns>
-		static private bool GetPlatforms(ref List<UnrealTargetPlatform> OutPlatforms, Func<UnrealTargetPlatform, bool> Predicate, bool bCheckValidity = true)
-		{
-			OutPlatforms.Clear();
-			foreach (UnrealTargetPlatform Platform in Enum.GetValues(typeof(UnrealTargetPlatform)))
-			{
-				if (Predicate(Platform))
-				{
-					var BuildPlatform = UEBuildPlatform.GetBuildPlatform(Platform, true);
-					if ((bCheckValidity == false) || (BuildPlatform != null))
-					{
-						OutPlatforms.Add(Platform);
-					}
-				}
-			}
-
-			return (OutPlatforms.Count > 0) ? true : false;
-		}
+        /// <summary>
+        /// Gets all platforms that satisfies the predicate.
+        /// </summary>
+        /// <param name="OutPlatforms">The list of platform to fill in.</param>
+        /// <param name="Predicate">The predicate to filter the platforms.</param>
+        /// <param name="bCheckValidity">If true, ensure it is a valid platform.</param>
+        /// <returns>True if successful and platforms are found, false if not.</returns>
+        static private bool GetPlatforms(ref List<UnrealTargetPlatform> OutPlatforms, Func<UnrealTargetPlatform, bool> Predicate, bool bCheckValidity = true)
+        {
+            OutPlatforms.Clear();
+            foreach (UnrealTargetPlatform Platform in Enum.GetValues(typeof(UnrealTargetPlatform)))
+            {
+                if (Predicate(Platform))
+                {
+                    var BuildPlatform = UEBuildPlatform.GetBuildPlatform(Platform, true);
+                    if ((bCheckValidity == false) || (BuildPlatform != null))
+                    {
+                        OutPlatforms.Add(Platform);
+                    }
+                }
+            }
+
+            return (OutPlatforms.Count > 0) ? true : false;
+        }
+
+        /// <summary>
+        /// Gets all platforms.
+        /// </summary>
+        /// <param name="OutPlatforms">The list of platform to fill in.</param>
+        /// <param name="bCheckValidity">If true, ensure it is a valid platform.</param>
+        /// <returns>True if successful and platforms are found, false if not.</returns>
+        static public bool GetAllPlatforms(ref List<UnrealTargetPlatform> OutPlatforms, bool bCheckValidity = true)
+        {
+            return GetPlatforms(ref OutPlatforms, (Platform) => { return true; }, bCheckValidity);
+        }
+
+        /// <summary>
+        /// Gets all desktop platforms.
+        /// </summary>
+        /// <param name="OutPlatforms">The list of platform to fill in.</param>
+        /// <param name="bCheckValidity">If true, ensure it is a valid platform.</param>
+        /// <returns>True if successful and platforms are found, false if not.</returns>
+        static public bool GetAllDesktopPlatforms(ref List<UnrealTargetPlatform> OutPlatforms, bool bCheckValidity = true)
+        {
+            return GetPlatforms(ref OutPlatforms, UnrealBuildTool.IsDesktopPlatform, bCheckValidity);
+        }
 
         /// <summary>
-		/// Gets all platforms.
-		/// </summary>
-		/// <param name="OutPlatforms">The list of platform to fill in.</param>
-		/// <param name="bCheckValidity">If true, ensure it is a valid platform.</param>
-		/// <returns>True if successful and platforms are found, false if not.</returns>
-		static public bool GetAllPlatforms(ref List<UnrealTargetPlatform> OutPlatforms, bool bCheckValidity = true)
-		{
-			return GetPlatforms(ref OutPlatforms, (Platform) => { return true; }, bCheckValidity);
-		}
-
-		/// <summary>
-		/// Gets all desktop platforms.
-		/// </summary>
-		/// <param name="OutPlatforms">The list of platform to fill in.</param>
-		/// <param name="bCheckValidity">If true, ensure it is a valid platform.</param>
-		/// <returns>True if successful and platforms are found, false if not.</returns>
-		static public bool GetAllDesktopPlatforms(ref List<UnrealTargetPlatform> OutPlatforms, bool bCheckValidity = true)
-		{
-			return GetPlatforms(ref OutPlatforms, UnrealBuildTool.IsDesktopPlatform, bCheckValidity);
-		}
-
-		/// <summary>
-		/// Gets all editor platforms.
-		/// </summary>
-		/// <param name="OutPlatforms">The list of platform to fill in.</param>
-		/// <param name="bCheckValidity">If true, ensure it is a valid platform.</param>
-		/// <returns>True if successful and platforms are found, false if not.</returns>
-		static public bool GetAllEditorPlatforms(ref List<UnrealTargetPlatform> OutPlatforms, bool bCheckValidity = true)
-		{
-			return GetPlatforms(ref OutPlatforms, UnrealBuildTool.IsEditorPlatform, bCheckValidity);
-		}
+        /// Gets all editor platforms.
+        /// </summary>
+        /// <param name="OutPlatforms">The list of platform to fill in.</param>
+        /// <param name="bCheckValidity">If true, ensure it is a valid platform.</param>
+        /// <returns>True if successful and platforms are found, false if not.</returns>
+        static public bool GetAllEditorPlatforms(ref List<UnrealTargetPlatform> OutPlatforms, bool bCheckValidity = true)
+        {
+            return GetPlatforms(ref OutPlatforms, UnrealBuildTool.IsEditorPlatform, bCheckValidity);
+        }
 
         /// <summary>
         /// Gets all server platforms.
         /// </summary>
-		/// <param name="OutPlatforms">The list of platform to fill in.</param>
-		/// <param name="bCheckValidity">If true, ensure it is a valid platform.</param>
-		/// <returns>True if successful and platforms are found, false if not.</returns>
+        /// <param name="OutPlatforms">The list of platform to fill in.</param>
+        /// <param name="bCheckValidity">If true, ensure it is a valid platform.</param>
+        /// <returns>True if successful and platforms are found, false if not.</returns>
         static public bool GetAllServerPlatforms(ref List<UnrealTargetPlatform> OutPlatforms, bool bCheckValidity = true)
         {
-			return GetPlatforms(ref OutPlatforms, UnrealBuildTool.IsServerPlatform, bCheckValidity);
+            return GetPlatforms(ref OutPlatforms, UnrealBuildTool.IsServerPlatform, bCheckValidity);
+        }
+
+        /// <summary>
+        /// Is this a valid configuration
+        /// Used primarily for Rocket vs non-Rocket
+        /// </summary>
+        /// <param name="InConfiguration"></param>
+        /// <returns>true if valid, false if not</returns>
+        static public bool IsValidConfiguration(UnrealTargetConfiguration InConfiguration)
+        {
+            if (RunningRocket())
+            {
+                if (
+                    (InConfiguration != UnrealTargetConfiguration.Development)
+                    && (InConfiguration != UnrealTargetConfiguration.DebugGame)
+                    && (InConfiguration != UnrealTargetConfiguration.Shipping)
+                    )
+                {
+                    return false;
+                }
+            }
+            return true;
         }
 
-		/// <summary>
-		/// Is this a valid configuration
-		/// Used primarily for Rocket vs non-Rocket
-		/// </summary>
-		/// <param name="InConfiguration"></param>
-		/// <returns>true if valid, false if not</returns>
-		static public bool IsValidConfiguration(UnrealTargetConfiguration InConfiguration)
-		{
-			if (RunningRocket())
-			{
-				if (
-					(InConfiguration != UnrealTargetConfiguration.Development)
-					&& (InConfiguration != UnrealTargetConfiguration.DebugGame)
-					&& (InConfiguration != UnrealTargetConfiguration.Shipping)
-					)
-				{
-					return false;
-				}
-			}
-			return true;
-		}
-
-		/// <summary>
-		/// Is this a valid platform
-		/// Used primarily for Rocket vs non-Rocket
-		/// </summary>
-		/// <param name="InPlatform"></param>
-		/// <returns>true if valid, false if not</returns>
-		static public bool IsValidPlatform(UnrealTargetPlatform InPlatform)
-		{
-			if (RunningRocket())
-			{
-				if(Utils.IsRunningOnMono)
-				{
+        /// <summary>
+        /// Is this a valid platform
+        /// Used primarily for Rocket vs non-Rocket
+        /// </summary>
+        /// <param name="InPlatform"></param>
+        /// <returns>true if valid, false if not</returns>
+        static public bool IsValidPlatform(UnrealTargetPlatform InPlatform)
+        {
+            if (RunningRocket())
+            {
+                if(Utils.IsRunningOnMono)
+                {
                     if (InPlatform != UnrealTargetPlatform.Mac && InPlatform != UnrealTargetPlatform.IOS && InPlatform != UnrealTargetPlatform.Linux)
-					{
-						return false;
-					}
-				}
-				else
-				{
-					if(InPlatform != UnrealTargetPlatform.Win32 && InPlatform != UnrealTargetPlatform.Win64 && InPlatform != UnrealTargetPlatform.Android)
-					{
-						return false;
-					}
-				}
-			}
-			return true;
-		}
-
-
-		/// <summary>
-		/// Is this a valid platform
-		/// Used primarily for Rocket vs non-Rocket
-		/// </summary>
-		/// <param name="InPlatform"></param>
-		/// <returns>true if valid, false if not</returns>
-		static public bool IsValidDesktopPlatform(UnrealTargetPlatform InPlatform)
-		{
-			switch (ExternalExecution.GetRuntimePlatform())
-			{
-			case UnrealTargetPlatform.Linux:
-				return InPlatform == UnrealTargetPlatform.Linux;
-			case UnrealTargetPlatform.Mac:
-				return InPlatform == UnrealTargetPlatform.Mac;
-			case UnrealTargetPlatform.Win64:
-				return ((InPlatform == UnrealTargetPlatform.Win32) || (InPlatform == UnrealTargetPlatform.Win64));
-			default:
-				throw new BuildException("Invalid RuntimePlatform:" + ExternalExecution.GetRuntimePlatform());
-			}
-		}
-
-		/// <summary>
-		/// See if the given string was pass in on the command line
-		/// </summary>
-		/// <param name="InString">The argument to check for (must include '-' if expected)</param>
-		/// <returns>true if the argument is found, false if not</returns>
-		static public bool CommandLineContains(string InString)
-		{
-			string LowercaseInString = InString.ToLowerInvariant();
-			foreach (string Arg in CmdLine)
-			{
-				string LowercaseArg = Arg.ToLowerInvariant();
-				if (LowercaseArg == LowercaseInString)
-				{
-					return true;
-				}
-			}
-
-			return false;
-		}
-
-		public static void RegisterAllUBTClasses()
-		{
-			// Find and register all tool chains and build platforms that are present
-			Assembly UBTAssembly = Assembly.GetExecutingAssembly();
-			if (UBTAssembly != null)
-			{
-				Log.TraceVerbose("Searching for ToolChains, BuildPlatforms, BuildDeploys and ProjectGenerators...");
-
-				List<System.Type> ProjectGeneratorList = new List<System.Type>();
-				var AllTypes = UBTAssembly.GetTypes();
+                    {
+                        return false;
+                    }
+                }
+                else
+                {
+                    if(InPlatform != UnrealTargetPlatform.Win32 && InPlatform != UnrealTargetPlatform.Win64 && InPlatform != UnrealTargetPlatform.Android)
+                    {
+                        return false;
+                    }
+                }
+            }
+            return true;
+        }
+
+
+        /// <summary>
+        /// Is this a valid platform
+        /// Used primarily for Rocket vs non-Rocket
+        /// </summary>
+        /// <param name="InPlatform"></param>
+        /// <returns>true if valid, false if not</returns>
+        static public bool IsValidDesktopPlatform(UnrealTargetPlatform InPlatform)
+        {
+            switch (ExternalExecution.GetRuntimePlatform())
+            {
+            case UnrealTargetPlatform.Linux:
+                return InPlatform == UnrealTargetPlatform.Linux;
+            case UnrealTargetPlatform.Mac:
+                return InPlatform == UnrealTargetPlatform.Mac;
+            case UnrealTargetPlatform.Win64:
+                return ((InPlatform == UnrealTargetPlatform.Win32) || (InPlatform == UnrealTargetPlatform.Win64));
+            default:
+                throw new BuildException("Invalid RuntimePlatform:" + ExternalExecution.GetRuntimePlatform());
+            }
+        }
+
+        /// <summary>
+        /// See if the given string was pass in on the command line
+        /// </summary>
+        /// <param name="InString">The argument to check for (must include '-' if expected)</param>
+        /// <returns>true if the argument is found, false if not</returns>
+        static public bool CommandLineContains(string InString)
+        {
+            string LowercaseInString = InString.ToLowerInvariant();
+            foreach (string Arg in CmdLine)
+            {
+                string LowercaseArg = Arg.ToLowerInvariant();
+                if (LowercaseArg == LowercaseInString)
+                {
+                    return true;
+                }
+            }
+
+            return false;
+        }
+
+        public static void RegisterAllUBTClasses()
+        {
+            // Find and register all tool chains and build platforms that are present
+            Assembly UBTAssembly = Assembly.GetExecutingAssembly();
+            if (UBTAssembly != null)
+            {
+                Log.TraceVerbose("Searching for ToolChains, BuildPlatforms, BuildDeploys and ProjectGenerators...");
+
+                List<System.Type> ProjectGeneratorList = new List<System.Type>();
+                var AllTypes = UBTAssembly.GetTypes();
                 // register all build platforms first, since they implement SDK-switching logic that can set environment variables
                 foreach (var CheckType in AllTypes)
                 {
                     if (CheckType.IsClass && !CheckType.IsAbstract)
                     {
-						if (Utils.ImplementsInterface<IUEBuildPlatform>(CheckType))
+                        if (Utils.ImplementsInterface<IUEBuildPlatform>(CheckType))
                         {
                             Log.TraceVerbose("    Registering build platform: {0}", CheckType.ToString());
                             var TempInst = (UEBuildPlatform)(UBTAssembly.CreateInstance(CheckType.FullName, true));
@@ -400,278 +400,279 @@ namespace UnrealBuildTool
                     }
                 }
                 // register all the other classes
-				foreach (var CheckType in AllTypes)
-				{
-					if (CheckType.IsClass && !CheckType.IsAbstract)
-					{
-						if (Utils.ImplementsInterface<IUEToolChain>(CheckType))
-						{
-							Log.TraceVerbose("    Registering tool chain    : {0}", CheckType.ToString());
-							var TempInst = (UEToolChain)(UBTAssembly.CreateInstance(CheckType.FullName, true));
-							TempInst.RegisterToolChain();
-						}
-						else if (Utils.ImplementsInterface<IUEBuildDeploy>(CheckType))
-						{
-							Log.TraceVerbose("    Registering build deploy  : {0}", CheckType.ToString());
-							var TempInst = (UEBuildDeploy)(UBTAssembly.CreateInstance(CheckType.FullName, true));
-							TempInst.RegisterBuildDeploy();
-						}
-						else if (CheckType.IsSubclassOf(typeof(UEPlatformProjectGenerator)))
-						{
-							ProjectGeneratorList.Add(CheckType);
-						}
-					}
-				}
-
-				// We need to process the ProjectGenerators AFTER all BuildPlatforms have been 
-				// registered as they use the BuildPlatform to verify they are legitimate.
-				foreach (var ProjType in ProjectGeneratorList)
-				{
-					Log.TraceVerbose("    Registering project generator: {0}", ProjType.ToString());
-					UEPlatformProjectGenerator TempInst = (UEPlatformProjectGenerator)(UBTAssembly.CreateInstance(ProjType.FullName, true));
-					TempInst.RegisterPlatformProjectGenerator();
-				}
-			}
-		}
-
-		public static void SetProjectFile(string InProjectFile)
-		{
-			if (HasUProjectFile() == false)
-			{
-				UProjectFile = InProjectFile;
-				if (UProjectFile.Length > 0)
-				{
-					bHasUProjectFile = true;
-					UProjectPath = Path.GetDirectoryName(UProjectFile);
-				}
-				else
-				{
-					bHasUProjectFile = false;
-				}
-			}
-			else
-			{
-				throw new BuildException("Trying to set the project file to {0} when it is already set to {1}", InProjectFile, UProjectFile);
-			}
-		}
-
-		private static bool ParseRocketCommandlineArg(string InArg, ref string OutGameName)
-		{
-			string LowercaseArg = InArg.ToLowerInvariant();
-			bool bSetupProject = false;
-			
-			if (LowercaseArg == "-rocket")
-			{
-				bRunningRocket = true;
-			}
-			else if (LowercaseArg.StartsWith("-project="))
-			{
-				UProjectFile = InArg.Substring(9);
-				bSetupProject = true;
-			}
-			else if (LowercaseArg.EndsWith(".uproject"))
-			{
-				UProjectFile = InArg;
-				bSetupProject = true;
-			}
-			else if (LowercaseArg == "-buildrocket")
-			{
-				bBuildingRocket = true;
-			}
-			else
-			{
-				return false;
-			}
-
-			if (bSetupProject && (UProjectFile.Length > 0))
-			{
-				bHasUProjectFile = true;
-	
-				UProjectPath = Path.GetDirectoryName(UProjectFile);
-				OutGameName = Path.GetFileNameWithoutExtension(UProjectFile);
-
-				if (Path.IsPathRooted(UProjectPath) == false)
-				{
-					if (Directory.Exists(UProjectPath) == false)
-					{
-						// If it is *not* rooted, then we potentially received a path that is 
-						// relative to Engine/Binaries/[PLATFORM] rather than Engine/Source
-						if (UProjectPath.StartsWith("../") || UProjectPath.StartsWith("..\\"))
-						{
-							string AlternativeProjectpath = UProjectPath.Substring(3);
-							if (Directory.Exists(AlternativeProjectpath))
-							{
-								UProjectPath = AlternativeProjectpath;
-								UProjectFile = UProjectFile.Substring(3);
-								Debug.Assert(UProjectFile.Length > 0);
-							}
-						}
-					}
-				}
-			}
-
-			return true;
-		}
-
-		/// <summary>
-		/// Sets up UBT Trace listeners and filters.
-		/// </summary>
-		private static void InitLogging()
-		{
-			Trace.Listeners.Add(new ConsoleListener());
-			var Filter = new VerbosityFilter();
-			foreach (TraceListener Listener in Trace.Listeners)
-			{
-				Listener.Filter = Filter;
-			}
-		}
-
-		private static int Main(string[] Arguments)
-		{
-			InitLogging();
-
-			try
-			{
-				XmlConfigLoader.Init();
-			}
-			catch (BuildException Exception)
-			{
-				Log.TraceError("UnrealBuildTool Exception: " + Exception);
-				return (int) ECompilationResult.OtherCompilationError;
-			}
-
-			InitialEnvironment = Environment.GetEnvironmentVariables();
-			if (InitialEnvironment.Count < 1)
-			{
-				throw new BuildException("Environment could not be read");
-			}
-			// Helpers used for stats tracking.
-			StartTime = DateTime.UtcNow;
-
-			Telemetry.Initialize();
-
-			// Copy off the arguments to allow checking for command-line arguments elsewhere
-			CmdLine = new List<string>(Arguments);
-
-			foreach (string Arg in Arguments)
-			{
-				string TempGameName = null;
-				if (ParseRocketCommandlineArg(Arg, ref TempGameName) == true)
-				{
-					// This is to allow relative paths for the project file
-					Log.TraceVerbose("UBT Running for Rocket: " + UProjectFile);
-				}
-			}
-
-			// Change the working directory to be the Engine/Source folder. We are running from Engine/Binaries/DotNET
-			string EngineSourceDirectory = Path.Combine(Utils.GetExecutingAssemblyDirectory(), "..", "..", "..", "Engine", "Source");
-
-			[REDACTED_EMAIL]: This is a workaround for recompiling game code in editor
-			// The working directory when launching is *not* what we would expect
-			if (Directory.Exists(EngineSourceDirectory) == false)
-			{
-				// We are assuming UBT always runs from <>/Engine/...
-				EngineSourceDirectory = Path.GetDirectoryName(Utils.GetExecutingAssemblyLocation());
-				EngineSourceDirectory = EngineSourceDirectory.Replace("\\", "/");
-				Int32 EngineIdx = EngineSourceDirectory.IndexOf("/Engine/");
-				if (EngineIdx != 0)
-				{
-					EngineSourceDirectory = EngineSourceDirectory.Substring(0, EngineIdx + 8);
-					EngineSourceDirectory += "Source";
-				}
-			}
-			Directory.SetCurrentDirectory(EngineSourceDirectory);
-
-			// Build the list of game projects that we know about. When building from the editor (for hot-reload) or for Rocket projects, we require the 
-			// project file to be passed in. Otherwise we scan for projects in directories named in UE4Games.uprojectdirs.
-			if (HasUProjectFile())
-			{
-				string SourceFolder = Path.Combine(Path.GetDirectoryName(GetUProjectFile()), "Source");
-				UProjectInfo.AddProject(GetUProjectFile(), Directory.Exists(SourceFolder));
-			}
-			else if(!RunningRocket())
-			{
-				UProjectInfo.FillProjectInfo();
-			}
-
-			ECompilationResult Result = ECompilationResult.Succeeded;
-
-			// Reset early so we can access BuildConfiguration even before RunUBT() is called
-			XmlConfigLoader.Reset<BuildConfiguration>();
-			XmlConfigLoader.Reset<UEBuildConfiguration>();
-
-			Log.TraceVerbose("UnrealBuildTool (DEBUG OUTPUT MODE)");
-			Log.TraceVerbose("Command-line: {0}", String.Join(" ", Arguments));
-
-			bool bRunCopyrightVerification = false;
-			bool bDumpToFile = false;
-			bool bCheckThirdPartyHeaders = false;
-
-			// @todo: Ideally we never need to Mutex unless we are invoked with the same target project,
-			// in the same branch/path!  This would allow two clientspecs to build at the same time (even though we need
-			// to make sure that XGE is only used once at a time)
-			bool bUseMutex = true;
-			{
-				int NoMutexArgumentIndex;
-				if (Utils.ParseCommandLineFlag(Arguments, "-NoMutex", out NoMutexArgumentIndex))
-				{
-					bUseMutex = false;
-				}
-			}
-
-			// Don't allow simultaneous execution of Unreal Built Tool. Multi-selection in the UI e.g. causes this and you want serial
-			// execution in this case to avoid contention issues with shared produced items.
-			bool bCreatedMutex = false;
-			{
-				Mutex SingleInstanceMutex = null;
-				if (bUseMutex)
-				{
-					int LocationHash = Utils.GetExecutingAssemblyLocation().GetHashCode();
-					SingleInstanceMutex = new Mutex(true, "Global/UnrealBuildTool_Mutex_" + LocationHash.ToString(), out bCreatedMutex);
-					if (!bCreatedMutex)
-					{
-						// If this instance didn't create the mutex, wait for the existing mutex to be released by the mutex's creator.
-						DateTime MutexWaitStartTime = DateTime.UtcNow;
-						SingleInstanceMutex.WaitOne();
-						MutexWaitTime = DateTime.UtcNow - MutexWaitStartTime;
-					}
-				}
-
-				try
-				{
-					string GameName = null;
-					var bGenerateVCProjectFiles = false;
-					var bGenerateXcodeProjectFiles = false;
+                foreach (var CheckType in AllTypes)
+                {
+                    if (CheckType.IsClass && !CheckType.IsAbstract)
+                    {
+                        if (Utils.ImplementsInterface<IUEToolChain>(CheckType))
+                        {
+                            Log.TraceVerbose("    Registering tool chain    : {0}", CheckType.ToString());
+                            var TempInst = (UEToolChain)(UBTAssembly.CreateInstance(CheckType.FullName, true));
+                            TempInst.RegisterToolChain();
+                        }
+                        else if (Utils.ImplementsInterface<IUEBuildDeploy>(CheckType))
+                        {
+                            Log.TraceVerbose("    Registering build deploy  : {0}", CheckType.ToString());
+                            var TempInst = (UEBuildDeploy)(UBTAssembly.CreateInstance(CheckType.FullName, true));
+                            TempInst.RegisterBuildDeploy();
+                        }
+                        else if (CheckType.IsSubclassOf(typeof(UEPlatformProjectGenerator)))
+                        {
+                            ProjectGeneratorList.Add(CheckType);
+                        }
+                    }
+                }
+
+                // We need to process the ProjectGenerators AFTER all BuildPlatforms have been 
+                // registered as they use the BuildPlatform to verify they are legitimate.
+                foreach (var ProjType in ProjectGeneratorList)
+                {
+                    Log.TraceVerbose("    Registering project generator: {0}", ProjType.ToString());
+                    UEPlatformProjectGenerator TempInst = (UEPlatformProjectGenerator)(UBTAssembly.CreateInstance(ProjType.FullName, true));
+                    TempInst.RegisterPlatformProjectGenerator();
+                }
+            }
+        }
+
+        public static void SetProjectFile(string InProjectFile)
+        {
+            if (HasUProjectFile() == false)
+            {
+                UProjectFile = InProjectFile;
+                if (UProjectFile.Length > 0)
+                {
+                    bHasUProjectFile = true;
+                    UProjectPath = Path.GetDirectoryName(UProjectFile);
+                }
+                else
+                {
+                    bHasUProjectFile = false;
+                }
+            }
+            else
+            {
+                throw new BuildException("Trying to set the project file to {0} when it is already set to {1}", InProjectFile, UProjectFile);
+            }
+        }
+
+        private static bool ParseRocketCommandlineArg(string InArg, ref string OutGameName)
+        {
+            string LowercaseArg = InArg.ToLowerInvariant();
+            bool bSetupProject = false;
+            
+            if (LowercaseArg == "-rocket")
+            {
+                bRunningRocket = true;
+            }
+            else if (LowercaseArg.StartsWith("-project="))
+            {
+                UProjectFile = InArg.Substring(9);
+                bSetupProject = true;
+            }
+            else if (LowercaseArg.EndsWith(".uproject"))
+            {
+                UProjectFile = InArg;
+                bSetupProject = true;
+            }
+            else if (LowercaseArg == "-buildrocket")
+            {
+                bBuildingRocket = true;
+            }
+            else
+            {
+                return false;
+            }
+
+            if (bSetupProject && (UProjectFile.Length > 0))
+            {
+                bHasUProjectFile = true;
+    
+                UProjectPath = Path.GetDirectoryName(UProjectFile);
+                OutGameName = Path.GetFileNameWithoutExtension(UProjectFile);
+
+                if (Path.IsPathRooted(UProjectPath) == false)
+                {
+                    if (Directory.Exists(UProjectPath) == false)
+                    {
+                        // If it is *not* rooted, then we potentially received a path that is 
+                        // relative to Engine/Binaries/[PLATFORM] rather than Engine/Source
+                        if (UProjectPath.StartsWith("../") || UProjectPath.StartsWith("..\\"))
+                        {
+                            string AlternativeProjectpath = UProjectPath.Substring(3);
+                            if (Directory.Exists(AlternativeProjectpath))
+                            {
+                                UProjectPath = AlternativeProjectpath;
+                                UProjectFile = UProjectFile.Substring(3);
+                                Debug.Assert(UProjectFile.Length > 0);
+                            }
+                        }
+                    }
+                }
+            }
+
+            return true;
+        }
+
+        /// <summary>
+        /// Sets up UBT Trace listeners and filters.
+        /// </summary>
+        private static void InitLogging()
+        {
+            Trace.Listeners.Add(new ConsoleListener());
+            var Filter = new VerbosityFilter();
+            foreach (TraceListener Listener in Trace.Listeners)
+            {
+                Listener.Filter = Filter;
+            }
+        }
+
+        private static int Main(string[] Arguments)
+        {
+            InitLogging();
+
+            try
+            {
+                XmlConfigLoader.Init();
+            }
+            catch (BuildException Exception)
+            {
+                Log.TraceError("UnrealBuildTool Exception: " + Exception);
+                return (int) ECompilationResult.OtherCompilationError;
+            }
+
+            InitialEnvironment = Environment.GetEnvironmentVariables();
+            if (InitialEnvironment.Count < 1)
+            {
+                throw new BuildException("Environment could not be read");
+            }
+            // Helpers used for stats tracking.
+            StartTime = DateTime.UtcNow;
+
+            Telemetry.Initialize();
+
+            // Copy off the arguments to allow checking for command-line arguments elsewhere
+            CmdLine = new List<string>(Arguments);
+
+            foreach (string Arg in Arguments)
+            {
+                string TempGameName = null;
+                if (ParseRocketCommandlineArg(Arg, ref TempGameName) == true)
+                {
+                    // This is to allow relative paths for the project file
+                    Log.TraceVerbose("UBT Running for Rocket: " + UProjectFile);
+                }
+            }
+
+            // Change the working directory to be the Engine/Source folder. We are running from Engine/Binaries/DotNET
+            string EngineSourceDirectory = Path.Combine(Utils.GetExecutingAssemblyDirectory(), "..", "..", "..", "Engine", "Source");
+
+            [REDACTED_EMAIL]: This is a workaround for recompiling game code in editor
+            // The working directory when launching is *not* what we would expect
+            if (Directory.Exists(EngineSourceDirectory) == false)
+            {
+                // We are assuming UBT always runs from <>/Engine/...
+                EngineSourceDirectory = Path.GetDirectoryName(Utils.GetExecutingAssemblyLocation());
+                EngineSourceDirectory = EngineSourceDirectory.Replace("\\", "/");
+                Int32 EngineIdx = EngineSourceDirectory.IndexOf("/Engine/");
+                if (EngineIdx != 0)
+                {
+                    EngineSourceDirectory = EngineSourceDirectory.Substring(0, EngineIdx + 8);
+                    EngineSourceDirectory += "Source";
+                }
+            }
+            Directory.SetCurrentDirectory(EngineSourceDirectory);
+
+            // Build the list of game projects that we know about. When building from the editor (for hot-reload) or for Rocket projects, we require the 
+            // project file to be passed in. Otherwise we scan for projects in directories named in UE4Games.uprojectdirs.
+            if (HasUProjectFile())
+            {
+                string SourceFolder = Path.Combine(Path.GetDirectoryName(GetUProjectFile()), "Source");
+                UProjectInfo.AddProject(GetUProjectFile(), Directory.Exists(SourceFolder));
+            }
+            else if(!RunningRocket())
+            {
+                UProjectInfo.FillProjectInfo();
+            }
+
+            ECompilationResult Result = ECompilationResult.Succeeded;
+
+            // Reset early so we can access BuildConfiguration even before RunUBT() is called
+            XmlConfigLoader.Reset<BuildConfiguration>();
+            XmlConfigLoader.Reset<UEBuildConfiguration>();
+
+            Log.TraceVerbose("UnrealBuildTool (DEBUG OUTPUT MODE)");
+            Log.TraceVerbose("Command-line: {0}", String.Join(" ", Arguments));
+
+            bool bRunCopyrightVerification = false;
+            bool bDumpToFile = false;
+            bool bCheckThirdPartyHeaders = false;
+
+            // @todo: Ideally we never need to Mutex unless we are invoked with the same target project,
+            // in the same branch/path!  This would allow two clientspecs to build at the same time (even though we need
+            // to make sure that XGE is only used once at a time)
+            bool bUseMutex = true;
+            {
+                int NoMutexArgumentIndex;
+                if (Utils.ParseCommandLineFlag(Arguments, "-NoMutex", out NoMutexArgumentIndex))
+                {
+                    bUseMutex = false;
+                }
+            }
+
+            // Don't allow simultaneous execution of Unreal Built Tool. Multi-selection in the UI e.g. causes this and you want serial
+            // execution in this case to avoid contention issues with shared produced items.
+            bool bCreatedMutex = false;
+            {
+                Mutex SingleInstanceMutex = null;
+                if (bUseMutex)
+                {
+                    int LocationHash = Utils.GetExecutingAssemblyLocation().GetHashCode();
+                    SingleInstanceMutex = new Mutex(true, "Global/UnrealBuildTool_Mutex_" + LocationHash.ToString(), out bCreatedMutex);
+                    if (!bCreatedMutex)
+                    {
+                        // If this instance didn't create the mutex, wait for the existing mutex to be released by the mutex's creator.
+                        DateTime MutexWaitStartTime = DateTime.UtcNow;
+                        SingleInstanceMutex.WaitOne();
+                        MutexWaitTime = DateTime.UtcNow - MutexWaitStartTime;
+                    }
+                }
+
+                try
+                {
+                    string GameName = null;
+                    var bGenerateVCProjectFiles = false;
+                    var bGenerateXcodeProjectFiles = false;
                     var bGenerateMakefiles = false;
-					var bValidPlatformsOnly = false;
-					var bSpecificModulesOnly = false;
-					var bIgnoreJunk = false;
-
-					// We need to be able to identify if one of the arguments is the platform...
-					// Leverage the existing parser function in UEBuildTarget to get this information.
-					UnrealTargetPlatform CheckPlatform;
-					UnrealTargetConfiguration CheckConfiguration;
-					UEBuildTarget.ParsePlatformAndConfiguration(Arguments, out CheckPlatform, out CheckConfiguration, false);
-					//@todo SAS: Add a true Debug mode!
-					if (RunningRocket())
-					{
-						// Only Development and Shipping are supported for engine modules
-						if( CheckConfiguration != UnrealTargetConfiguration.Development && CheckConfiguration != UnrealTargetConfiguration.Shipping )
-						{
-							CheckConfiguration = UnrealTargetConfiguration.Development;
-						}
-					}
-
-					foreach (var Arg in Arguments)
-					{
-						string LowercaseArg = Arg.ToLowerInvariant();
-			            const string OnlyPlatformSpecificForArg = "-onlyplatformspecificfor=";
-
-						if (ParseRocketCommandlineArg(Arg, ref GameName))
-						{
-							// Already handled at startup. Calling now just to properly set the game name
-							continue;
-						}
+                    var bValidPlatformsOnly = false;
+                    var bSpecificModulesOnly = false;
+                    var bIgnoreJunk = false;
+                    var bAutoSDKOnly = false;
+
+                    // We need to be able to identify if one of the arguments is the platform...
+                    // Leverage the existing parser function in UEBuildTarget to get this information.
+                    UnrealTargetPlatform CheckPlatform;
+                    UnrealTargetConfiguration CheckConfiguration;
+                    UEBuildTarget.ParsePlatformAndConfiguration(Arguments, out CheckPlatform, out CheckConfiguration, false);
+                    //@todo SAS: Add a true Debug mode!
+                    if (RunningRocket())
+                    {
+                        // Only Development and Shipping are supported for engine modules
+                        if( CheckConfiguration != UnrealTargetConfiguration.Development && CheckConfiguration != UnrealTargetConfiguration.Shipping )
+                        {
+                            CheckConfiguration = UnrealTargetConfiguration.Development;
+                        }
+                    }
+
+                    foreach (var Arg in Arguments)
+                    {
+                        string LowercaseArg = Arg.ToLowerInvariant();
+                        const string OnlyPlatformSpecificForArg = "-onlyplatformspecificfor=";
+
+                        if (ParseRocketCommandlineArg(Arg, ref GameName))
+                        {
+                            // Already handled at startup. Calling now just to properly set the game name
+                            continue;
+                        }
                         else if (LowercaseArg == "-skipnonhostplatforms")
                         {
                             bSkipNonHostPlatforms = true;
@@ -697,590 +698,600 @@ namespace UnrealBuildTool
                             ProjectFileGenerator.bGenerateProjectFiles = true;
                         }
                         else if (LowercaseArg.StartsWith("-projectfile"))
-						{
-							bGenerateVCProjectFiles = true;
-							ProjectFileGenerator.bGenerateProjectFiles = true;
-						}
-						else if (LowercaseArg.StartsWith("-xcodeprojectfile"))
-						{
-							bGenerateXcodeProjectFiles = true;
-							ProjectFileGenerator.bGenerateProjectFiles = true;
-						}
-						else if (LowercaseArg == "-validplatformsonly")
-						{
-							bValidPlatformsOnly = true;
-						}
-						else if (LowercaseArg == "development" || LowercaseArg == "debug" || LowercaseArg == "shipping" || LowercaseArg == "test" || LowercaseArg == "debuggame")
-						{
-							//ConfigName = Arg;
-						}
+                        {
+                            bGenerateVCProjectFiles = true;
+                            ProjectFileGenerator.bGenerateProjectFiles = true;
+                        }
+                        else if (LowercaseArg.StartsWith("-xcodeprojectfile"))
+                        {
+                            bGenerateXcodeProjectFiles = true;
+                            ProjectFileGenerator.bGenerateProjectFiles = true;
+                        }
+                        else if (LowercaseArg == "-validplatformsonly")
+                        {
+                            bValidPlatformsOnly = true;
+                        }
+                        else if (LowercaseArg == "development" || LowercaseArg == "debug" || LowercaseArg == "shipping" || LowercaseArg == "test" || LowercaseArg == "debuggame")
+                        {
+                            //ConfigName = Arg;
+                        }
                         else if (LowercaseArg == "-obbinapk")
                         {
                             UEBuildConfiguration.bOBBinAPK = true;
                         }
-						else if (LowercaseArg == "-modulewithsuffix")
-						{
-							bSpecificModulesOnly = true;
-							continue;
-						}
-						else if (LowercaseArg == "-forceheadergeneration")
-						{
-							// Don't reset this as we don't recheck the args in RunUBT
-							UEBuildConfiguration.bForceHeaderGeneration = true;
-						}
-						else if (LowercaseArg == "-nobuilduht")
-						{
-							UEBuildConfiguration.bDoNotBuildUHT = true;
-						}
-						else if (LowercaseArg == "-failifgeneratedcodechanges")
-						{
-							UEBuildConfiguration.bFailIfGeneratedCodeChanges = true;
-						}
-						else if (LowercaseArg == "-copyrightverification")
-						{
-							bRunCopyrightVerification = true;
-						}
-						else if (LowercaseArg == "-dumptofile")
-						{
-							bDumpToFile = true;
-						}
-						else if (LowercaseArg == "-check3rdpartyheaders")
-						{
-							bCheckThirdPartyHeaders = true;
-						}
-						else if (LowercaseArg == "-clean")
-						{
-							UEBuildConfiguration.bCleanProject = true;
-						}
-						else if (LowercaseArg == "-prepfordeploy")
-						{
-							UEBuildConfiguration.bPrepForDeployment = true;
-						}
-						else if (LowercaseArg == "-generateexternalfilelist")
-						{
-							UEBuildConfiguration.bGenerateExternalFileList = true;
-						}
-						else if (LowercaseArg == "-mergeexternalfilelist")
-						{
-							UEBuildConfiguration.bMergeExternalFileList = true;
-						}
-						else if (LowercaseArg == "-generatemanifest")
-						{
-							// Generate a manifest file containing all the files required to be in Perforce
-							UEBuildConfiguration.bGenerateManifest = true;
-						}
-						else if (LowercaseArg == "-mergemanifests")
-						{
-							// Whether to add to the existing manifest (if it exists), or start afresh
-							UEBuildConfiguration.bMergeManifests = true;
-						}
-						else if (LowercaseArg == "-noxge")
-						{
-							BuildConfiguration.bAllowXGE = false;
-						}
-						else if (LowercaseArg == "-xgeexport")
-						{
-							BuildConfiguration.bXGEExport = true;
-							BuildConfiguration.bAllowXGE = true;
-						}
-						else if (LowercaseArg == "-nosimplygon")
-						{
-							UEBuildConfiguration.bCompileSimplygon = false;
-						}
-						else if (LowercaseArg == "-nospeedtree")
-						{
-							UEBuildConfiguration.bCompileSpeedTree = false;
-						}
-						else if (LowercaseArg == "-ignorejunk")
-						{
-							bIgnoreJunk = true;
-						}
-						else if (LowercaseArg == "-define")
-						{
-							// Skip -define
-						}
-						else if (LowercaseArg == "-progress")
-						{
-							ProgressWriter.bWriteMarkup = true;
-						}
-						else if (CheckPlatform.ToString().ToLowerInvariant() == LowercaseArg)
-						{
-							// It's the platform set...
-							//PlatformName = Arg;
-						}
-						else
-						{
-							// This arg may be a game name. Check for the existence of a game folder with this name.
-							// "Engine" is not a valid game name.
-							if (LowercaseArg != "engine" && Arg.IndexOfAny(Path.GetInvalidPathChars()) == -1 &&
-								Directory.Exists(Path.Combine(ProjectFileGenerator.RootRelativePath, Arg, "Config")))
-							{
-								GameName = Arg;
-								Log.TraceVerbose("CommandLine: Found game name '{0}'", GameName);
-							}
-							else if (LowercaseArg == "rocket")
-							{
-								GameName = Arg;
-								Log.TraceVerbose("CommandLine: Found game name '{0}'", GameName);
-							}
-						}
-					}
-
-					// Send an event with basic usage dimensions
-					Telemetry.SendEvent("CommonAttributes.2",
-						// @todo No platform independent way to do processor speed and count. Defer for now. Mono actually has ways to do this but needs to be tested.
-						"ProcessorCount", Environment.ProcessorCount.ToString(),
-						"ComputerName", Environment.MachineName,
-						"User", Environment.UserName,
-						"Domain", Environment.UserDomainName,
-						"CommandLine", string.Join(" ", Arguments),
-						"UBT Action", bGenerateVCProjectFiles 
-							? "GenerateVCProjectFiles" 
-							: bGenerateXcodeProjectFiles 
-								? "GenerateXcodeProjectFiles" 
-								: bRunCopyrightVerification
-									? "RunCopyrightVerification"
-                                        : bGenerateMakefiles
-                                        ? "GenerateMakefiles"
-                                            : "Build",
-						"Platform", CheckPlatform.ToString(),
-						"Configuration", CheckConfiguration.ToString(),
-						"IsRocket", bRunningRocket.ToString(),
-						"EngineVersion", Utils.GetEngineVersionFromObjVersionCPP().ToString()
-						);
+                        else if (LowercaseArg == "-modulewithsuffix")
+                        {
+                            bSpecificModulesOnly = true;
+                            continue;
+                        }
+                        else if (LowercaseArg == "-forceheadergeneration")
+                        {
+                            // Don't reset this as we don't recheck the args in RunUBT
+                            UEBuildConfiguration.bForceHeaderGeneration = true;
+                        }
+                        else if (LowercaseArg == "-nobuilduht")
+                        {
+                            UEBuildConfiguration.bDoNotBuildUHT = true;
+                        }
+                        else if (LowercaseArg == "-failifgeneratedcodechanges")
+                        {
+                            UEBuildConfiguration.bFailIfGeneratedCodeChanges = true;
+                        }
+                        else if (LowercaseArg == "-copyrightverification")
+                        {
+                            bRunCopyrightVerification = true;
+                        }
+                        else if (LowercaseArg == "-dumptofile")
+                        {
+                            bDumpToFile = true;
+                        }
+                        else if (LowercaseArg == "-check3rdpartyheaders")
+                        {
+                            bCheckThirdPartyHeaders = true;
+                        }
+                        else if (LowercaseArg == "-clean")
+                        {
+                            UEBuildConfiguration.bCleanProject = true;
+                        }
+                        else if (LowercaseArg == "-prepfordeploy")
+                        {
+                            UEBuildConfiguration.bPrepForDeployment = true;
+                        }
+                        else if (LowercaseArg == "-generateexternalfilelist")
+                        {
+                            UEBuildConfiguration.bGenerateExternalFileList = true;
+                        }
+                        else if (LowercaseArg == "-mergeexternalfilelist")
+                        {
+                            UEBuildConfiguration.bMergeExternalFileList = true;
+                        }
+                        else if (LowercaseArg == "-generatemanifest")
+                        {
+                            // Generate a manifest file containing all the files required to be in Perforce
+                            UEBuildConfiguration.bGenerateManifest = true;
+                        }
+                        else if (LowercaseArg == "-mergemanifests")
+                        {
+                            // Whether to add to the existing manifest (if it exists), or start afresh
+                            UEBuildConfiguration.bMergeManifests = true;
+                        }
+                        else if (LowercaseArg == "-noxge")
+                        {
+                            BuildConfiguration.bAllowXGE = false;
+                        }
+                        else if (LowercaseArg == "-xgeexport")
+                        {
+                            BuildConfiguration.bXGEExport = true;
+                            BuildConfiguration.bAllowXGE = true;
+                        }
+                        else if (LowercaseArg == "-nosimplygon")
+                        {
+                            UEBuildConfiguration.bCompileSimplygon = false;
+                        }
+                        else if (LowercaseArg == "-nospeedtree")
+                        {
+                            UEBuildConfiguration.bCompileSpeedTree = false;
+                        }
+                        else if (LowercaseArg == "-ignorejunk")
+                        {
+                            bIgnoreJunk = true;
+                        }
+                        else if (LowercaseArg == "-define")
+                        {
+                            // Skip -define
+                        }
+                        else if (LowercaseArg == "-progress")
+                        {
+                            ProgressWriter.bWriteMarkup = true;
+                        }
+                        else if (LowercaseArg == "-autosdkonly")
+                        {
+                            bAutoSDKOnly = true;                            
+                        }
+                        else if (CheckPlatform.ToString().ToLowerInvariant() == LowercaseArg)
+                        {
+                            // It's the platform set...
+                            //PlatformName = Arg;
+                        }
+                        
+                        else
+                        {
+                            // This arg may be a game name. Check for the existence of a game folder with this name.
+                            // "Engine" is not a valid game name.
+                            if (LowercaseArg != "engine" && Arg.IndexOfAny(Path.GetInvalidPathChars()) == -1 &&
+                                Directory.Exists(Path.Combine(ProjectFileGenerator.RootRelativePath, Arg, "Config")))
+                            {
+                                GameName = Arg;
+                                Log.TraceVerbose("CommandLine: Found game name '{0}'", GameName);
+                            }
+                            else if (LowercaseArg == "rocket")
+                            {
+                                GameName = Arg;
+                                Log.TraceVerbose("CommandLine: Found game name '{0}'", GameName);
+                            }
+                        }
+                    }
+
+                    // Send an event with basic usage dimensions
+                    Telemetry.SendEvent("CommonAttributes.2",
+                        // @todo No platform independent way to do processor speed and count. Defer for now. Mono actually has ways to do this but needs to be tested.
+                        "ProcessorCount", Environment.ProcessorCount.ToString(),
+                        "ComputerName", Environment.MachineName,
+                        "User", Environment.UserName,
+                        "Domain", Environment.UserDomainName,
+                        "CommandLine", string.Join(" ", Arguments),
+                        "UBT Action", bGenerateVCProjectFiles 
+                            ? "GenerateVCProjectFiles" 
+                            : bGenerateXcodeProjectFiles 
+                                ? "GenerateXcodeProjectFiles" 
+                                : bRunCopyrightVerification
+                                    ? "RunCopyrightVerification"
+                                        : bGenerateMakefiles
+                                        ? "GenerateMakefiles"
+                                            : "Build",
+                        "Platform", CheckPlatform.ToString(),
+                        "Configuration", CheckConfiguration.ToString(),
+                        "IsRocket", bRunningRocket.ToString(),
+                        "EngineVersion", Utils.GetEngineVersionFromObjVersionCPP().ToString()
+                        );
 
 
 
-					if (bValidPlatformsOnly == true)
-					{
-						// We turned the flag on if generating projects so that the
-						// build platforms would know we are doing so... In this case,
-						// turn it off to only generate projects for valid platforms.
-						ProjectFileGenerator.bGenerateProjectFiles = false;
-					}
+                    if (bValidPlatformsOnly == true)
+                    {
+                        // We turned the flag on if generating projects so that the
+                        // build platforms would know we are doing so... In this case,
+                        // turn it off to only generate projects for valid platforms.
+                        ProjectFileGenerator.bGenerateProjectFiles = false;
+                    }
 
-					// Find and register all tool chains, build platforms, etc. that are present
-					RegisterAllUBTClasses();
-					ProjectFileGenerator.bGenerateProjectFiles = false;
+                    // Find and register all tool chains, build platforms, etc. that are present
+                    RegisterAllUBTClasses();
+                    ProjectFileGenerator.bGenerateProjectFiles = false;
 
-					// now that we know the available platforms, we can delete other platforms' junk. if we're only building specific modules from the editor, don't touch anything else (it may be in use).
-					if (!bSpecificModulesOnly && !bIgnoreJunk)
-					{
-						JunkDeleter.DeleteJunk();
-					}
+                    // now that we know the available platforms, we can delete other platforms' junk. if we're only building specific modules from the editor, don't touch anything else (it may be in use).
+                    if (!bSpecificModulesOnly && !bIgnoreJunk)
+                    {
+                        JunkDeleter.DeleteJunk();
+                    }
 
                     if (bGenerateVCProjectFiles || bGenerateXcodeProjectFiles || bGenerateMakefiles)
-					{
-						bool bGenerationSuccess = true;
-						if (bGenerateVCProjectFiles)
-						{
-							bGenerationSuccess &= GenerateProjectFiles(new VCProjectFileGenerator(), Arguments);
-						}
-						if (bGenerateXcodeProjectFiles)
-						{
-							bGenerationSuccess &= GenerateProjectFiles(new XcodeProjectFileGenerator(), Arguments);
-						}
+                    {
+                        bool bGenerationSuccess = true;
+                        if (bGenerateVCProjectFiles)
+                        {
+                            bGenerationSuccess &= GenerateProjectFiles(new VCProjectFileGenerator(), Arguments);
+                        }
+                        if (bGenerateXcodeProjectFiles)
+                        {
+                            bGenerationSuccess &= GenerateProjectFiles(new XcodeProjectFileGenerator(), Arguments);
+                        }
                         if (bGenerateMakefiles)
- 						{
-							bGenerationSuccess &= GenerateProjectFiles(new MakefileGenerator(), Arguments);
- 						}
-
-						if(!bGenerationSuccess)
-						{
-							Result = ECompilationResult.OtherCompilationError;
-						}
-					}
-					else if (bRunCopyrightVerification)
-					{
-						CopyrightVerify.RunCopyrightVerification(EngineSourceDirectory, GameName, bDumpToFile);
-						Log.TraceInformation("Completed... exiting.");
-					}
-					else
-					{
-						// Check if any third party headers are included from public engine headers.
-						if (bCheckThirdPartyHeaders)
-						{
-							ThirdPartyHeaderFinder.FindThirdPartyIncludes(CheckPlatform, CheckConfiguration);
-						}
-
-						// Build our project
-						if (Result == ECompilationResult.Succeeded)
-						{
-							if (UEBuildConfiguration.bPrepForDeployment == false)
-							{
-								// If we are only prepping for deployment, assume the build already occurred.
-								Result = RunUBT(Arguments);
-							}
-							else
-							{
-								var BuildPlatform = UEBuildPlatform.GetBuildPlatform(CheckPlatform, true);
-								if (BuildPlatform != null)
-								{
-									// Setup environment wasn't called, so set the flag
-									BuildConfiguration.bDeployAfterCompile = BuildPlatform.RequiresDeployPrepAfterCompile();
-									BuildConfiguration.PlatformIntermediateFolder = Path.Combine(BuildConfiguration.BaseIntermediateFolder, CheckPlatform.ToString(), BuildPlatform.GetActiveArchitecture());
-								}
-							}
-
-							// If we build w/ bXGEExport true, we didn't REALLY build at this point, 
-							// so don't bother with doing the PrepTargetForDeployment call. 
-							if ((Result == ECompilationResult.Succeeded) && (BuildConfiguration.bDeployAfterCompile == true) && (BuildConfiguration.bXGEExport == false) &&
-								(UEBuildConfiguration.bGenerateManifest == false) && (UEBuildConfiguration.bGenerateExternalFileList == false) && (UEBuildConfiguration.bCleanProject == false))
-							{
-								var DeployHandler = UEBuildDeploy.GetBuildDeploy(CheckPlatform);
-								if (DeployHandler != null)
-								{
-									// We need to be able to identify the Target.Type we can derive it from the Arguments.
-									BuildConfiguration.bFlushBuildDirOnRemoteMac = false;
-									UEBuildTarget CheckTarget = UEBuildTarget.CreateTarget(Arguments);
-									CheckTarget.SetupGlobalEnvironment();
-									if ((CheckTarget.Rules.Type == TargetRules.TargetType.Game) || 
-										(CheckTarget.Rules.Type == TargetRules.TargetType.Server) || 
-										(CheckTarget.Rules.Type == TargetRules.TargetType.Client))
-									{
-										CheckTarget.AppName = CheckTarget.GameName;
-									}
-									DeployHandler.PrepTargetForDeployment(CheckTarget);
-								}
-							}
-						}
-					}
-					// Print some performance info
-					var BuildDuration = (DateTime.UtcNow - StartTime - MutexWaitTime).TotalSeconds;
-					if (BuildConfiguration.bPrintPerformanceInfo)
-					{
-						Log.TraceInformation("GetIncludes time: " + CPPEnvironment.TotalTimeSpentGettingIncludes + "s (" + CPPEnvironment.TotalIncludesRequested + " includes)" );
-						Log.TraceInformation("DirectIncludes cache miss time: " + CPPEnvironment.DirectIncludeCacheMissesTotalTime + "s (" + CPPEnvironment.TotalDirectIncludeCacheMisses + " misses)" );
-						Log.TraceInformation("FindIncludePaths calls: " + CPPEnvironment.TotalFindIncludedFileCalls + " (" + CPPEnvironment.IncludePathSearchAttempts + " searches)" );
-						Log.TraceInformation("PCH time: " + UEBuildModuleCPP.TotalPCHTime + "s (Shared PCH time: " + UEBuildModuleCPP.TotalSharedPCHTime + "s)" );
-						Log.TraceInformation("Include Resolves: {0} ({1} misses, {2:0.00}%)", CPPEnvironment.TotalDirectIncludeResolves, CPPEnvironment.TotalDirectIncludeResolveCacheMisses, (float)CPPEnvironment.TotalDirectIncludeResolveCacheMisses / (float)CPPEnvironment.TotalDirectIncludeResolves * 100);
-						Log.TraceInformation("Total FileItems: {0} ({1} missing)", FileItem.TotalFileItemCount, FileItem.MissingFileItemCount);
-
-						Log.TraceInformation("Execution time: {0}s", BuildDuration);
-					}
-
-					Telemetry.SendEvent("PerformanceInfo.2",
-						"TotalExecutionTimeSec", BuildDuration.ToString("0.00"),
-						"MutexWaitTimeSec", MutexWaitTime.TotalSeconds.ToString("0.00"),
-						"TotalTimeSpentGettingIncludesSec", CPPEnvironment.TotalTimeSpentGettingIncludes.ToString("0.00"),
-						"TotalIncludesRequested", CPPEnvironment.TotalIncludesRequested.ToString(),
-						"DirectIncludeCacheMissesTotalTimeSec", CPPEnvironment.DirectIncludeCacheMissesTotalTime.ToString("0.00"),
-						"TotalDirectIncludeCacheMisses", CPPEnvironment.TotalDirectIncludeCacheMisses.ToString(),
-						"TotalFindIncludedFileCalls", CPPEnvironment.TotalFindIncludedFileCalls.ToString(),
-						"IncludePathSearchAttempts", CPPEnvironment.IncludePathSearchAttempts.ToString(),
-						"TotalFileItemCount", FileItem.TotalFileItemCount.ToString(),
-						"MissingFileItemCount", FileItem.MissingFileItemCount.ToString()
-						);
-				}
-				catch (Exception Exception)
-				{
-					Log.TraceError("UnrealBuildTool Exception: " + Exception);
-					Result = ECompilationResult.OtherCompilationError;
-				}
-
-				if (bUseMutex)
-				{
-					// Release the mutex to avoid the abandoned mutex timeout.
-					SingleInstanceMutex.ReleaseMutex();
-					SingleInstanceMutex.Dispose();
-					SingleInstanceMutex = null;
-				}
-			}
-
-			Telemetry.Shutdown();
-
-			// Print some performance info
-			Log.TraceVerbose("Execution time: {0}", (DateTime.UtcNow - StartTime - MutexWaitTime).TotalSeconds);
-
-			return (int) Result;
-		}
-
-
-		/// <summary>
-		/// Generates project files.  Can also be used to generate projects "in memory", to allow for building
-		/// targets without even having project files at all.
-		/// </summary>
-		/// <param name="Generator">Project generator to use</param>
-		/// <param name="Arguments">Command-line arguments</param>
-		/// <returns>True if successful</returns>
-		public static bool GenerateProjectFiles(ProjectFileGenerator Generator, string[] Arguments)
-		{
-			bool bSuccess;
-			ProjectFileGenerator.bGenerateProjectFiles = true;
-			Generator.GenerateProjectFiles(Arguments, out bSuccess);
-			ProjectFileGenerator.bGenerateProjectFiles = false;
-			return bSuccess;
-		}
-
-
-
-
-		// @todo: Ideally get rid of RunUBT() and all of the Clear/Reset stuff!
-		public static ECompilationResult RunUBT(string[] Arguments)
-		{
-			bool bSuccess = true;
-
-			// Reset global configurations
-			ResetAllActions();
-
-			ParseBuildConfigurationFlags(Arguments);
-
-			// We need to allow the target platform to perform the 'reset' as well...
-			UnrealTargetPlatform ResetPlatform = UnrealTargetPlatform.Unknown;
-			UnrealTargetConfiguration ResetConfiguration;
-			UEBuildTarget.ParsePlatformAndConfiguration(Arguments, out ResetPlatform, out ResetConfiguration);
-			if (RunningRocket())
-			{
-				// Only Development and Shipping are supported for engine modules
-				if( ResetConfiguration != UnrealTargetConfiguration.Development && ResetConfiguration != UnrealTargetConfiguration.Shipping )
-				{
-					ResetConfiguration = UnrealTargetConfiguration.Development;
-				}
-			}
-			var BuildPlatform = UEBuildPlatform.GetBuildPlatform(ResetPlatform);
-
-			// now that we have the platform, we can set the intermediate path to include the platform/architecture name
-			BuildConfiguration.PlatformIntermediateFolder = Path.Combine(BuildConfiguration.BaseIntermediateFolder, ResetPlatform.ToString(), BuildPlatform.GetActiveArchitecture());
-
-			BuildPlatform.ResetBuildConfiguration(ResetPlatform, ResetConfiguration);
-
-			var Targets = new List<UEBuildTarget>();
-			string ExecutorName = "Unknown";
-			ECompilationResult BuildResult = ECompilationResult.Succeeded;
-
-			var ToolChain = UEToolChain.GetPlatformToolChain(BuildPlatform.GetCPPTargetPlatform(ResetPlatform));
-
-			string EULAViolationWarning = null;
-
-			try
-			{
-				List<string[]> TargetSettings = ParseCommandLineFlags(Arguments);
-
-				int ArgumentIndex;
-				// action graph implies using the dependency resolve cache
-				bool GeneratingActionGraph = Utils.ParseCommandLineFlag(Arguments, "-graph", out ArgumentIndex );
-				if (GeneratingActionGraph)
-				{
-					BuildConfiguration.bUseIncludeDependencyResolveCache = true;
-				}
-
-				bool CreateStub = Utils.ParseCommandLineFlag(Arguments, "-nocreatestub", out ArgumentIndex);
-				if (CreateStub || (String.IsNullOrEmpty(Environment.GetEnvironmentVariable("uebp_LOCAL_ROOT")) && ExternalExecution.GetRuntimePlatform() == UnrealTargetPlatform.Mac))
-				{
-					BuildConfiguration.bCreateStubIPA = false;
-				}
-
-				// Build action lists for all passed in targets.
-				var OutputItemsForAllTargets = new List<FileItem>();
-				foreach (string[] TargetSetting in TargetSettings)
-				{
-					var TargetStartTime = DateTime.UtcNow;
-					var Target = UEBuildTarget.CreateTarget(TargetSetting);
-					if ((Target == null) && (UEBuildConfiguration.bCleanProject))
-					{
-						continue;
-					}
-					Targets.Add(Target);
-
-					// Load the include dependency cache.
-					if (CPPEnvironment.IncludeDependencyCache == null)
-					{
-						CPPEnvironment.IncludeDependencyCache = DependencyCache.Create(
-							DependencyCache.GetDependencyCachePathForTarget(Target)
-							);
-					}
-
-					var TargetOutputItems = new List<FileItem>();
-					BuildResult = Target.Build(ToolChain, TargetOutputItems, out EULAViolationWarning);
-
-					if(BuildResult != ECompilationResult.Succeeded)
-					{
-						break;
-					}
-
-					OutputItemsForAllTargets.AddRange( TargetOutputItems );
-
-					if ( (BuildConfiguration.bXGEExport && UEBuildConfiguration.bGenerateManifest) || (!ProjectFileGenerator.bGenerateProjectFiles && !UEBuildConfiguration.bGenerateManifest && !UEBuildConfiguration.bCleanProject))
-					{
-						// We don't currently support building game targets in rocket.
-						if (UnrealBuildTool.RunningRocket() && Target.Rules != null && Target.Rules.Type == TargetRules.TargetType.Game)
-						{
-							throw new BuildException(
-								"You currently can not build a game target in Rocket.\nTry again in a future release.\nFor now, build and run your editor project."
-							);
-						}
-
-						// Generate an action graph if we were asked to do that.  The graph generation needs access to the include dependency cache, so
-						// we generate it before saving and cleaning that up.
-						if( GeneratingActionGraph )
-						{
-							// The graph generation feature currently only works with a single target at a time.  This is because of how we need access
-							// to include dependencies for the target, but those aren't kept around as we process multiple targets
-							if( TargetSettings.Count != 1 )
-							{
-								throw new BuildException( "ERROR: The '-graph' option only works with a single target at a time" );
-							}
-							LinkActionsAndItems();
-							var ActionsToExecute = AllActions;
-
-							var VisualizationType = ActionGraphVisualizationType.OnlyCPlusPlusFilesAndHeaders;
-							SaveActionGraphVisualization( Path.Combine( BuildConfiguration.BaseIntermediatePath, Target.GetTargetName() + ".gexf" ), Target.GetTargetName(), VisualizationType, ActionsToExecute );
-						}
-
-						// Save the include dependency cache.
-						if (CPPEnvironment.IncludeDependencyCache != null)
-						{
-							CPPEnvironment.IncludeDependencyCache.Save();
-							CPPEnvironment.IncludeDependencyCache = null;
-						}
-					}
-
-					var TargetBuildTime = (DateTime.UtcNow - TargetStartTime).TotalSeconds;
-
-					// Send out telemetry for this target
-					Telemetry.SendEvent("TargetBuildStats.2",
-						"AppName", Target.AppName,
-						"GameName", Target.GameName,
-						"Platform", Target.Platform.ToString(),
-						"Configuration", Target.Configuration.ToString(),
-						"CleanTarget", UEBuildConfiguration.bCleanProject.ToString(),
-						"Monolithic", Target.ShouldCompileMonolithic().ToString(),
-						"CreateDebugInfo", Target.IsCreatingDebugInfo().ToString(),
-						"TargetType", Target.Rules.Type.ToString(),
-						"TargetCreateTimeSec", TargetBuildTime.ToString("0.00")
-						);
-				}
-
-				if (BuildResult == ECompilationResult.Succeeded &&
-					(
-						(BuildConfiguration.bXGEExport && UEBuildConfiguration.bGenerateManifest) ||
-						(!GeneratingActionGraph && !ProjectFileGenerator.bGenerateProjectFiles && !UEBuildConfiguration.bGenerateManifest && !UEBuildConfiguration.bCleanProject && !UEBuildConfiguration.bGenerateExternalFileList)
-					))
-				{
-					// Plan the actions to execute for the build.
-					List<Action> ActionsToExecute = GetActionsToExecute(OutputItemsForAllTargets, Targets);
-
-					// Display some stats to the user.
-					Log.TraceVerbose(
-							"{0} actions, {1} outdated and requested actions",
-							AllActions.Count,
-							ActionsToExecute.Count
-							);
-
-					ToolChain.PreBuildSync();
-
-					// Execute the actions.
-					bSuccess = ExecuteActions(ActionsToExecute, out ExecutorName);
-
-					// if the build succeeded, do any needed syncing
-					if (bSuccess)
-					{
-						foreach (UEBuildTarget Target in Targets)
-						{
-							ToolChain.PostBuildSync(Target);
-						}
-					}
-					else
-					{
-						BuildResult = ECompilationResult.OtherCompilationError;
-					}
-				}
-			}
-			catch (BuildException Exception)
-			{
-				// Output the message only, without the call stack
-				Log.TraceInformation(Exception.Message);
-				BuildResult = ECompilationResult.OtherCompilationError;
-			}
-			catch (Exception Exception)
-			{
-				Log.TraceInformation("ERROR: {0}", Exception);
-				BuildResult = ECompilationResult.OtherCompilationError;
-			}
-
-			if(EULAViolationWarning != null)
-			{
-				Log.TraceWarning("WARNING: {0}", EULAViolationWarning);
-			}
-
-			// Figure out how long we took to execute.
-			double BuildDuration = (DateTime.UtcNow - StartTime - MutexWaitTime).TotalSeconds;
-			if (ExecutorName == "Local")
-			{
-				Log.TraceInformation("Cumulative action seconds ({0} processors): {1:0.00} building projects, {2:0.00} compiling, {3:0.00} creating app bundles, {4:0.00} generating debug info, {5:0.00} linking, {6:0.00} other",
-					System.Environment.ProcessorCount,
-					TotalBuildProjectTime,
-					TotalCompileTime,
-					TotalCreateAppBundleTime,
-					TotalGenerateDebugInfoTime,
-					TotalLinkTime,
-					TotalOtherActionsTime
-				);
-				Telemetry.SendEvent("BuildStatsTotal.2",
-					"ExecutorName", ExecutorName,
-					"TotalUBTWallClockTimeSec", BuildDuration.ToString("0.00"),
-					"TotalBuildProjectThreadTimeSec", TotalBuildProjectTime.ToString("0.00"),
-					"TotalCompileThreadTimeSec", TotalCompileTime.ToString("0.00"),
-					"TotalCreateAppBundleThreadTimeSec", TotalCreateAppBundleTime.ToString("0.00"),
-					"TotalGenerateDebugInfoThreadTimeSec", TotalGenerateDebugInfoTime.ToString("0.00"),
-					"TotalLinkThreadTimeSec", TotalLinkTime.ToString("0.00"),
-					"TotalOtherActionsThreadTimeSec", TotalOtherActionsTime.ToString("0.00")
-					);
-
-				Log.TraceInformation("UBT execution time: {0:0.00} seconds", BuildDuration);
-
-				// reset statistics
-				TotalBuildProjectTime = 0;
-				TotalCompileTime = 0;
-				TotalCreateAppBundleTime = 0;
-				TotalGenerateDebugInfoTime = 0;
-				TotalLinkTime = 0;
-				TotalOtherActionsTime = 0;
-			}
-			else
-			{
-				if (ExecutorName == "XGE")
-				{
-					Log.TraceInformation("XGE execution time: {0:0.00} seconds", BuildDuration);
-				}
-				Telemetry.SendEvent("BuildStatsTotal.2",
-					"ExecutorName", ExecutorName,
-					"TotalUBTWallClockTimeSec", BuildDuration.ToString("0.00")
-					);
-			}
-
-			return BuildResult;
-		}
-
-		private static void ParseBuildConfigurationFlags(string[] Arguments)
-		{
-			int ArgumentIndex = 0;
-
-			// Parse optional command-line flags.
-			if (Utils.ParseCommandLineFlag(Arguments, "-verbose", out ArgumentIndex))
-			{
-				BuildConfiguration.bPrintDebugInfo = true;
-			}
-
-			if (Utils.ParseCommandLineFlag(Arguments, "-noxge", out ArgumentIndex))
-			{
-				BuildConfiguration.bAllowXGE = false;
-			}
-
-			if (Utils.ParseCommandLineFlag(Arguments, "-noxgemonitor", out ArgumentIndex))
-			{
-				BuildConfiguration.bShowXGEMonitor = false;
-			}
-
-			if (Utils.ParseCommandLineFlag(Arguments, "-stresstestunity", out ArgumentIndex))
-			{
-				BuildConfiguration.bStressTestUnity = true;
-			}
-
-			if (Utils.ParseCommandLineFlag(Arguments, "-disableunity", out ArgumentIndex))
-			{
-				BuildConfiguration.bUseUnityBuild = false;
-			}
+                        {
+                            bGenerationSuccess &= GenerateProjectFiles(new MakefileGenerator(), Arguments);
+                        }
+
+                        if(!bGenerationSuccess)
+                        {
+                            Result = ECompilationResult.OtherCompilationError;
+                        }
+                    }
+                    else if (bRunCopyrightVerification)
+                    {
+                        CopyrightVerify.RunCopyrightVerification(EngineSourceDirectory, GameName, bDumpToFile);
+                        Log.TraceInformation("Completed... exiting.");
+                    }
+                    else if (bAutoSDKOnly)
+                    {
+                        // RegisterAllUBTClasses has already done all the SDK setup.
+                        Result = ECompilationResult.Succeeded;
+                    }
+                    else
+                    {
+                        // Check if any third party headers are included from public engine headers.
+                        if (bCheckThirdPartyHeaders)
+                        {
+                            ThirdPartyHeaderFinder.FindThirdPartyIncludes(CheckPlatform, CheckConfiguration);
+                        }
+
+                        // Build our project
+                        if (Result == ECompilationResult.Succeeded)
+                        {
+                            if (UEBuildConfiguration.bPrepForDeployment == false)
+                            {
+                                // If we are only prepping for deployment, assume the build already occurred.
+                                Result = RunUBT(Arguments);
+                            }
+                            else
+                            {
+                                var BuildPlatform = UEBuildPlatform.GetBuildPlatform(CheckPlatform, true);
+                                if (BuildPlatform != null)
+                                {
+                                    // Setup environment wasn't called, so set the flag
+                                    BuildConfiguration.bDeployAfterCompile = BuildPlatform.RequiresDeployPrepAfterCompile();
+                                    BuildConfiguration.PlatformIntermediateFolder = Path.Combine(BuildConfiguration.BaseIntermediateFolder, CheckPlatform.ToString(), BuildPlatform.GetActiveArchitecture());
+                                }
+                            }
+
+                            // If we build w/ bXGEExport true, we didn't REALLY build at this point, 
+                            // so don't bother with doing the PrepTargetForDeployment call. 
+                            if ((Result == ECompilationResult.Succeeded) && (BuildConfiguration.bDeployAfterCompile == true) && (BuildConfiguration.bXGEExport == false) &&
+                                (UEBuildConfiguration.bGenerateManifest == false) && (UEBuildConfiguration.bGenerateExternalFileList == false) && (UEBuildConfiguration.bCleanProject == false))
+                            {
+                                var DeployHandler = UEBuildDeploy.GetBuildDeploy(CheckPlatform);
+                                if (DeployHandler != null)
+                                {
+                                    // We need to be able to identify the Target.Type we can derive it from the Arguments.
+                                    BuildConfiguration.bFlushBuildDirOnRemoteMac = false;
+                                    UEBuildTarget CheckTarget = UEBuildTarget.CreateTarget(Arguments);
+                                    CheckTarget.SetupGlobalEnvironment();
+                                    if ((CheckTarget.Rules.Type == TargetRules.TargetType.Game) || 
+                                        (CheckTarget.Rules.Type == TargetRules.TargetType.Server) || 
+                                        (CheckTarget.Rules.Type == TargetRules.TargetType.Client))
+                                    {
+                                        CheckTarget.AppName = CheckTarget.GameName;
+                                    }
+                                    DeployHandler.PrepTargetForDeployment(CheckTarget);
+                                }
+                            }
+                        }
+                    }
+                    // Print some performance info
+                    var BuildDuration = (DateTime.UtcNow - StartTime - MutexWaitTime).TotalSeconds;
+                    if (BuildConfiguration.bPrintPerformanceInfo)
+                    {
+                        Log.TraceInformation("GetIncludes time: " + CPPEnvironment.TotalTimeSpentGettingIncludes + "s (" + CPPEnvironment.TotalIncludesRequested + " includes)" );
+                        Log.TraceInformation("DirectIncludes cache miss time: " + CPPEnvironment.DirectIncludeCacheMissesTotalTime + "s (" + CPPEnvironment.TotalDirectIncludeCacheMisses + " misses)" );
+                        Log.TraceInformation("FindIncludePaths calls: " + CPPEnvironment.TotalFindIncludedFileCalls + " (" + CPPEnvironment.IncludePathSearchAttempts + " searches)" );
+                        Log.TraceInformation("PCH time: " + UEBuildModuleCPP.TotalPCHTime + "s (Shared PCH time: " + UEBuildModuleCPP.TotalSharedPCHTime + "s)" );
+                        Log.TraceInformation("Include Resolves: {0} ({1} misses, {2:0.00}%)", CPPEnvironment.TotalDirectIncludeResolves, CPPEnvironment.TotalDirectIncludeResolveCacheMisses, (float)CPPEnvironment.TotalDirectIncludeResolveCacheMisses / (float)CPPEnvironment.TotalDirectIncludeResolves * 100);
+                        Log.TraceInformation("Total FileItems: {0} ({1} missing)", FileItem.TotalFileItemCount, FileItem.MissingFileItemCount);
+
+                        Log.TraceInformation("Execution time: {0}s", BuildDuration);
+                    }
+
+                    Telemetry.SendEvent("PerformanceInfo.2",
+                        "TotalExecutionTimeSec", BuildDuration.ToString("0.00"),
+                        "MutexWaitTimeSec", MutexWaitTime.TotalSeconds.ToString("0.00"),
+                        "TotalTimeSpentGettingIncludesSec", CPPEnvironment.TotalTimeSpentGettingIncludes.ToString("0.00"),
+                        "TotalIncludesRequested", CPPEnvironment.TotalIncludesRequested.ToString(),
+                        "DirectIncludeCacheMissesTotalTimeSec", CPPEnvironment.DirectIncludeCacheMissesTotalTime.ToString("0.00"),
+                        "TotalDirectIncludeCacheMisses", CPPEnvironment.TotalDirectIncludeCacheMisses.ToString(),
+                        "TotalFindIncludedFileCalls", CPPEnvironment.TotalFindIncludedFileCalls.ToString(),
+                        "IncludePathSearchAttempts", CPPEnvironment.IncludePathSearchAttempts.ToString(),
+                        "TotalFileItemCount", FileItem.TotalFileItemCount.ToString(),
+                        "MissingFileItemCount", FileItem.MissingFileItemCount.ToString()
+                        );
+                }
+                catch (Exception Exception)
+                {
+                    Log.TraceError("UnrealBuildTool Exception: " + Exception);
+                    Result = ECompilationResult.OtherCompilationError;
+                }
+
+                if (bUseMutex)
+                {
+                    // Release the mutex to avoid the abandoned mutex timeout.
+                    SingleInstanceMutex.ReleaseMutex();
+                    SingleInstanceMutex.Dispose();
+                    SingleInstanceMutex = null;
+                }
+            }
+
+            Telemetry.Shutdown();
+
+            // Print some performance info
+            Log.TraceVerbose("Execution time: {0}", (DateTime.UtcNow - StartTime - MutexWaitTime).TotalSeconds);
+
+            return (int) Result;
+        }
+
+
+        /// <summary>
+        /// Generates project files.  Can also be used to generate projects "in memory", to allow for building
+        /// targets without even having project files at all.
+        /// </summary>
+        /// <param name="Generator">Project generator to use</param>
+        /// <param name="Arguments">Command-line arguments</param>
+        /// <returns>True if successful</returns>
+        public static bool GenerateProjectFiles(ProjectFileGenerator Generator, string[] Arguments)
+        {
+            bool bSuccess;
+            ProjectFileGenerator.bGenerateProjectFiles = true;
+            Generator.GenerateProjectFiles(Arguments, out bSuccess);
+            ProjectFileGenerator.bGenerateProjectFiles = false;
+            return bSuccess;
+        }
+
+
+
+
+        // @todo: Ideally get rid of RunUBT() and all of the Clear/Reset stuff!
+        public static ECompilationResult RunUBT(string[] Arguments)
+        {
+            bool bSuccess = true;
+
+            // Reset global configurations
+            ResetAllActions();
+
+            ParseBuildConfigurationFlags(Arguments);
+
+            // We need to allow the target platform to perform the 'reset' as well...
+            UnrealTargetPlatform ResetPlatform = UnrealTargetPlatform.Unknown;
+            UnrealTargetConfiguration ResetConfiguration;
+            UEBuildTarget.ParsePlatformAndConfiguration(Arguments, out ResetPlatform, out ResetConfiguration);
+            if (RunningRocket())
+            {
+                // Only Development and Shipping are supported for engine modules
+                if( ResetConfiguration != UnrealTargetConfiguration.Development && ResetConfiguration != UnrealTargetConfiguration.Shipping )
+                {
+                    ResetConfiguration = UnrealTargetConfiguration.Development;
+                }
+            }
+            var BuildPlatform = UEBuildPlatform.GetBuildPlatform(ResetPlatform);
+
+            // now that we have the platform, we can set the intermediate path to include the platform/architecture name
+            BuildConfiguration.PlatformIntermediateFolder = Path.Combine(BuildConfiguration.BaseIntermediateFolder, ResetPlatform.ToString(), BuildPlatform.GetActiveArchitecture());
+
+            BuildPlatform.ResetBuildConfiguration(ResetPlatform, ResetConfiguration);
+
+            var Targets = new List<UEBuildTarget>();
+            string ExecutorName = "Unknown";
+            ECompilationResult BuildResult = ECompilationResult.Succeeded;
+
+            var ToolChain = UEToolChain.GetPlatformToolChain(BuildPlatform.GetCPPTargetPlatform(ResetPlatform));
+
+            string EULAViolationWarning = null;
+
+            try
+            {
+                List<string[]> TargetSettings = ParseCommandLineFlags(Arguments);
+
+                int ArgumentIndex;
+                // action graph implies using the dependency resolve cache
+                bool GeneratingActionGraph = Utils.ParseCommandLineFlag(Arguments, "-graph", out ArgumentIndex );
+                if (GeneratingActionGraph)
+                {
+                    BuildConfiguration.bUseIncludeDependencyResolveCache = true;
+                }
+
+                bool CreateStub = Utils.ParseCommandLineFlag(Arguments, "-nocreatestub", out ArgumentIndex);
+                if (CreateStub || (String.IsNullOrEmpty(Environment.GetEnvironmentVariable("uebp_LOCAL_ROOT")) && ExternalExecution.GetRuntimePlatform() == UnrealTargetPlatform.Mac))
+                {
+                    BuildConfiguration.bCreateStubIPA = false;
+                }
+
+                // Build action lists for all passed in targets.
+                var OutputItemsForAllTargets = new List<FileItem>();
+                foreach (string[] TargetSetting in TargetSettings)
+                {
+                    var TargetStartTime = DateTime.UtcNow;
+                    var Target = UEBuildTarget.CreateTarget(TargetSetting);
+                    if ((Target == null) && (UEBuildConfiguration.bCleanProject))
+                    {
+                        continue;
+                    }
+                    Targets.Add(Target);
+
+                    // Load the include dependency cache.
+                    if (CPPEnvironment.IncludeDependencyCache == null)
+                    {
+                        CPPEnvironment.IncludeDependencyCache = DependencyCache.Create(
+                            DependencyCache.GetDependencyCachePathForTarget(Target)
+                            );
+                    }
+
+                    var TargetOutputItems = new List<FileItem>();
+                    BuildResult = Target.Build(ToolChain, TargetOutputItems, out EULAViolationWarning);
+
+                    if(BuildResult != ECompilationResult.Succeeded)
+                    {
+                        break;
+                    }
+
+                    OutputItemsForAllTargets.AddRange( TargetOutputItems );
+
+                    if ( (BuildConfiguration.bXGEExport && UEBuildConfiguration.bGenerateManifest) || (!ProjectFileGenerator.bGenerateProjectFiles && !UEBuildConfiguration.bGenerateManifest && !UEBuildConfiguration.bCleanProject))
+                    {
+                        // We don't currently support building game targets in rocket.
+                        if (UnrealBuildTool.RunningRocket() && Target.Rules != null && Target.Rules.Type == TargetRules.TargetType.Game)
+                        {
+                            throw new BuildException(
+                                "You currently can not build a game target in Rocket.\nTry again in a future release.\nFor now, build and run your editor project."
+                            );
+                        }
+
+                        // Generate an action graph if we were asked to do that.  The graph generation needs access to the include dependency cache, so
+                        // we generate it before saving and cleaning that up.
+                        if( GeneratingActionGraph )
+                        {
+                            // The graph generation feature currently only works with a single target at a time.  This is because of how we need access
+                            // to include dependencies for the target, but those aren't kept around as we process multiple targets
+                            if( TargetSettings.Count != 1 )
+                            {
+                                throw new BuildException( "ERROR: The '-graph' option only works with a single target at a time" );
+                            }
+                            LinkActionsAndItems();
+                            var ActionsToExecute = AllActions;
+
+                            var VisualizationType = ActionGraphVisualizationType.OnlyCPlusPlusFilesAndHeaders;
+                            SaveActionGraphVisualization( Path.Combine( BuildConfiguration.BaseIntermediatePath, Target.GetTargetName() + ".gexf" ), Target.GetTargetName(), VisualizationType, ActionsToExecute );
+                        }
+
+                        // Save the include dependency cache.
+                        if (CPPEnvironment.IncludeDependencyCache != null)
+                        {
+                            CPPEnvironment.IncludeDependencyCache.Save();
+                            CPPEnvironment.IncludeDependencyCache = null;
+                        }
+                    }
+
+                    var TargetBuildTime = (DateTime.UtcNow - TargetStartTime).TotalSeconds;
+
+                    // Send out telemetry for this target
+                    Telemetry.SendEvent("TargetBuildStats.2",
+                        "AppName", Target.AppName,
+                        "GameName", Target.GameName,
+                        "Platform", Target.Platform.ToString(),
+                        "Configuration", Target.Configuration.ToString(),
+                        "CleanTarget", UEBuildConfiguration.bCleanProject.ToString(),
+                        "Monolithic", Target.ShouldCompileMonolithic().ToString(),
+                        "CreateDebugInfo", Target.IsCreatingDebugInfo().ToString(),
+                        "TargetType", Target.Rules.Type.ToString(),
+                        "TargetCreateTimeSec", TargetBuildTime.ToString("0.00")
+                        );
+                }
+
+                if (BuildResult == ECompilationResult.Succeeded &&
+                    (
+                        (BuildConfiguration.bXGEExport && UEBuildConfiguration.bGenerateManifest) ||
+                        (!GeneratingActionGraph && !ProjectFileGenerator.bGenerateProjectFiles && !UEBuildConfiguration.bGenerateManifest && !UEBuildConfiguration.bCleanProject && !UEBuildConfiguration.bGenerateExternalFileList)
+                    ))
+                {
+                    // Plan the actions to execute for the build.
+                    List<Action> ActionsToExecute = GetActionsToExecute(OutputItemsForAllTargets, Targets);
+
+                    // Display some stats to the user.
+                    Log.TraceVerbose(
+                            "{0} actions, {1} outdated and requested actions",
+                            AllActions.Count,
+                            ActionsToExecute.Count
+                            );
+
+                    ToolChain.PreBuildSync();
+
+                    // Execute the actions.
+                    bSuccess = ExecuteActions(ActionsToExecute, out ExecutorName);
+
+                    // if the build succeeded, do any needed syncing
+                    if (bSuccess)
+                    {
+                        foreach (UEBuildTarget Target in Targets)
+                        {
+                            ToolChain.PostBuildSync(Target);
+                        }
+                    }
+                    else
+                    {
+                        BuildResult = ECompilationResult.OtherCompilationError;
+                    }
+                }
+            }
+            catch (BuildException Exception)
+            {
+                // Output the message only, without the call stack
+                Log.TraceInformation(Exception.Message);
+                BuildResult = ECompilationResult.OtherCompilationError;
+            }
+            catch (Exception Exception)
+            {
+                Log.TraceInformation("ERROR: {0}", Exception);
+                BuildResult = ECompilationResult.OtherCompilationError;
+            }
+
+            if(EULAViolationWarning != null)
+            {
+                Log.TraceWarning("WARNING: {0}", EULAViolationWarning);
+            }
+
+            // Figure out how long we took to execute.
+            double BuildDuration = (DateTime.UtcNow - StartTime - MutexWaitTime).TotalSeconds;
+            if (ExecutorName == "Local")
+            {
+                Log.TraceInformation("Cumulative action seconds ({0} processors): {1:0.00} building projects, {2:0.00} compiling, {3:0.00} creating app bundles, {4:0.00} generating debug info, {5:0.00} linking, {6:0.00} other",
+                    System.Environment.ProcessorCount,
+                    TotalBuildProjectTime,
+                    TotalCompileTime,
+                    TotalCreateAppBundleTime,
+                    TotalGenerateDebugInfoTime,
+                    TotalLinkTime,
+                    TotalOtherActionsTime
+                );
+                Telemetry.SendEvent("BuildStatsTotal.2",
+                    "ExecutorName", ExecutorName,
+                    "TotalUBTWallClockTimeSec", BuildDuration.ToString("0.00"),
+                    "TotalBuildProjectThreadTimeSec", TotalBuildProjectTime.ToString("0.00"),
+                    "TotalCompileThreadTimeSec", TotalCompileTime.ToString("0.00"),
+                    "TotalCreateAppBundleThreadTimeSec", TotalCreateAppBundleTime.ToString("0.00"),
+                    "TotalGenerateDebugInfoThreadTimeSec", TotalGenerateDebugInfoTime.ToString("0.00"),
+                    "TotalLinkThreadTimeSec", TotalLinkTime.ToString("0.00"),
+                    "TotalOtherActionsThreadTimeSec", TotalOtherActionsTime.ToString("0.00")
+                    );
+
+                Log.TraceInformation("UBT execution time: {0:0.00} seconds", BuildDuration);
+
+                // reset statistics
+                TotalBuildProjectTime = 0;
+                TotalCompileTime = 0;
+                TotalCreateAppBundleTime = 0;
+                TotalGenerateDebugInfoTime = 0;
+                TotalLinkTime = 0;
+                TotalOtherActionsTime = 0;
+            }
+            else
+            {
+                if (ExecutorName == "XGE")
+                {
+                    Log.TraceInformation("XGE execution time: {0:0.00} seconds", BuildDuration);
+                }
+                Telemetry.SendEvent("BuildStatsTotal.2",
+                    "ExecutorName", ExecutorName,
+                    "TotalUBTWallClockTimeSec", BuildDuration.ToString("0.00")
+                    );
+            }
+
+            return BuildResult;
+        }
+
+        private static void ParseBuildConfigurationFlags(string[] Arguments)
+        {
+            int ArgumentIndex = 0;
+
+            // Parse optional command-line flags.
+            if (Utils.ParseCommandLineFlag(Arguments, "-verbose", out ArgumentIndex))
+            {
+                BuildConfiguration.bPrintDebugInfo = true;
+            }
+
+            if (Utils.ParseCommandLineFlag(Arguments, "-noxge", out ArgumentIndex))
+            {
+                BuildConfiguration.bAllowXGE = false;
+            }
+
+            if (Utils.ParseCommandLineFlag(Arguments, "-noxgemonitor", out ArgumentIndex))
+            {
+                BuildConfiguration.bShowXGEMonitor = false;
+            }
+
+            if (Utils.ParseCommandLineFlag(Arguments, "-stresstestunity", out ArgumentIndex))
+            {
+                BuildConfiguration.bStressTestUnity = true;
+            }
+
+            if (Utils.ParseCommandLineFlag(Arguments, "-disableunity", out ArgumentIndex))
+            {
+                BuildConfiguration.bUseUnityBuild = false;
+            }
             if (Utils.ParseCommandLineFlag(Arguments, "-forceunity", out ArgumentIndex))
             {
                 BuildConfiguration.bForceUnityBuild = true;
@@ -1297,113 +1308,113 @@ namespace UnrealBuildTool
             }
             
             if (Utils.ParseCommandLineFlag(Arguments, "-uniqueintermediate", out ArgumentIndex))
-			{
-				BuildConfiguration.BaseIntermediateFolder = "Intermediate/Build/" + BuildGuid + "/";
-			}
-
-			if (Utils.ParseCommandLineFlag(Arguments, "-nopch", out ArgumentIndex))
-			{
-				BuildConfiguration.bUsePCHFiles = false;
-			}
-
-			if (Utils.ParseCommandLineFlag(Arguments, "-skipActionHistory", out ArgumentIndex))
-			{
-				BuildConfiguration.bUseActionHistory = false;
-			}
-
-			if (Utils.ParseCommandLineFlag(Arguments, "-noLTCG", out ArgumentIndex))
-			{
-				BuildConfiguration.bAllowLTCG = false;
-			}
-
-			if (Utils.ParseCommandLineFlag(Arguments, "-nopdb", out ArgumentIndex))
-			{
-				BuildConfiguration.bUsePDBFiles = false;
-			}
-
-			if (Utils.ParseCommandLineFlag(Arguments, "-deploy", out ArgumentIndex))
-			{
-				BuildConfiguration.bDeployAfterCompile = true;
-			}
+            {
+                BuildConfiguration.BaseIntermediateFolder = "Intermediate/Build/" + BuildGuid + "/";
+            }
+
+            if (Utils.ParseCommandLineFlag(Arguments, "-nopch", out ArgumentIndex))
+            {
+                BuildConfiguration.bUsePCHFiles = false;
+            }
+
+            if (Utils.ParseCommandLineFlag(Arguments, "-skipActionHistory", out ArgumentIndex))
+            {
+                BuildConfiguration.bUseActionHistory = false;
+            }
+
+            if (Utils.ParseCommandLineFlag(Arguments, "-noLTCG", out ArgumentIndex))
+            {
+                BuildConfiguration.bAllowLTCG = false;
+            }
+
+            if (Utils.ParseCommandLineFlag(Arguments, "-nopdb", out ArgumentIndex))
+            {
+                BuildConfiguration.bUsePDBFiles = false;
+            }
+
+            if (Utils.ParseCommandLineFlag(Arguments, "-deploy", out ArgumentIndex))
+            {
+                BuildConfiguration.bDeployAfterCompile = true;
+            }
 
             if (Utils.ParseCommandLineFlag(Arguments, "-CopyAppBundleBackToDevice", out ArgumentIndex))
             {
                 BuildConfiguration.bCopyAppBundleBackToDevice = true;
             }
 
-			if (Utils.ParseCommandLineFlag(Arguments, "-flushmac", out ArgumentIndex))
-			{
-				BuildConfiguration.bFlushBuildDirOnRemoteMac = true;
-			}
-
-			if (Utils.ParseCommandLineFlag(Arguments, "-nodebuginfo", out ArgumentIndex))
-			{
-				BuildConfiguration.bDisableDebugInfo = true;
-			}
-			else if (Utils.ParseCommandLineFlag(Arguments, "-forcedebuginfo", out ArgumentIndex))
-			{
-				BuildConfiguration.bDisableDebugInfo = false;
-				BuildConfiguration.bOmitPCDebugInfoInDevelopment = false;
-			}
-
-			foreach(string Argument in Arguments)
-			{
-				const string WriteTargetInfoArg = "-writetargetinfo=";
-				if (Argument.ToLowerInvariant().StartsWith(WriteTargetInfoArg))
-				{
-					BuildConfiguration.WriteTargetInfoPath = Argument.Substring(WriteTargetInfoArg.Length);
-					break;
-				}
-			}
-		}
-
-		/**
-		 * Parses the passed in command line for build configuration overrides.
-		 * 
-		 * @param	Arguments	List of arguments to parse
-		 * @return	List of build target settings
-		 */
-		private static List<string[]> ParseCommandLineFlags(string[] Arguments)
-		{
-			var TargetSettings = new List<string[]>();
-			int ArgumentIndex = 0;
-
-			// Log command-line arguments.
-			if (BuildConfiguration.bPrintDebugInfo)
-			{
-				Console.Write("Command-line arguments: ");
-				foreach (string Argument in Arguments)
-				{
-					Console.Write("{0} ", Argument);
-				}
-				Log.TraceInformation("");
-			}
-
-			ParseBuildConfigurationFlags(Arguments);
-
-			if (Utils.ParseCommandLineFlag(Arguments, "-targets", out ArgumentIndex))
-			{
-				if (ArgumentIndex + 1 >= Arguments.Length)
-				{
-					throw new BuildException("Expected filename after -targets argument, but found nothing.");
-				}
-				// Parse lines from the referenced file into target settings.
-				var Lines = File.ReadAllLines(Arguments[ArgumentIndex + 1]);
-				foreach (string Line in Lines)
-				{
-					if (Line != "" && Line[0] != ';')
-					{
-						TargetSettings.Add(Line.Split(' '));
-					}
-				}
-			}
-			// Simply use full command line arguments as target setting if not otherwise overriden.
-			else
-			{
-				TargetSettings.Add(Arguments);
-			}
-
-			return TargetSettings;
-		}
-	}
+            if (Utils.ParseCommandLineFlag(Arguments, "-flushmac", out ArgumentIndex))
+            {
+                BuildConfiguration.bFlushBuildDirOnRemoteMac = true;
+            }
+
+            if (Utils.ParseCommandLineFlag(Arguments, "-nodebuginfo", out ArgumentIndex))
+            {
+                BuildConfiguration.bDisableDebugInfo = true;
+            }
+            else if (Utils.ParseCommandLineFlag(Arguments, "-forcedebuginfo", out ArgumentIndex))
+            {
+                BuildConfiguration.bDisableDebugInfo = false;
+                BuildConfiguration.bOmitPCDebugInfoInDevelopment = false;
+            }
+
+            foreach(string Argument in Arguments)
+            {
+                const string WriteTargetInfoArg = "-writetargetinfo=";
+                if (Argument.ToLowerInvariant().StartsWith(WriteTargetInfoArg))
+                {
+                    BuildConfiguration.WriteTargetInfoPath = Argument.Substring(WriteTargetInfoArg.Length);
+                    break;
+                }
+            }
+        }
+
+        /**
+         * Parses the passed in command line for build configuration overrides.
+         * 
+         * @param	Arguments	List of arguments to parse
+         * @return	List of build target settings
+         */
+        private static List<string[]> ParseCommandLineFlags(string[] Arguments)
+        {
+            var TargetSettings = new List<string[]>();
+            int ArgumentIndex = 0;
+
+            // Log command-line arguments.
+            if (BuildConfiguration.bPrintDebugInfo)
+            {
+                Console.Write("Command-line arguments: ");
+                foreach (string Argument in Arguments)
+                {
+                    Console.Write("{0} ", Argument);
+                }
+                Log.TraceInformation("");
+            }
+
+            ParseBuildConfigurationFlags(Arguments);
+
+            if (Utils.ParseCommandLineFlag(Arguments, "-targets", out ArgumentIndex))
+            {
+                if (ArgumentIndex + 1 >= Arguments.Length)
+                {
+                    throw new BuildException("Expected filename after -targets argument, but found nothing.");
+                }
+                // Parse lines from the referenced file into target settings.
+                var Lines = File.ReadAllLines(Arguments[ArgumentIndex + 1]);
+                foreach (string Line in Lines)
+                {
+                    if (Line != "" && Line[0] != ';')
+                    {
+                        TargetSettings.Add(Line.Split(' '));
+                    }
+                }
+            }
+            // Simply use full command line arguments as target setting if not otherwise overriden.
+            else
+            {
+                TargetSettings.Add(Arguments);
+            }
+
+            return TargetSettings;
+        }
+    }
 }
diff --git a/Engine/Source/Runtime/Core/Private/Containers/String.cpp b/Engine/Source/Runtime/Core/Private/Containers/String.cpp
--- a/Engine/Source/Runtime/Core/Private/Containers/String.cpp
+++ b/Engine/Source/Runtime/Core/Private/Containers/String.cpp
@@ -743,6 +743,26 @@ int32 FString::ParseIntoArrayWS( TArray<FString>* InArray, const TCHAR* pchExtra
 		WhiteSpace[NumWhiteSpaces++] = pchExtraDelim;
 	}
 
+	return ParseIntoArray(InArray, WhiteSpace, NumWhiteSpaces);
+}
+
+int32 FString::ParseIntoArrayLines(TArray<FString>* InArray) const
+{
+	// default array of LineEndings
+	static const TCHAR* LineEndings[] =
+	{				
+		TEXT("\r"),
+		TEXT("\n"),	
+	};
+
+	// start with just the standard line endings
+	int32 NumLineEndings = ARRAY_COUNT(LineEndings);	
+	return ParseIntoArray(InArray, LineEndings, NumLineEndings);
+}
+
+int32 FString::ParseIntoArray(TArray<FString>* InArray, const TCHAR** DelimArray, int32 NumDelims) const
+{
+	check(DelimArray);
 	check(InArray);
 	InArray->Empty();
 
@@ -754,35 +774,35 @@ int32 FString::ParseIntoArrayWS( TArray<FString>* InArray, const TCHAR* pchExtra
 	while (!bStop)
 	{
 		// skip over any white space at the beginning of the string
-		SkipOver(WhiteSpace, NumWhiteSpaces, S);
-		
+		SkipOver(DelimArray, NumDelims, S);
+
 		// find the first token in the string, and if we get one, add it to the output array of tokens
 		FString Token;
 		TCHAR ch;
-		if (SplitOn(WhiteSpace, NumWhiteSpaces, Token, S, ch))
+		if (SplitOn(DelimArray, NumDelims, Token, S, ch))
 		{
-			if( Token[0] == TEXT('"') )
+			if (Token[0] == TEXT('"'))
 			{
 				int32 SaveSz = Token.Len();
 
-				FString Wk = FString::Printf( TEXT("%s%c"), *Token, ch );
-				for( int32 x = 1 ; x < S.Len() ; ++x )
+				FString Wk = FString::Printf(TEXT("%s%c"), *Token, ch);
+				for (int32 x = 1; x < S.Len(); ++x)
 				{
-					if( S[x] == TEXT('"') )
+					if (S[x] == TEXT('"'))
 					{
 						Wk += TEXT("\"");
 						break;
 					}
 					else
 					{
-						Wk = Wk + S.Mid(x,1);
+						Wk = Wk + S.Mid(x, 1);
 					}
 				}
 
 				Token = Wk;
 
 				int32 DiffSz = Token.Len() - SaveSz;
-				S = S.Mid( DiffSz );
+				S = S.Mid(DiffSz);
 			}
 
 			// stick it on the end
diff --git a/Engine/Source/Runtime/Core/Private/GenericPlatform/GenericPlatformMisc.cpp b/Engine/Source/Runtime/Core/Private/GenericPlatform/GenericPlatformMisc.cpp
--- a/Engine/Source/Runtime/Core/Private/GenericPlatform/GenericPlatformMisc.cpp
+++ b/Engine/Source/Runtime/Core/Private/GenericPlatform/GenericPlatformMisc.cpp
@@ -161,6 +161,16 @@ GenericApplication* FGenericPlatformMisc::CreateApplication()
 	return new GenericApplication( nullptr );
 }
 
+void FGenericPlatformMisc::SetEnvironmentVar(const TCHAR* VariableName, const TCHAR* Value)
+{
+	UE_LOG(LogGenericPlatformMisc, Error, TEXT("SetEnvironmentVar not implemented for this platform: %s = %s"), VariableName, Value);
+}
+
+const TCHAR* FGenericPlatformMisc::GetPathVarDelimiter()
+{
+	return TEXT(";");
+}
+
 TArray<uint8> FGenericPlatformMisc::GetMacAddress()
 {
 	return TArray<uint8>();
diff --git a/Engine/Source/Runtime/Core/Private/Modules/ModuleManager.cpp b/Engine/Source/Runtime/Core/Private/Modules/ModuleManager.cpp
--- a/Engine/Source/Runtime/Core/Private/Modules/ModuleManager.cpp
+++ b/Engine/Source/Runtime/Core/Private/Modules/ModuleManager.cpp
@@ -1363,7 +1363,7 @@
         }
 	}
 
-	if( InvokeUnrealBuildTool( CmdLineParams, Ar ) )
+	if( InvokeUnrealBuildToolForCompile( CmdLineParams, Ar ) )
 	{
 		const bool bWaitForCompletion = true;	// Always wait
 		bool bCompileStillInProgress = false;
@@ -1382,31 +1382,6 @@
 	return bCompileSucceeded;
 }
 
-
-bool FModuleManager::IsUnrealBuildToolAvailable()
-{
-	// If using Rocket and the Rocket unreal build tool executable exists, then UBT is available
-	if ( FApp::IsEngineInstalled() )
-	{
-		return FPaths::FileExists(GetUnrealBuildToolExecutableFilename());
-	}
-	else
-	{
-		// If not using Rocket, check to make sure UBT can be built, since it is an intermediate.
-		// We are simply checking for the existence of source code files, which is a heuristic to determine if UBT can be built
-		TArray<FString> Filenames;
-		const FString UBTSourcePath = GetUnrealBuildToolSourceCodePath();
-		const FString SearchPattern = TEXT("*.cs");
-		const bool bFiles = true;
-		const bool bDirectories = false;
-		const bool bClearFileNames = false;
-		IFileManager::Get().FindFilesRecursive(Filenames, *UBTSourcePath, *SearchPattern, bFiles, bDirectories, bClearFileNames);
-
-		return Filenames.Num() > 0;
-	}
-}
-
-
 void FModuleManager::UnloadOrAbandonModuleWithCallback( const FName InModuleName, FOutputDevice &Ar )
 {
 	TSharedRef< FModuleInfo > Module = Modules.FindChecked( InModuleName );
@@ -1532,7 +1507,7 @@
 		BuildPlatformName, BuildConfigurationName, 
 		*ExtraArg, *InAdditionalCmdLineArgs );
 
-	const bool bInvocationSuccessful = InvokeUnrealBuildTool(CmdLineParams, Ar);
+	const bool bInvocationSuccessful = InvokeUnrealBuildToolForCompile(CmdLineParams, Ar);
 	if ( !bInvocationSuccessful )
 	{
 		// No longer compiling modules
@@ -1552,163 +1527,12 @@
 #endif
 }
 
-
-FString FModuleManager::GetUnrealBuildToolSourceCodePath()
+bool FModuleManager::InvokeUnrealBuildToolForCompile(const FString& InCmdLineParams, FOutputDevice &Ar)
 {
-	return FPaths::Combine(*FPaths::EngineDir(), TEXT("Source"), TEXT("Programs"), TEXT("UnrealBuildTool"));
-}
-
-
-FString FModuleManager::GetUnrealBuildToolExecutableFilename()
-{
-	const FString UBTPath = FString::Printf(TEXT("%sBinaries/DotNET/"), *FPaths::EngineDir());
-	const FString UBTExe = TEXT("UnrealBuildTool.exe");
-
-	FString ExecutableFileName = UBTPath / UBTExe;
-	ExecutableFileName = FPaths::ConvertRelativePathToFull(ExecutableFileName);
-	return ExecutableFileName;
-}
-
-
-bool FModuleManager::BuildUnrealBuildTool(FOutputDevice &Ar)
-{
-#if !IS_MONOLITHIC
-	if ( FApp::IsEngineInstalled() )
-	{
-		// We may not build UBT in rocket
-		return false;
-	}
-
-	FString CompilerExecutableFilename;
-	FString CmdLineParams;
-#if PLATFORM_WINDOWS
-	// To build UBT for windows, we must assemble a batch file that first registers the environment variable necessary to run msbuild then run it
-	// This can not be done in a single invocation of CMD.exe because the environment variables do not transfer between subsequent commands when using the "&" syntax
-	// devenv.exe can be used to build as well but it takes several seconds to start up so it is not desirable
-
-	// First determine the appropriate vcvars batch file to launch
-	FString VCVarsBat;
-
-#if _MSC_VER >= 1800
-	FPlatformMisc::GetVSComnTools(12, VCVarsBat);
-#else
-	FPlatformMisc::GetVSComnTools(11, VCVarsBat);
-#endif
-
-	VCVarsBat = FPaths::Combine(*VCVarsBat, L"../../VC/bin/x86_amd64/vcvarsx86_amd64.bat");
-
-	// Check to make sure we found one.
-	if ( VCVarsBat.IsEmpty() || !FPaths::FileExists(VCVarsBat) )
-	{
-		// VCVars doesn't exist, we can not build UBT
-		CompilerExecutableFilename = TEXT("");
-	}
-	else
-	{
-		// Now make a batch file in the intermediate directory to invoke the vcvars batch then msbuild
-		FString BuildBatchFile = FPaths::EngineIntermediateDir() / TEXT("Build") / TEXT("UnrealBuildTool") / TEXT("BuildUBT.bat");
-		BuildBatchFile.ReplaceInline(TEXT("/"), TEXT("\\"));
-
-		const FString CsProjLocation = FPaths::ConvertRelativePathToFull(GetUnrealBuildToolSourceCodePath()) / TEXT("UnrealBuildTool.csproj");
-		
-		FString BatchFileContents;
-		BatchFileContents = FString::Printf(TEXT("call \"%s\"") LINE_TERMINATOR, *VCVarsBat);
-		BatchFileContents += FString::Printf(TEXT("msbuild /nologo /verbosity:quiet %s /property:Configuration=Development /property:Platform=AnyCPU"), *CsProjLocation);
-		FFileHelper::SaveStringToFile(BatchFileContents, *BuildBatchFile);
-
-		{
-			TCHAR CmdExePath[MAX_PATH];
-			FPlatformMisc::GetEnvironmentVariable(TEXT("ComSpec"), CmdExePath, ARRAY_COUNT(CmdExePath));
-			CompilerExecutableFilename = CmdExePath;
-		}
-
-		CmdLineParams = TEXT("/c ");
-		CmdLineParams += BuildBatchFile;
-	}
-#elif PLATFORM_MAC
-	const FString CsProjLocation = FPaths::ConvertRelativePathToFull(GetUnrealBuildToolSourceCodePath()) / TEXT("UnrealBuildTool_Mono.csproj");
-	FString ScriptPath = FPaths::ConvertRelativePathToFull(FPaths::EngineDir() / TEXT("Build/BatchFiles/Mac/RunXBuild.sh"));
-	CompilerExecutableFilename = TEXT("/bin/sh");
-	CmdLineParams = FString::Printf(TEXT("\"%s\" /property:Configuration=Development %s"), *ScriptPath, *CsProjLocation);
-#else
-	Ar.Log( TEXT( "Unknown platform, unable to build UnrealBuildTool." ) );
-#endif
-
-	// If a compiler executable was provided, try to build now
-	if ( !CompilerExecutableFilename.IsEmpty() )
-	{
-		const bool bLaunchDetached = false;
-		const bool bLaunchHidden = true;
-		const bool bLaunchReallyHidden = bLaunchHidden;
-		FProcHandle ProcHandle = FPlatformProcess::CreateProc( *CompilerExecutableFilename, *CmdLineParams, bLaunchDetached, bLaunchHidden, bLaunchReallyHidden, NULL, 0, NULL, NULL );
-		if ( ProcHandle.IsValid() )
-		{
-			FPlatformProcess::WaitForProc( ProcHandle );
-			ProcHandle.Close();
-		}
-
-		// If the executable appeared where we expect it, then we were successful
-		return FPaths::FileExists(GetUnrealBuildToolExecutableFilename());
-	}
-	else
-#endif // !IS_MONOLITHIC
-	{
-		return false;
-	}
-}
-
-
-bool FModuleManager::InvokeUnrealBuildTool( const FString& InCmdLineParams, FOutputDevice &Ar )
-{
 #if PLATFORM_DESKTOP && !IS_MONOLITHIC
-	FString CmdLineParams = InCmdLineParams;
 
-	if ( FRocketSupport::IsRocket() )
-	{
-		CmdLineParams += TEXT(" -rocket");
-	}
-
 	// Make sure we're not already compiling something!
-	check( !IsCurrentlyCompiling() );
+	check(!IsCurrentlyCompiling());
 
-	// UnrealBuildTool is currently always located in the Binaries/DotNET folder
-	FString ExecutableFileName = GetUnrealBuildToolExecutableFilename();
-
-	// Rocket never builds UBT, UnrealBuildTool should already exist
-	if ( !FApp::IsEngineInstalled() )
-	{
-		// When not using rocket, we should attempt to build UBT to make sure it is up to date
-		// Only do this if we have not already successfully done it once during this session.
-		static bool bSuccessfullyBuiltUBTOnce = false;
-		if ( !bSuccessfullyBuiltUBTOnce )
-		{
-			Ar.Log( TEXT( "Building UnrealBuildTool..." ) );
-			if ( BuildUnrealBuildTool(Ar) )
-			{
-				bSuccessfullyBuiltUBTOnce = true;
-			}
-			else
-			{
-				// Failed to build UBT
-				Ar.Log( TEXT( "Failed to build UnrealBuildTool." ) );
-				return false;
-			}
-		}
-	}
-
-	Ar.Logf( TEXT( "Launching UnrealBuildTool... [%s %s]" ), *ExecutableFileName, *CmdLineParams );
-
-#if PLATFORM_MAC
-	// On Mac we launch UBT with Mono
-	FString ScriptPath = FPaths::ConvertRelativePathToFull(FPaths::EngineDir() / TEXT("Build/BatchFiles/Mac/RunMono.sh"));
-	CmdLineParams = FString::Printf(TEXT("\"%s\" \"%s\" %s"), *ScriptPath, *ExecutableFileName, *CmdLineParams);
-	ExecutableFileName = TEXT( "/bin/sh" );
-#endif
-
-	// Run UnrealBuildTool
-	const bool bLaunchDetached = false;
-	const bool bLaunchHidden = true;
-	const bool bLaunchReallyHidden = bLaunchHidden;
-
 	// Setup output redirection pipes, so that we can harvest compiler output and display it ourselves
 #if PLATFORM_LINUX
@@ -1724,11 +1548,18 @@
 	verify(FPlatformProcess::CreatePipe(PipeRead, PipeWrite));
 	ModuleCompileReadPipeText = TEXT("");
 
-	FProcHandle ProcHandle = FPlatformProcess::CreateProc( *ExecutableFileName, *CmdLineParams, bLaunchDetached, bLaunchHidden, bLaunchReallyHidden, NULL, 0, NULL, PipeWrite );
-	if( !ProcHandle.IsValid() )
+	FProcHandle ProcHandle = FUBTInvoker::InvokeUnrealBuildToolAsync(InCmdLineParams, Ar, PipeRead, PipeWrite);
+
+	// We no longer need the Write pipe so close it.
+	// We DO need the Read pipe however...
+#if PLATFORM_LINUX
+	close(*(int*)PipeWrite);
+#else
+	FPlatformProcess::ClosePipe(0, PipeWrite);
+#endif
+
+	if (!ProcHandle.IsValid())
 	{
-		Ar.Logf( TEXT( "Failed to launch Unreal Build Tool. (%s)" ), *ExecutableFileName );
-
 		// We're done with the process handle now
 		ModuleCompileProcessHandle.Reset();
 		ModuleCompileReadPipe = NULL;
@@ -1739,14 +1570,6 @@
 		ModuleCompileReadPipe = PipeRead;
 	}
 
-	// We no longer need the Write pipe so close it.
-	// We DO need the Read pipe however...
-#if PLATFORM_LINUX
-	close(*(int*)PipeWrite);
-#else
-	FPlatformProcess::ClosePipe(0, PipeWrite);
-#endif
-
 	return ProcHandle.IsValid();
 #else
 	return false;
@@ -1980,3 +1803,229 @@
 
 	return false;
 }
+
+bool FModuleManager::IsUnrealBuildToolAvailable()
+{
+	// If using Rocket and the Rocket unreal build tool executable exists, then UBT is available
+	if (FApp::IsEngineInstalled())
+	{
+		return FPaths::FileExists(FUBTInvoker::GetUnrealBuildToolExecutableFilename());
+	}
+	else
+	{
+		// If not using Rocket, check to make sure UBT can be built, since it is an intermediate.
+		// We are simply checking for the existence of source code files, which is a heuristic to determine if UBT can be built
+		TArray<FString> Filenames;
+		const FString UBTSourcePath = FUBTInvoker::GetUnrealBuildToolSourceCodePath();
+		const FString SearchPattern = TEXT("*.cs");
+		const bool bFiles = true;
+		const bool bDirectories = false;
+		const bool bClearFileNames = false;
+		IFileManager::Get().FindFilesRecursive(Filenames, *UBTSourcePath, *SearchPattern, bFiles, bDirectories, bClearFileNames);
+
+		return Filenames.Num() > 0;
+	}
+}
+
+FString FUBTInvoker::GetUnrealBuildToolExecutableFilename()
+{
+	const FString UBTPath = FString::Printf(TEXT("%sBinaries/DotNET/"), *FPaths::EngineDir());
+	const FString UBTExe = TEXT("UnrealBuildTool.exe");
+
+	FString ExecutableFileName = UBTPath / UBTExe;
+	ExecutableFileName = FPaths::ConvertRelativePathToFull(ExecutableFileName);
+	return ExecutableFileName;
+}
+
+
+bool FUBTInvoker::BuildUnrealBuildTool(FOutputDevice &Ar)
+{
+#if !IS_MONOLITHIC
+	if (FApp::IsEngineInstalled())
+	{
+		// We may not build UBT in rocket
+		return false;
+	}
+
+	FString CompilerExecutableFilename;
+	FString CmdLineParams;
+#if PLATFORM_WINDOWS
+	// To build UBT for windows, we must assemble a batch file that first registers the environment variable necessary to run msbuild then run it
+	// This can not be done in a single invocation of CMD.exe because the environment variables do not transfer between subsequent commands when using the "&" syntax
+	// devenv.exe can be used to build as well but it takes several seconds to start up so it is not desirable
+
+	// First determine the appropriate vcvars batch file to launch
+	FString VCVarsBat;
+
+#if _MSC_VER >= 1800
+	FPlatformMisc::GetVSComnTools(12, VCVarsBat);
+#else
+	FPlatformMisc::GetVSComnTools(11, VCVarsBat);
+#endif
+
+	VCVarsBat = FPaths::Combine(*VCVarsBat, L"../../VC/bin/x86_amd64/vcvarsx86_amd64.bat");
+
+	// Check to make sure we found one.
+	if (VCVarsBat.IsEmpty() || !FPaths::FileExists(VCVarsBat))
+	{
+		// VCVars doesn't exist, we can not build UBT
+		CompilerExecutableFilename = TEXT("");
+	}
+	else
+	{
+		// Now make a batch file in the intermediate directory to invoke the vcvars batch then msbuild
+		FString BuildBatchFile = FPaths::EngineIntermediateDir() / TEXT("Build") / TEXT("UnrealBuildTool") / TEXT("BuildUBT.bat");
+		BuildBatchFile.ReplaceInline(TEXT("/"), TEXT("\\"));
+
+		const FString CsProjLocation = FPaths::ConvertRelativePathToFull(GetUnrealBuildToolSourceCodePath()) / TEXT("UnrealBuildTool.csproj");
+
+		FString BatchFileContents;
+		BatchFileContents = FString::Printf(TEXT("call \"%s\"") LINE_TERMINATOR, *VCVarsBat);
+		BatchFileContents += FString::Printf(TEXT("msbuild /nologo /verbosity:quiet %s /property:Configuration=Development /property:Platform=AnyCPU"), *CsProjLocation);
+		FFileHelper::SaveStringToFile(BatchFileContents, *BuildBatchFile);
+
+		{
+			TCHAR CmdExePath[MAX_PATH];
+			FPlatformMisc::GetEnvironmentVariable(TEXT("ComSpec"), CmdExePath, ARRAY_COUNT(CmdExePath));
+			CompilerExecutableFilename = CmdExePath;
+		}
+
+		CmdLineParams = TEXT("/c ");
+		CmdLineParams += BuildBatchFile;
+	}
+#elif PLATFORM_MAC
+	const FString CsProjLocation = FPaths::ConvertRelativePathToFull(GetUnrealBuildToolSourceCodePath()) / TEXT("UnrealBuildTool_Mono.csproj");
+	FString ScriptPath = FPaths::ConvertRelativePathToFull(FPaths::EngineDir() / TEXT("Build/BatchFiles/Mac/RunXBuild.sh"));
+	CompilerExecutableFilename = TEXT("/bin/sh");
+	CmdLineParams = FString::Printf(TEXT("\"%s\" /property:Configuration=Development %s"), *ScriptPath, *CsProjLocation);
+#else
+	Ar.Log(TEXT("Unknown platform, unable to build UnrealBuildTool."));
+#endif
+
+	// If a compiler executable was provided, try to build now
+	if (!CompilerExecutableFilename.IsEmpty())
+	{
+		const bool bLaunchDetached = false;
+		const bool bLaunchHidden = true;
+		const bool bLaunchReallyHidden = bLaunchHidden;
+		FProcHandle ProcHandle = FPlatformProcess::CreateProc(*CompilerExecutableFilename, *CmdLineParams, bLaunchDetached, bLaunchHidden, bLaunchReallyHidden, NULL, 0, NULL, NULL);
+		if (ProcHandle.IsValid())
+		{
+			FPlatformProcess::WaitForProc(ProcHandle);
+			ProcHandle.Close();
+		}
+
+		// If the executable appeared where we expect it, then we were successful
+		return FPaths::FileExists(GetUnrealBuildToolExecutableFilename());
+	}
+	else
+#endif // !IS_MONOLITHIC
+	{
+		return false;
+	}
+}
+
+bool FUBTInvoker::InvokeUnrealBuildToolSync(const FString& InCmdLineParams, FOutputDevice &Ar, bool bSkipBuildUBT, int32& OutReturnCode, FString& OutProcOutput)
+{
+	// Setup output redirection pipes, so that we can harvest compiler output and display it ourselves
+#if PLATFORM_LINUX
+	int pipefd[2];
+	pipe(pipefd);
+	void* PipeRead = &pipefd[0];
+	void* PipeWrite = &pipefd[1];
+#else
+	void* PipeRead = NULL;
+	void* PipeWrite = NULL;
+#endif
+
+	verify(FPlatformProcess::CreatePipe(PipeRead, PipeWrite));
+
+	bool bInvoked = false;
+	FProcHandle ProcHandle = InvokeUnrealBuildToolAsync(InCmdLineParams, Ar, PipeRead, PipeWrite, bSkipBuildUBT);
+	if (ProcHandle.IsValid())
+	{
+		FPlatformProcess::WaitForProc(ProcHandle);
+		bInvoked = true;
+		bool bGotReturnCode = FPlatformProcess::GetProcReturnCode(ProcHandle, &OutReturnCode);
+		OutProcOutput = FPlatformProcess::ReadPipe(PipeRead);
+		check(bGotReturnCode);
+	}
+	else
+	{
+		bInvoked = false;
+		OutReturnCode = -1;
+		OutProcOutput = TEXT("");
+	}
+
+
+	FPlatformProcess::ClosePipe(PipeRead, PipeWrite);
+
+	return bInvoked;
+}
+
+FProcHandle FUBTInvoker::InvokeUnrealBuildToolAsync(const FString& InCmdLineParams, FOutputDevice &Ar, void*& OutReadPipe, void*& OutWritePipe, bool bSkipBuildUBT)
+{
+#if PLATFORM_DESKTOP && !IS_MONOLITHIC
+	FString CmdLineParams = InCmdLineParams;
+
+	if (FRocketSupport::IsRocket())
+	{
+		CmdLineParams += TEXT(" -rocket");
+	}
+
+	// UnrealBuildTool is currently always located in the Binaries/DotNET folder
+	FString ExecutableFileName = GetUnrealBuildToolExecutableFilename();
+
+	// Rocket never builds UBT, UnrealBuildTool should already exist
+	bool bSkipBuild = FApp::IsEngineInstalled() || bSkipBuildUBT;
+	if (!bSkipBuild)
+	{
+		// When not using rocket, we should attempt to build UBT to make sure it is up to date
+		// Only do this if we have not already successfully done it once during this session.
+		static bool bSuccessfullyBuiltUBTOnce = false;
+		if (!bSuccessfullyBuiltUBTOnce)
+		{
+			Ar.Log(TEXT("Building UnrealBuildTool..."));
+			if (BuildUnrealBuildTool(Ar))
+			{
+				bSuccessfullyBuiltUBTOnce = true;
+			}
+			else
+			{
+				// Failed to build UBT
+				Ar.Log(TEXT("Failed to build UnrealBuildTool."));
+				return FProcHandle();
+			}
+		}
+	}
+
+	Ar.Logf(TEXT("Launching UnrealBuildTool... [%s %s]"), *ExecutableFileName, *CmdLineParams);
+
+#if PLATFORM_MAC
+	// On Mac we launch UBT with Mono
+	FString ScriptPath = FPaths::ConvertRelativePathToFull(FPaths::EngineDir() / TEXT("Build/BatchFiles/Mac/RunMono.sh"));
+	CmdLineParams = FString::Printf(TEXT("\"%s\" \"%s\" %s"), *ScriptPath, *ExecutableFileName, *CmdLineParams);
+	ExecutableFileName = TEXT("/bin/sh");
+#endif
+
+	// Run UnrealBuildTool
+	const bool bLaunchDetached = false;
+	const bool bLaunchHidden = true;
+	const bool bLaunchReallyHidden = bLaunchHidden;
+
+	FProcHandle ProcHandle = FPlatformProcess::CreateProc(*ExecutableFileName, *CmdLineParams, bLaunchDetached, bLaunchHidden, bLaunchReallyHidden, NULL, 0, NULL, OutWritePipe);
+	if (!ProcHandle.IsValid())
+	{
+		Ar.Logf(TEXT("Failed to launch Unreal Build Tool. (%s)"), *ExecutableFileName);
+	}
+
+	return ProcHandle;
+#else
+	return FProcHandle();
+#endif // PLATFORM_DESKTOP && !IS_MONOLITHIC
+}
+
+FString FUBTInvoker::GetUnrealBuildToolSourceCodePath()
+{
+	return FPaths::Combine(*FPaths::EngineDir(), TEXT("Source"), TEXT("Programs"), TEXT("UnrealBuildTool"));
+}
diff --git a/Engine/Source/Runtime/Core/Private/Windows/WindowsPlatformMisc.cpp b/Engine/Source/Runtime/Core/Private/Windows/WindowsPlatformMisc.cpp
--- a/Engine/Source/Runtime/Core/Private/Windows/WindowsPlatformMisc.cpp
+++ b/Engine/Source/Runtime/Core/Private/Windows/WindowsPlatformMisc.cpp
@@ -481,12 +481,20 @@ void FWindowsPlatformMisc::GetEnvironmentVariable(const TCHAR* VariableName, TCH
 {
 	uint32 Error = ::GetEnvironmentVariable(VariableName, Result, ResultLength);
 	if (Error <= 0)
-	{
-		// on error, just return an empty string
+	{		
 		*Result = 0;
 	}
 }
 
+void FWindowsPlatformMisc::SetEnvironmentVar(const TCHAR* VariableName, const TCHAR* Value)
+{
+	uint32 Error = ::SetEnvironmentVariable(VariableName, Value);
+	if (Error == 0)
+	{
+		UE_LOG(LogWindows, Warning, TEXT("Failed to set EnvironmentVariable: %s to : %s"), VariableName, Value);
+	}
+}
+
 TArray<uint8> FWindowsPlatformMisc::GetMacAddress()
 {
 	TArray<uint8> Result;
diff --git a/Engine/Source/Runtime/Core/Public/Android/AndroidProperties.h b/Engine/Source/Runtime/Core/Public/Android/AndroidProperties.h
--- a/Engine/Source/Runtime/Core/Public/Android/AndroidProperties.h
+++ b/Engine/Source/Runtime/Core/Public/Android/AndroidProperties.h
@@ -58,6 +58,11 @@ struct FAndroidPlatformProperties
 		return (BuildTarget == EBuildTargets::Game);
 	}
 
+	static FORCEINLINE bool SupportsAutoSDK()
+	{
+		return true;
+	}
+
 	static FORCEINLINE bool SupportsHighQualityLightmaps()
 	{
 		return false;
diff --git a/Engine/Source/Runtime/Core/Public/Containers/UnrealString.h b/Engine/Source/Runtime/Core/Public/Containers/UnrealString.h
--- a/Engine/Source/Runtime/Core/Public/Containers/UnrealString.h
+++ b/Engine/Source/Runtime/Core/Public/Containers/UnrealString.h
@@ -1192,6 +1192,28 @@ public:
 	 */
 	int32 ParseIntoArrayWS( TArray<FString>* InArray, const TCHAR* pchExtraDelim = NULL ) const;
 
+	/**
+	* Breaks up a delimited string into elements of a string array, using line ending characters
+	* @warning Caution!! this routine is O(N^2) allocations...use it for parsing very short text or not at all!
+	*
+	* @param	InArray			The array to fill with the string pieces	
+	*
+	* @return	The number of elements in InArray
+	*/
+	int32 ParseIntoArrayLines(TArray<FString>* InArray) const;
+
+	/**
+	* Breaks up a delimited string into elements of a string array, using the given delimiters
+	* @warning Caution!! this routine is O(N^2) allocations...use it for parsing very short text or not at all!
+	*
+	* @param	InArray			The array to fill with the string pieces
+	* @param	DelimArray		The strings to delimit on
+	* @param	NumDelims		The number of delimters.
+	*
+	* @return	The number of elements in InArray
+	*/
+	int32 ParseIntoArray(TArray<FString>* InArray, const TCHAR** DelimArray, int32 NumDelims) const;
+
 	/**
 	 * Takes an array of strings and removes any zero length entries.
 	 *
diff --git a/Engine/Source/Runtime/Core/Public/GenericPlatform/GenericPlatformMisc.h b/Engine/Source/Runtime/Core/Public/GenericPlatform/GenericPlatformMisc.h
--- a/Engine/Source/Runtime/Core/Public/GenericPlatform/GenericPlatformMisc.h
+++ b/Engine/Source/Runtime/Core/Public/GenericPlatform/GenericPlatformMisc.h
@@ -244,6 +244,20 @@ struct CORE_API FGenericPlatformMisc
 		*Result = 0;
 	}
 
+	
+	/**
+	* Sets an environment variable to the local process's environment
+	*
+	* @param VariableName The name of the variable (ie "Path")
+	* @param Value The string to set the variable to.	
+	*/
+	static void SetEnvironmentVar(const TCHAR* VariableName, const TCHAR* Value);
+
+	/**
+	 * return the delimiter between paths in the PATH environment variable.
+	 */
+	static const TCHAR* GetPathVarDelimiter();
+
 	/**
 	 * Retrieve the Mac address of the current adapter.
 	 * 
diff --git a/Engine/Source/Runtime/Core/Public/GenericPlatform/GenericPlatformProperties.h b/Engine/Source/Runtime/Core/Public/GenericPlatform/GenericPlatformProperties.h
--- a/Engine/Source/Runtime/Core/Public/GenericPlatform/GenericPlatformProperties.h
+++ b/Engine/Source/Runtime/Core/Public/GenericPlatform/GenericPlatformProperties.h
@@ -152,6 +152,14 @@ struct FGenericPlatformProperties
 		return true;
 	}
 
+	/**
+	 * Returns true if platform supports the AutoSDK system
+	 */
+	static FORCEINLINE bool SupportsAutoSDK()
+	{
+		return false;
+	}
+
 	/**
 	 * Gets whether this platform supports gray scale sRGB texture formats.
 	 *
diff --git a/Engine/Source/Runtime/Core/Public/Linux/LinuxPlatformMisc.h b/Engine/Source/Runtime/Core/Public/Linux/LinuxPlatformMisc.h
--- a/Engine/Source/Runtime/Core/Public/Linux/LinuxPlatformMisc.h
+++ b/Engine/Source/Runtime/Core/Public/Linux/LinuxPlatformMisc.h
@@ -108,6 +108,11 @@ struct CORE_API FLinuxMisc : public FGenericPlatformMisc
 	static const TCHAR* RootDir();
 	static void NormalizePath(FString& InPath);
 
+	static const TCHAR* GetPathVarDelimiter()
+	{
+		return TEXT(":");
+	}
+
 	FORCEINLINE static void MemoryBarrier()
 	{
 		__sync_synchronize();
diff --git a/Engine/Source/Runtime/Core/Public/Linux/LinuxPlatformProperties.h b/Engine/Source/Runtime/Core/Public/Linux/LinuxPlatformProperties.h
--- a/Engine/Source/Runtime/Core/Public/Linux/LinuxPlatformProperties.h
+++ b/Engine/Source/Runtime/Core/Public/Linux/LinuxPlatformProperties.h
@@ -74,6 +74,16 @@ struct FLinuxPlatformProperties
 		return true;
 	}
 
+	static FORCEINLINE bool SupportsAutoSDK()
+	{
+// linux cross-compiling / cross-building from windows supports AutoSDK.  But hosted linux doesn't yet.
+#if PLATFORM_WINDOWS
+		return true;
+#else
+		return false;
+#endif
+	}
+
 	static FORCEINLINE bool SupportsMultipleGameInstances( )
 	{
 		return true;
diff --git a/Engine/Source/Runtime/Core/Public/Mac/MacPlatformMisc.h b/Engine/Source/Runtime/Core/Public/Mac/MacPlatformMisc.h
--- a/Engine/Source/Runtime/Core/Public/Mac/MacPlatformMisc.h
+++ b/Engine/Source/Runtime/Core/Public/Mac/MacPlatformMisc.h
@@ -16,6 +16,12 @@ struct CORE_API FMacPlatformMisc : public FGenericPlatformMisc
 	static void PlatformPostInit(bool IsMoviePlaying = false);
 	static class GenericApplication* CreateApplication();
 	static void GetEnvironmentVariable(const TCHAR* VariableName, TCHAR* Result, int32 ResultLength);
+
+	static const TCHAR* GetPathVarDelimiter()
+	{
+		return TEXT(":");
+	}
+
 	static TArray<uint8> GetMacAddress();
 	static void SubmitErrorReport( const TCHAR* InErrorHist, EErrorReportMode::Type InMode );
 #if !UE_BUILD_SHIPPING
diff --git a/Engine/Source/Runtime/Core/Public/Modules/ModuleManager.h b/Engine/Source/Runtime/Core/Public/Modules/ModuleManager.h
--- a/Engine/Source/Runtime/Core/Public/Modules/ModuleManager.h
+++ b/Engine/Source/Runtime/Core/Public/Modules/ModuleManager.h
@@ -109,6 +109,29 @@ struct FModuleStatus
 	FString CompilationMethod;
 };
 
+/**
+ * Class for managing Sync/Async calls to UBT with the option to rebuild UBT.
+ * Separated from ModuleManager in prep for merging this into the UBT calling code in DesktopPlatform.
+ */
+class CORE_API FUBTInvoker
+{
+public:	
+
+	/** Returns the filename for UBT including the path */
+	static FString GetUnrealBuildToolExecutableFilename();
+
+	/** Returns true if tool was invoked properly */
+	static bool InvokeUnrealBuildToolSync(const FString& InCmdLineParams, FOutputDevice &Ar, bool bSkipBuildUBT, int32& OutReturnCode, FString& OutProcOutput);
+
+	/** Launches UnrealBuildTool with the specified command line parameters */
+	static FProcHandle InvokeUnrealBuildToolAsync(const FString& InCmdLineParams, FOutputDevice &Ar, void*& OutReadPipe, void*& OutWritePipe, bool bSkipBuildUBT = false);
+
+	/** Builds unreal build tool using a compiler specific to the currently running platform */
+	static bool BuildUnrealBuildTool(FOutputDevice &Ar);
+
+	/** Returns the path to the unreal build tool source code */
+	static FString GetUnrealBuildToolSourceCodePath();	
+};
 
 /**
  * Implements the module manager.
@@ -438,12 +461,7 @@ public:
 	 * @param Ar Output device for logging compilation status.
 	 * @return	Returns true if the project was successfully compiled.
 	 */
-	bool GenerateCodeProjectFiles( const FString& GameProjectFilename, FOutputDevice &Ar );
-
-	/**
-	 * @return true if the UBT executable exists (in Rocket) or source code is available to compile it (in non-Rocket)
-	 */
-	bool IsUnrealBuildToolAvailable();
+	bool GenerateCodeProjectFiles( const FString& GameProjectFilename, FOutputDevice &Ar );	
 
 	/** Delegate that's used by the module manager to initialize a registered module that we statically linked with (monolithic only) */
 	DECLARE_DELEGATE_RetVal( IModuleInterface*, FInitializeStaticallyLinkedModule )
@@ -565,14 +583,16 @@ public:
 
 public:
 
-	// FSelfRegisteringExec interface.
-
-	virtual bool Exec( UWorld* Inworld, const TCHAR* Cmd, FOutputDevice& Ar ) override;
+	/**
+	* @return true if the UBT executable exists (in Rocket) or source code is available to compile it (in non-Rocket)
+	*/
+	bool IsUnrealBuildToolAvailable();
 
 public:
 
-	/** @returns Static: Returns arguments to pass to UnrealBuildTool when compiling modules */
-	static FString MakeUBTArgumentsForModuleCompiling();
+	// FSelfRegisteringExec interface.
+
+	virtual bool Exec( UWorld* Inworld, const TCHAR* Cmd, FOutputDevice& Ar ) override;
 
 protected:
 
@@ -686,6 +706,9 @@ protected:
 	 */
 	void MakeUniqueModuleFilename( const FName InModuleName, FString& UniqueSuffix, FString& UniqueModuleFileName );
 
+	/** @returns Static: Returns arguments to pass to UnrealBuildTool when compiling modules */
+	FString MakeUBTArgumentsForModuleCompiling();
+
 	/** 
 	 *	Starts compiling DLL files for one or more modules.
 	 *
@@ -701,17 +724,8 @@ protected:
 		const FRecompileModulesCallback& InRecompileModulesCallback, FOutputDevice& Ar, bool bInFailIfGeneratedCodeChanges, 
 		const FString& InAdditionalCmdLineArgs = FString() );
 
-	/** Returns the path to the unreal build tool source code */
-	FString GetUnrealBuildToolSourceCodePath();
-
-	/** Returns the filename for UBT including the path */
-	FString GetUnrealBuildToolExecutableFilename();
-
-	/** Builds unreal build tool using a compiler specific to the currently running platform */
-	bool BuildUnrealBuildTool(FOutputDevice &Ar);
-
 	/** Launches UnrealBuildTool with the specified command line parameters */
-	bool InvokeUnrealBuildTool(const FString& InCmdLineParams, FOutputDevice &Ar);
+	bool InvokeUnrealBuildToolForCompile(const FString& InCmdLineParams, FOutputDevice &Ar);	
 
 	/** Checks to see if a pending compilation action has completed and optionally waits for it to finish.  If completed, fires any appropriate callbacks and reports status provided bFireEvents is true. */
 	void CheckForFinishedModuleDLLCompile( const bool bWaitForCompletion, bool& bCompileStillInProgress, bool& bCompileSucceeded, FOutputDevice& Ar, const FText& SlowTaskOverrideTest = FText::GetEmpty(), bool bFireEvents = true );
@@ -797,7 +811,6 @@ private:
 	TArray<FString> GameBinariesDirectories;
 };
 
-
 /**
  * Utility class for registering modules that are statically linked.
  */
diff --git a/Engine/Source/Runtime/Core/Public/Windows/WindowsPlatformMisc.h b/Engine/Source/Runtime/Core/Public/Windows/WindowsPlatformMisc.h
--- a/Engine/Source/Runtime/Core/Public/Windows/WindowsPlatformMisc.h
+++ b/Engine/Source/Runtime/Core/Public/Windows/WindowsPlatformMisc.h
@@ -31,6 +31,8 @@ struct CORE_API FWindowsPlatformMisc : public FGenericPlatformMisc
 	static class GenericApplication* CreateApplication();
 	static void SetGracefulTerminationHandler();
 	static void GetEnvironmentVariable(const TCHAR* VariableName, TCHAR* Result, int32 ResultLength);
+	static void SetEnvironmentVar(const TCHAR* VariableName, const TCHAR* Value);
+
 	static TArray<uint8> GetMacAddress();
 	static void SubmitErrorReport( const TCHAR* InErrorHist, EErrorReportMode::Type InMode );
 
