Linux: fix building editor natively (in toolchain).

diff --git a/Engine/Source/Programs/UnrealBuildTool/Linux/LinuxToolChain.cs b/Engine/Source/Programs/UnrealBuildTool/Linux/LinuxToolChain.cs
--- a/Engine/Source/Programs/UnrealBuildTool/Linux/LinuxToolChain.cs
+++ b/Engine/Source/Programs/UnrealBuildTool/Linux/LinuxToolChain.cs
@@ -576,7 +576,7 @@ namespace UnrealBuildTool
 
             bool bUseCmdExe = ExternalExecution.GetRuntimePlatform() == UnrealTargetPlatform.Win64 || ExternalExecution.GetRuntimePlatform() == UnrealTargetPlatform.Win32;
             string ShellBinary = bUseCmdExe ? "cmd.exe" : "/bin/sh";
-            string ExecuteSwitch = bUseCmdExe ? " /C" : " -c";
+            string ExecuteSwitch = bUseCmdExe ? " /C" : ""; // avoid -c so scripts don't need +x
             string ScriptName = bUseCmdExe ? "FixDependencies.bat" : "FixDependencies.sh";
 
             FileItem FixDepsScript = FileItem.GetItemByFullPath(Path.Combine(LinkEnvironment.Config.LocalShadowDirectory, ScriptName));
@@ -589,14 +589,7 @@ namespace UnrealBuildTool
 			PostLinkAction.bCanExecuteRemotely = false;
             PostLinkAction.CommandArguments = ExecuteSwitch;
             
-            if (bUseCmdExe)
-            {
-                PostLinkAction.CommandArguments += " \"";
-            }
-            else
-            {
-                PostLinkAction.CommandArguments += String.Format(" 'chmod +x {0}' && {1} {2} '", FixDepsScript.AbsolutePath, ShellBinary, ExecuteSwitch);
-            }
+            PostLinkAction.CommandArguments += bUseCmdExe ? " \"" : " -c '";
 
 			FileItem OutputFile = FileItem.GetItemByPath(Path.Combine(LinkEnvironment.Config.LocalShadowDirectory, Path.GetFileNameWithoutExtension(Executable.AbsolutePath) + ".link"));
 
@@ -610,7 +603,7 @@ namespace UnrealBuildTool
 
 			PostLinkAction.CommandArguments += ShellBinary + ExecuteSwitch + " \"" + FixDepsScript.AbsolutePath + "\" && ";
 
-            string Touch = bUseCmdExe ? "echo \"\" >> \"{0}\" && copy /b \"{0}\" +,," : "echo \"After script\" && touch \"{0}\"";
+            string Touch = bUseCmdExe ? "echo \"\" >> \"{0}\" && copy /b \"{0}\" +,," : "touch \"{0}\"";
 
             PostLinkAction.CommandArguments += String.Format(Touch, OutputFile.AbsolutePath);
 			PostLinkAction.CommandArguments += bUseCmdExe ? "\"" : "'";
