Temporary disable EULA violation check by default.

diff --git a/Engine/Source/Programs/UnrealBuildTool/Configuration/BuildConfiguration.cs b/Engine/Source/Programs/UnrealBuildTool/Configuration/BuildConfiguration.cs
--- a/Engine/Source/Programs/UnrealBuildTool/Configuration/BuildConfiguration.cs
+++ b/Engine/Source/Programs/UnrealBuildTool/Configuration/BuildConfiguration.cs
@@ -452,7 +452,7 @@ namespace UnrealBuildTool
 			BaseIntermediateFolder = "Intermediate/Build/";
 
 			// By default check and stop the build on EULA violation
-			bCheckLicenseViolations = true;
+			bCheckLicenseViolations = false;
 			bBreakBuildOnLicenseViolation = true;
 		}
 
