Support for using Binaries/ThirdParty libs in self-contained Mac app bundles

diff --git a/Engine/Source/Programs/UnrealBuildTool/Mac/MacToolChain.cs b/Engine/Source/Programs/UnrealBuildTool/Mac/MacToolChain.cs
--- a/Engine/Source/Programs/UnrealBuildTool/Mac/MacToolChain.cs
+++ b/Engine/Source/Programs/UnrealBuildTool/Mac/MacToolChain.cs
@@ -558,6 +558,29 @@ namespace UnrealBuildTool
 			return String.Format("{0}.{1}.{2}", CL / (100 * 100), (CL / 100) % 100, CL % 100);
 		}
 
+		private void AddLibraryPathToRPaths(string Library, string ExeAbsolutePath, ref List<string> RPaths, ref string LinkCommand, bool bIsBuildingAppBundle)
+		{
+			string LibraryDir = Path.GetDirectoryName(Library);
+			string ExeDir = Path.GetDirectoryName(ExeAbsolutePath);
+			if ((Library.Contains("/Plugins/") || Library.Contains("/Binaries/ThirdParty/")) && LibraryDir != ExeDir)
+			{
+				string RelativePath = Utils.MakePathRelativeTo(LibraryDir, ExeDir);
+				if (!RelativePath.Contains(LibraryDir) && !RPaths.Contains(RelativePath))
+				{
+					RPaths.Add(RelativePath);
+					LinkCommand += " -rpath \"@loader_path/" + RelativePath + "\"";
+
+					if (bIsBuildingAppBundle)
+					{
+						string PathInBundle = Path.Combine(Path.GetDirectoryName(ExeDir), "UE4/Engine/Binaries/Mac", RelativePath.Substring(9));
+						Utils.CollapseRelativeDirectories(ref PathInBundle);
+						string RelativePathInBundle = Utils.MakePathRelativeTo(PathInBundle, ExeDir);
+						LinkCommand += " -rpath \"@loader_path/" + RelativePathInBundle + "\"";
+					}
+				}
+			}
+		}
+
 		public override FileItem LinkFiles(LinkEnvironment LinkEnvironment, bool bBuildImportLibraryOnly)
 		{
 			bool bIsBuildingLibrary = LinkEnvironment.Config.bIsBuildingLibrary || bBuildImportLibraryOnly;
@@ -628,6 +651,8 @@ namespace UnrealBuildTool
 				}
 			}
 
+			bool bIsBuildingAppBundle = !LinkEnvironment.Config.bIsBuildingDLL && !LinkEnvironment.Config.bIsBuildingLibrary && !LinkEnvironment.Config.bIsBuildingConsoleApplication;
+
 			if (!bIsBuildingLibrary || LinkEnvironment.Config.bIncludeDependentLibrariesInLibrary)
 			{
 				List<string> RPaths = new List<string>();
@@ -690,15 +715,7 @@ namespace UnrealBuildTool
 						LinkCommand += string.Format(" -l{0}", AdditionalLibrary);
 					}
 
-					if ((AdditionalLibrary.Contains("/Plugins/") || AdditionalLibrary.Contains("/Binaries/ThirdParty/")) && Path.GetDirectoryName(AdditionalLibrary) != Path.GetDirectoryName(AbsolutePath))
-					{
-						string RelativePath = Utils.MakePathRelativeTo(Path.GetDirectoryName(AdditionalLibrary), Path.GetDirectoryName(AbsolutePath));
-						if (!RelativePath.Contains(Path.GetDirectoryName(AdditionalLibrary)) && !RPaths.Contains(RelativePath))
-						{
-							RPaths.Add(RelativePath);
-							LinkCommand += " -rpath \"@loader_path/" + RelativePath + "\"";
-						}
-					}
+					AddLibraryPathToRPaths(AdditionalLibrary, AbsolutePath, ref RPaths, ref LinkCommand, bIsBuildingAppBundle);
 				}
 
 				foreach (string AdditionalLibrary in LinkEnvironment.Config.DelayLoadDLLs)
@@ -711,15 +728,7 @@ namespace UnrealBuildTool
 
 					LinkCommand += string.Format(" -weak_library \"{0}\"", ConvertPath(Path.GetFullPath(AdditionalLibrary)));
 
-					if ((AdditionalLibrary.Contains("/Plugins/") || AdditionalLibrary.Contains("/Binaries/ThirdParty/")) && Path.GetDirectoryName(AdditionalLibrary) != Path.GetDirectoryName(AbsolutePath))
-					{
-						string RelativePath = Utils.MakePathRelativeTo(Path.GetDirectoryName(AdditionalLibrary), Path.GetDirectoryName(AbsolutePath));
-						if (!RelativePath.Contains(Path.GetDirectoryName(AdditionalLibrary)) && !RPaths.Contains(RelativePath))
-						{
-							RPaths.Add(RelativePath);
-							LinkCommand += " -rpath \"@loader_path/" + RelativePath + "\"";
-						}
-					}
+					AddLibraryPathToRPaths(AdditionalLibrary, AbsolutePath, ref RPaths, ref LinkCommand, bIsBuildingAppBundle);
 				}
 			}
 
@@ -897,7 +906,7 @@ namespace UnrealBuildTool
 				DylibCopyScript.Close();
 
 				// For non-console application, prepare a script that will create the app bundle. It'll be run by CreateAppBundle action
-				if (!LinkEnvironment.Config.bIsBuildingDLL && !LinkEnvironment.Config.bIsBuildingLibrary && !LinkEnvironment.Config.bIsBuildingConsoleApplication)
+				if (bIsBuildingAppBundle)
 				{
 					string CreateAppBundleScriptPath = Path.Combine(LinkEnvironment.Config.IntermediateDirectory, "CreateAppBundle.sh");
 					StreamWriter CreateAppBundleScript = File.CreateText(CreateAppBundleScriptPath);
