Saved few more seconds of UBT's time to determine if a target is up to date on Mac by caching File.Exists() result in GetRuntimePlatform()

diff --git a/Engine/Source/Programs/UnrealBuildTool/System/ExternalExecution.cs b/Engine/Source/Programs/UnrealBuildTool/System/ExternalExecution.cs
--- a/Engine/Source/Programs/UnrealBuildTool/System/ExternalExecution.cs
+++ b/Engine/Source/Programs/UnrealBuildTool/System/ExternalExecution.cs
@@ -206,6 +206,7 @@ namespace UnrealBuildTool
 			return LatestWriteTime > DateTime.MinValue ? LatestWriteTime : DateTime.MaxValue;
 		}
 
+		private static bool bIsMac = File.Exists("/System/Library/CoreServices/SystemVersion.plist");
 
 		/** Returns the name of platform UBT is running on */
 		public static UnrealTargetPlatform GetRuntimePlatform()
@@ -216,11 +217,7 @@ namespace UnrealBuildTool
 			case PlatformID.Win32NT:
 				return UnrealTargetPlatform.Win64;
 			case PlatformID.Unix:
-                if (File.Exists("/System/Library/CoreServices/SystemVersion.plist"))
-                {
-                    return UnrealTargetPlatform.Mac;
-                }
-                return UnrealTargetPlatform.Linux;
+				return bIsMac ? UnrealTargetPlatform.Mac : UnrealTargetPlatform.Linux;
 			default:
 				throw new BuildException("Unhandled runtime platform " + Platform);
 			}
