Prefer AutoSDK for Android.

diff --git a/Engine/Source/Programs/UnrealBuildTool/Android/UEBuildAndroid.cs b/Engine/Source/Programs/UnrealBuildTool/Android/UEBuildAndroid.cs
--- a/Engine/Source/Programs/UnrealBuildTool/Android/UEBuildAndroid.cs
+++ b/Engine/Source/Programs/UnrealBuildTool/Android/UEBuildAndroid.cs
@@ -53,6 +53,12 @@ namespace UnrealBuildTool
             return "2.0";
         }
 
+        // prefer auto sdk on android as correct 'manual' sdk detection isn't great at the moment.
+        protected override bool PreferAutoSDK()
+        {
+            return true;
+        }
+
         /// <summary>
         /// checks if the sdk is installed or has been synced
         /// </summary>
