- Fixed Android SDK version parsing to allow for android-L (currently ignoring letter SDKs, assuming they mean beta)

diff --git a/Engine/Source/Programs/UnrealBuildTool/Android/AndroidToolChain.cs b/Engine/Source/Programs/UnrealBuildTool/Android/AndroidToolChain.cs
--- a/Engine/Source/Programs/UnrealBuildTool/Android/AndroidToolChain.cs
+++ b/Engine/Source/Programs/UnrealBuildTool/Android/AndroidToolChain.cs
@@ -46,12 +46,19 @@ namespace UnrealBuildTool
 				string[] Tokens = LocalLevel.Split("-".ToCharArray());
 				if (Tokens.Length >= 2)
 				{
-					int ParsedLevel = int.Parse(Tokens[1]);
-					// bigger? remember it
-					if (ParsedLevel > LargestLevel)
+					try
 					{
-						LargestLevel = ParsedLevel;
-						LargestString = LocalLevel;
+						int ParsedLevel = int.Parse(Tokens[1]);
+						// bigger? remember it
+						if (ParsedLevel > LargestLevel)
+						{
+							LargestLevel = ParsedLevel;
+							LargestString = LocalLevel;
+						}
+					}
+					catch (Exception)
+					{
+						// ignore poorly formed string
 					}
 				}
 			}
