One more UBT on Mac time saver - removed TRACE define from UBT_Mono project

diff --git a/Engine/Source/Programs/UnrealBuildTool/UnrealBuildTool_Mono.csproj b/Engine/Source/Programs/UnrealBuildTool/UnrealBuildTool_Mono.csproj
--- a/Engine/Source/Programs/UnrealBuildTool/UnrealBuildTool_Mono.csproj
+++ b/Engine/Source/Programs/UnrealBuildTool/UnrealBuildTool_Mono.csproj
@@ -1,9 +1,9 @@
-﻿<?xml version="1.0" encoding="utf-8"?>
+<?xml version="1.0" encoding="utf-8"?>
 <Project DefaultTargets="Build" xmlns="http://schemas.microsoft.com/developer/msbuild/2003" ToolsVersion="4.0">
   <PropertyGroup>
     <Configuration Condition=" '$(Configuration)' == '' ">Debug</Configuration>
     <Platform Condition=" '$(Platform)' == '' ">AnyCPU</Platform>
-    <ProductVersion>9.0.30729</ProductVersion>
+    <ProductVersion>8.0.30703</ProductVersion>
     <SchemaVersion>2.0</SchemaVersion>
     <ProjectGuid>{FD7C5E1A-CFE4-4FD5-A525-1EB1599A39AC}</ProjectGuid>
     <OutputType>Exe</OutputType>
@@ -38,7 +38,7 @@
     <DebugType>pdbonly</DebugType>
     <Optimize>true</Optimize>
     <OutputPath>..\..\..\Binaries\DotNET\</OutputPath>
-    <DefineConstants>TRACE</DefineConstants>
+    <DefineConstants></DefineConstants>
     <Optimize>true</Optimize>
     <DebugType>pdbonly</DebugType>
     <PlatformTarget>AnyCPU</PlatformTarget>
@@ -51,7 +51,7 @@
     <DebugType>pdbonly</DebugType>
     <Optimize>true</Optimize>
     <OutputPath>..\..\..\Binaries\DotNET\</OutputPath>
-    <DefineConstants>TRACE</DefineConstants>
+    <DefineConstants></DefineConstants>
     <Optimize>true</Optimize>
     <DebugType>pdbonly</DebugType>
     <PlatformTarget>x86</PlatformTarget>
@@ -62,7 +62,7 @@
   </PropertyGroup>
   <PropertyGroup Condition=" '$(Configuration)|$(Platform)' == 'Debug|AnyCPU' ">
     <OutputPath>..\..\..\Binaries\DotNET\</OutputPath>
-    <DefineConstants>DEBUG;TRACE</DefineConstants>
+    <DefineConstants>DEBUG</DefineConstants>
     <Optimize>true</Optimize>
     <DebugType>pdbonly</DebugType>
     <PlatformTarget>AnyCPU</PlatformTarget>
