- Added the ability for remote compiles to have the [x/y] text in UBT output

diff --git a/Engine/Source/Programs/UnrealBuildTool/IOS/IOSToolChain.cs b/Engine/Source/Programs/UnrealBuildTool/IOS/IOSToolChain.cs
--- a/Engine/Source/Programs/UnrealBuildTool/IOS/IOSToolChain.cs
+++ b/Engine/Source/Programs/UnrealBuildTool/IOS/IOSToolChain.cs
@@ -613,6 +613,7 @@ namespace UnrealBuildTool
 				CompileAction.bIsGCCCompiler = true;
 				// We're already distributing the command by execution on Mac.
 				CompileAction.bCanExecuteRemotely = false;
+				CompileAction.bShouldOutputStatusDescription = true;
 				CompileAction.OutputEventHandler = new DataReceivedEventHandler(RemoteOutputReceivedEventHandler);
 			}
 			return Result;
diff --git a/Engine/Source/Programs/UnrealBuildTool/System/LocalExecutor.cs b/Engine/Source/Programs/UnrealBuildTool/System/LocalExecutor.cs
--- a/Engine/Source/Programs/UnrealBuildTool/System/LocalExecutor.cs
+++ b/Engine/Source/Programs/UnrealBuildTool/System/LocalExecutor.cs
@@ -53,12 +53,6 @@ namespace UnrealBuildTool
 		 */
 		private static void SendOutputToEventHandler(Action Action, string Output)
 		{
-			// Output status description (file name) when no output is returned
-			if (string.IsNullOrEmpty(Output))
-			{
-				Output = Action.StatusDescription;
-			}
-
 			// pass the output to any handler requested
 			if (Action.OutputEventHandler != null)
 			{
@@ -122,6 +116,20 @@ namespace UnrealBuildTool
 				// call the function and get the ExitCode and an output string
 				string Output;
 				Action.ActionHandler(Action, out ExitCode, out Output);
+
+				// Output status description (file name) when no output is returned
+				if (string.IsNullOrEmpty(Output))
+				{
+					if (Action.bShouldOutputStatusDescription)
+					{
+						Output = string.Format("[{0}/{1}] {2} {3}", JobNumber, TotalJobs, Action.CommandDescription, Action.StatusDescription);
+					}
+					else
+					{
+						Output = Action.StatusDescription;
+					}
+				}
+				
 				SendOutputToEventHandler(Action, Output);
 			}
 			else
