- Massive SSH usability improvements for remote compiling. Added a script to interactively run SSH and RSync to make a key on the Mac, install it, and pull down the private key for use after
- Currently only UBT will do the check when using SSH, and will run the script. This should really be moved to the editor.
- Keys can be in My Documents, or in Engine\\Build (for source control sharing)
- SSH is still not 100% ready for primetime (IPP needs re-integration into UAT)
- Epic default is to use the standard builder account on the standard machine

diff --git a/Engine/Build/BatchFiles/MakeAndInstallSSHKey.bat b/Engine/Build/BatchFiles/MakeAndInstallSSHKey.bat
new file mode 100644
--- /dev/null
+++ b/Engine/Build/BatchFiles/MakeAndInstallSSHKey.bat
@@ -0,0 +1,61 @@
+@echo off
+setlocal
+
+set SSH=%1
+set SSHDIR=%~d1%~p1
+set RSYNC=%2
+set USER=%3
+set MACHINE=%4
+set DOCUMENTS=%5
+set CYGWIN_DOCUMENTS=%6
+set ENGINE=%~7
+
+set KEY_DIR=%DOCUMENTS%\Unreal Engine\UnrealBuildTool\SSHKeys\%MACHINE%\%USER%
+set KEY_PATH=%DOCUMENTS%\Unreal Engine\UnrealBuildTool\SSHKeys\%MACHINE%\%USER%\RemoteToolChainPrivate.key
+set CYGWIN_KEY_PATH=%CYGWIN_DOCUMENTS%/Unreal Engine/UnrealBuildTool/SSHKeys/%MACHINE%/%USER%/RemoteToolChainPrivate.key
+
+echo %SSHDIR%
+
+echo.
+echo ================================================================================
+echo Connecting to %MACHINE% as user %USER% to create an SSH Key.
+echo.
+echo Please note the following:
+echo   It may prompt you about "authenticity of host can't be established." You must enter yes
+echo   It will then prompt your for your password on %MACHINE%
+echo   Finally it will ask your for a 'passphrase' for your new SSH key. Do NOT give it a passcode, just hit enter
+echo   Hit enter again when prompted to confirm.
+echo ================================================================================
+echo.
+
+pause
+
+%SSH% %USER%@%MACHINE% "if [[ ! -e .ssh ]]; then mkdir .ssh; fi && cd .ssh && if [[ -e authorized_keys ]]; then cp -f authorized_keys authorized_keys_UEBackup; fi && ssh-keygen -t rsa -f RemoteToolChain && mv -f RemoteToolChain.pub RemoteToolChainPublic.key && mv -f RemoteToolChain RemoteToolChainPrivate.key && cat RemoteToolChainPublic.key >> authorized_keys";
+
+echo.
+echo ================================================================================
+echo Now we will connect again to %MACHINE% to download the generated private key.
+echo You will have to enter your password again.
+echo The private key will be stored in:
+echo    %KEY_PATH%
+echo.
+echo If you want to share this key with others on your team, copy the contents of SSHKeys directory to:
+echo    %ENGINE%\Build\SSHKeys
+echo and then put it into your source control.
+echo ================================================================================
+echo.
+
+pause
+
+mkdir "%KEY_DIR%" 2> NUL
+pushd %SSHDIR%
+%RSYNC% -za %USER%@%MACHINE%:.ssh/RemoteToolChainPrivate.key "%CYGWIN_KEY_PATH%"
+popd
+
+echo.
+echo ===================================================================================================
+echo Assuming you have seen no errors above, you are ready to compile UE4 remotely on %MACHINE%!
+echo ===================================================================================================
+echo.
+
+pause
\ No newline at end of file
diff --git a/Engine/Source/Programs/UnrealBuildTool/System/RPCUtilHelper.cs b/Engine/Source/Programs/UnrealBuildTool/System/RPCUtilHelper.cs
--- a/Engine/Source/Programs/UnrealBuildTool/System/RPCUtilHelper.cs
+++ b/Engine/Source/Programs/UnrealBuildTool/System/RPCUtilHelper.cs
@@ -94,7 +94,7 @@ namespace UnrealBuildTool
 			}
 		}
 
-		static public void Initialize(string InMacName)
+		static public int Initialize(string InMacName)
 		{
 			MacName = InMacName;
 
@@ -105,8 +105,13 @@ namespace UnrealBuildTool
 				{
 					if (!RemoteToolChain.bUseRPCUtil)
 					{
+						// make sure we have SSH setup
+						if (RemoteToolChain.ResolvedSSHAuthentication == null)
+						{
+							Log.TraceError("SSH authentication required a key, but one was not found. Use Editor to setup remote authentication!");
+							return 100;
+						}
 						// ask for current time, free memory and num CPUs
-
 						string[] Commands = new string[] 
 						{
 							"+\"%s\"",
@@ -115,6 +120,12 @@ namespace UnrealBuildTool
 							"sysctl -a | grep hw.logicalcpu: | awk '{print $2}'",
 						};
 						Hashtable Results = Command("/", "date", string.Join(" && ", Commands), null);
+						if ((Int64)Results["ExitCode"] != 0)
+						{
+							Log.TraceError("Failed to run init commands on {0}. Output = {1}", MacName, Results["CommandOutput"]);
+							return 101;
+						}
+
 						string[] Lines = ((string)Results["CommandOutput"]).Split("\r\n".ToCharArray(), StringSplitOptions.RemoveEmptyEntries);
 
 						// convert it to a time, with TimeDifferenceFromRemote as 0
@@ -146,14 +157,18 @@ namespace UnrealBuildTool
 				}
 				catch(Exception Ex)
 				{
-					Console.WriteLine("Ex {0}", Ex.ToString());
-					throw new BuildException("Failed to run init commands on " + MacName);
+					Log.TraceVerbose("SSH Initialize exception {0}", Ex.ToString());
+					Log.TraceError("Failed to run init commands on {0}", MacName);
+					return 101;
 				}
 			}
  			else
  			{
- 				throw new BuildException("Failed to ping Mac named " + MacName);
+				Log.TraceError("Failed to ping Mac named {0}", MacName);
+				return 102;
  			}
+
+			return 0;
 		}
 
 		/**
@@ -520,7 +535,7 @@ namespace UnrealBuildTool
 			}
 			else
 			{
-				return RemoteToolChain.SSHCommand(WorkingDirectory, Command, CommandArgs, RemoteOutputPath);
+				return RemoteToolChain.SSHCommand(WorkingDirectory, Command + " " + CommandArgs, RemoteOutputPath);
 			}
 		}
 	}
diff --git a/Engine/Source/Programs/UnrealBuildTool/ToolChain/RemoteToolChain.cs b/Engine/Source/Programs/UnrealBuildTool/ToolChain/RemoteToolChain.cs
--- a/Engine/Source/Programs/UnrealBuildTool/ToolChain/RemoteToolChain.cs
+++ b/Engine/Source/Programs/UnrealBuildTool/ToolChain/RemoteToolChain.cs
@@ -19,14 +19,18 @@ namespace UnrealBuildTool
 		{
 			RemoteToolChainPlatform = InPlatform;
 
-			// Register this tool chain for IOS
-			Log.TraceVerbose("        Registered for {0}", CPPPlatform.ToString());
-			UEToolChain.RegisterPlatformToolChain(CPPPlatform, this);
+			// make sure no error happened during init time
+			if (InitializeRemoteExecution() == 0)
+			{
+				// Register this tool chain for IOS
+				Log.TraceVerbose("        Registered for {0}", CPPPlatform.ToString());
+				UEToolChain.RegisterPlatformToolChain(CPPPlatform, this);
+			}
 		}
 
 		/** These two variables will be loaded from XML config file in XmlConfigLoader.Init() */
 		[XmlConfig]
-		public static string RemoteServerName;
+		public static string RemoteServerName = "";
 		[XmlConfig]
 		public static string[] PotentialServerNames = new string[] { };
 
@@ -42,23 +46,36 @@ namespace UnrealBuildTool
 		
 		/** Path to rsync executable and parameters for your rsync utility */
 		[XmlConfig]
-		public static string RSyncExe = "C:\\Program Files (x86)\\DeltaCopy\\rsync.exe";
+		public static string RSyncExe = "${PROGRAM_FILES}\\DeltaCopy\\rsync.exe";
+		public static string ResolvedRSyncExe = null;
 
 		/** Path to rsync executable and parameters for your rsync utility */
 		[XmlConfig]
-		public static string SSHExe = "C:\\Program Files (x86)\\DeltaCopy\\ssh.exe";
+		public static string SSHExe = "${PROGRAM_FILES}\\DeltaCopy\\ssh.exe";
+		public static string ResolvedSSHExe = null;
+
+		/** Instead of looking for RemoteToolChainPrivate.key in the usual places (Documents/Unreal Engine/UnrealBuildTool/SSHKeys, Engine/Build/SSHKeys), this private key will be used if set */
+		[XmlConfig]
+		public static string SSHPrivateKeyOverridePath = "";
+		public static string ResolvedSSHPrivateKey = null;
 
 		/** The authentication used for Rsync (for the -e rsync flag) */
 		[XmlConfig]
-		public static string RsyncAuthentication = "ssh -i id_rsa";
+		public static string RsyncAuthentication = "ssh -i '${CYGWIN_SSH_PRIVATE_KEY}'";
+		public static string ResolvedRsyncAuthentication = null;
 
 		/** The authentication used for SSH (probably similar to RsyncAuthentication) */
 		[XmlConfig]
-		public static string SSHAuthentication = "-i id_rsa";
+		public static string SSHAuthentication = "-i '${CYGWIN_SSH_PRIVATE_KEY}'";
+		public static string ResolvedSSHAuthentication = null;
 
 		/** Username on the remote machine to connect to with RSync */
 		[XmlConfig]
-		public static string RSyncUsername = "josh.adams";
+		public static string RSyncUsername = "${CURRENT_USER}";
+		public static string ResolvedRSyncUsername = null;
+
+		// has the toolchain initialized remote execution yet? no need to do it multiple times
+		private static bool bHasBeenInitialized = false;
 
 		/** The directory that this local branch is in, without drive information (strip off X:\ from X:\UE4\iOS) */
 		public static string BranchDirectory = Path.GetFullPath(".\\");
@@ -84,11 +101,104 @@ namespace UnrealBuildTool
 
 			BranchDirectory = BranchDirectory.Replace("Engine\\Binaries\\DotNET", "");
 			BranchDirectory = BranchDirectory.Replace("Engine\\Source\\", "");
-        }
+		}
+
+		private static string ResolveString(string Input, bool bIsPath)
+		{
+			string Result = Input;
+
+			// these assume entire string is a path, and will do file operations on the whole string
+			if (bIsPath)
+			{
+				if (Result.Contains("${PROGRAM_FILES}"))
+				{
+					// first look in real ProgramFiles
+					string Temp = Result.Replace("${PROGRAM_FILES}", Environment.GetFolderPath(Environment.SpecialFolder.ProgramFiles, Environment.SpecialFolderOption.DoNotVerify));
+					if (File.Exists(Temp) || Directory.Exists(Temp))
+					{
+						Result = Temp;
+					}
+					else
+					{
+						// fallback to ProgramFilesX86
+						Temp = Result.Replace("${PROGRAM_FILES}", Environment.GetFolderPath(Environment.SpecialFolder.ProgramFilesX86, Environment.SpecialFolderOption.DoNotVerify));
+						if (File.Exists(Temp) || Directory.Exists(Temp))
+						{
+							Result = Temp;
+						}
+					}
+				}
+
+				if (Result.Contains("${ENGINE_ROOT}"))
+				{
+					string Temp = Result.Replace("${ENGINE_ROOT}", Path.GetFullPath(BuildConfiguration.RelativeEnginePath + "\\.."));
+
+					// get the best version
+					Result = LookForSpecialFile(Temp);
+				}
+
+				if (Result.Contains("${PROJECT_ROOT}"))
+				{
+					if (!UnrealBuildTool.HasUProjectFile())
+					{
+						throw new BuildException("Configuration setting was using ${PROJECT_ROOT}, but there was no project specified");
+					}
+
+					string Temp = Result.Replace("${PROJECT_ROOT}", Path.GetFullPath(UnrealBuildTool.GetUProjectPath()));
+					
+					// get the best version
+					Result = LookForSpecialFile(Temp);
+				}
+			}
+
+			// non path variables
+			Result = Result.Replace("${CURRENT_USER}", Environment.UserName);
+
+			// needs a resolved key (which isn't required if user is using alternate authentication)
+			if (Result.Contains("${SSH_PRIVATE_KEY}") || Result.Contains("${CYGWIN_SSH_PRIVATE_KEY}"))
+			{
+				// if it needs the key, then make sure we have it!
+				if (ResolvedSSHPrivateKey != null)
+				{
+					Result = Result.Replace("${SSH_PRIVATE_KEY}", ResolvedSSHPrivateKey);
+					Result = Result.Replace("${CYGWIN_SSH_PRIVATE_KEY}", ConvertPathToCygwin(ResolvedSSHPrivateKey));
+				}
+				else
+				{
+					Result = null;
+				}
+			}
+
+			return Result;
+		}
+
+		private static string LookForSpecialFile(string InPath)
+		{
+			// look in special NotForLicensees dir first
+			string Special = Path.Combine(Path.GetDirectoryName(InPath), "NoRedist", Path.GetFileName(InPath));
+			if (File.Exists(Special) || Directory.Exists(Special))
+			{
+				return Special;
+			}
+
+			Special = Path.Combine(Path.GetDirectoryName(InPath), "NotForLicensees", Path.GetFileName(InPath));
+			if (File.Exists(Special) || Directory.Exists(Special))
+			{
+				return Special;
+			}
+
+			return InPath;
+		}
 
         // Do any one-time, global initialization for the tool chain
-        public override void SetUpGlobalEnvironment()
+		static int InitializationErrorCode = 0;
+		private static int InitializeRemoteExecution()
         {
+			if (bHasBeenInitialized)
+			{
+				return InitializationErrorCode;
+			}
+
             if (ExternalExecution.GetRuntimePlatform() != UnrealTargetPlatform.Mac)
             {
             	// If we don't care which machine we're going to build on, query and
@@ -146,16 +256,92 @@ namespace UnrealBuildTool
 					Log.TraceInformation("Picking the default remote server " + RemoteServerName);
 				}
 
-				if (!Utils.IsRunningOnMono)
+				// we need a server name!
+				if (string.IsNullOrEmpty(RemoteServerName))
 				{
-					// crank up RPC communications
-					RPCUtilHelper.Initialize(RemoteServerName);
+					Log.TraceError("Remote compiling requires a server name. Use the editor to set up your remote compilation settings.");
+					InitializationErrorCode = 99;
 				}
-            }
+
+				// we need the RemoteServerName and the Username to find the private key
+				ResolvedRSyncUsername = ResolveString(RSyncUsername, false);
+
+				// if the override path is set, just use it directly
+				if (!string.IsNullOrEmpty(SSHPrivateKeyOverridePath))
+				{
+					ResolvedSSHPrivateKey = ResolveString(SSHPrivateKeyOverridePath, true);
+					// make sure it exists
+					if (!File.Exists(ResolvedSSHPrivateKey))
+					{
+						throw new BuildException("An SSHKey override was specified [" + SSHPrivateKeyOverridePath + "] but it doesn't exist. Can't continue...");
+					}
+				}
+				else
+				{
+					// all the places to look for a key
+					string[] Locations = new string[] {
+						Path.Combine(Environment.GetFolderPath(Environment.SpecialFolder.ApplicationData), "Unreal Engine", "UnrealBuildTool"),
+						Path.Combine(Environment.GetFolderPath(Environment.SpecialFolder.Personal), "Unreal Engine", "UnrealBuildTool"),
+						Path.Combine(BuildConfiguration.RelativeEnginePath, "Build", "NotForLicensees"),
+						Path.Combine(BuildConfiguration.RelativeEnginePath, "Build", "NoRedist"),
+						Path.Combine(BuildConfiguration.RelativeEnginePath, "Build"),
+					};
+
+					// look for a key file
+					foreach (string Location in Locations)
+					{
+						string KeyPath = Path.Combine(Location, "SSHKeys", RemoteServerName, ResolvedRSyncUsername, "RemoteToolChainPrivate.key");
+						if (File.Exists(KeyPath))
+						{
+							ResolvedSSHPrivateKey = KeyPath;
+							break;
+						}
+					}
+				}
+
+				// resolve the rest of the strings
+				ResolvedRSyncExe = ResolveString(RSyncExe, true);
+				ResolvedSSHExe = ResolveString(SSHExe, true);
+				ResolvedRsyncAuthentication = ResolveString(RsyncAuthentication, false);
+				ResolvedSSHAuthentication = ResolveString(SSHAuthentication, false);
+
+				// start up remote communication and record if it succeeds
+				InitializationErrorCode = RPCUtilHelper.Initialize(RemoteServerName);
+
+				// allow user to set up
+				if (InitializationErrorCode == 100)
+				{
+					Process KeyProcess = new Process();
+					KeyProcess.StartInfo.WorkingDirectory = Path.GetFullPath(Path.Combine(BuildConfiguration.RelativeEnginePath, "Build", "BatchFiles"));
+					KeyProcess.StartInfo.FileName = "MakeAndInstallSSHKey.bat";
+					KeyProcess.StartInfo.Arguments = string.Format(
+						"\"{0}\" \"{1}\" {2} {3} \"{4}\" \"{5}\" \"{6}\"",
+						ResolvedSSHExe,
+						ResolvedRSyncExe,
+						ResolvedRSyncUsername,
+						RemoteServerName,
+						Environment.GetFolderPath(Environment.SpecialFolder.Personal),
+						ConvertPathToCygwin(Environment.GetFolderPath(Environment.SpecialFolder.Personal)),
+						Path.GetFullPath(BuildConfiguration.RelativeEnginePath));
+
+					KeyProcess.Start();
+					KeyProcess.WaitForExit();
+
+					// make sure it succeeded if we want to re-init
+					if (KeyProcess.ExitCode == 0)
+					{
+						InitializeRemoteExecution();
+					}
+				}
+			}
             else
-            {
+			{
 				RemoteServerName = Environment.MachineName;
-            }
+				// can't error in this case
+			}
+
+			bHasBeenInitialized = true;
+			return InitializationErrorCode;
         }
 
 
@@ -347,6 +533,10 @@ namespace UnrealBuildTool
 		
 		private static string ConvertPathToCygwin(string InPath)
 		{
+			if (InPath == null)
+			{
+				return null;
+			}
 			return "/cygdrive/" + Utils.CleanDirectorySeparators(InPath.Replace(":", ""), '/');
 		}
 
@@ -413,8 +603,8 @@ namespace UnrealBuildTool
 				string CygRootPath = ConvertPathToCygwin(Path.GetFullPath("../../"));
 				string RemotePath = ConvertPath(Path.GetFullPath("../../")).Replace(" ", "\\ ");
 
-				// get the executable dir for Rsync
-				string ExeDir = Path.GetDirectoryName(RSyncExe);
+				// get the executable dir for SSH, so Rsync can call it easily
+				string ExeDir = Path.GetDirectoryName(ResolvedSSHExe);
 	
 				Process RsyncProcess = new Process();
 				if (ExeDir != "")
@@ -423,10 +613,10 @@ namespace UnrealBuildTool
 				}
 
 				// --exclude='*'  ??? why???
-				RsyncProcess.StartInfo.FileName = RSyncExe;
+				RsyncProcess.StartInfo.FileName = ResolvedRSyncExe;
 				RsyncProcess.StartInfo.Arguments = string.Format(
-					"-vzae \"{0}\" --rsync-path=\"mkdir -p {2} && rsync\" --delete --files-from=\"{4}\" --include-from=\"{5}\" --include='*/' --exclude='*.o' {1} {6}@{3}:'{2}'",
-					RsyncAuthentication,
+					"-vzae \"{0}\" --rsync-path=\"mkdir -p {2} && rsync\" --delete --files-from=\"{4}\" --include-from=\"{5}\" --include='*/' --exclude='*.o' --exclude='Timestamp' {1} {6}@{3}:'{2}'",
+					ResolvedRsyncAuthentication,
 					CygRootPath,
 					RemotePath,
 					RemoteServerName,
@@ -454,8 +644,8 @@ namespace UnrealBuildTool
 			string RemoteDir = Path.GetDirectoryName(RemotePath).Replace("\\", "/");
 			string RemoteFilename = Path.GetFileName(RemotePath);
 
-			// get the executable dir for Rsync
-			string ExeDir = Path.GetDirectoryName(RSyncExe);
+			// get the executable dir for SSH, so Rsync can call it easily
+			string ExeDir = Path.GetDirectoryName(ResolvedSSHExe);
 
 			Process RsyncProcess = new Process();
 			if (ExeDir != "")
@@ -464,10 +654,10 @@ namespace UnrealBuildTool
 			}
 
 			// make simple rsync commandline to send a file
-			RsyncProcess.StartInfo.FileName = RSyncExe;
+			RsyncProcess.StartInfo.FileName = ResolvedRSyncExe;
 			RsyncProcess.StartInfo.Arguments = string.Format(
 				"-zae \"{0}\" --rsync-path=\"mkdir -p '{1}' && rsync\" \"{2}\" {3}@{4}:\"{1}/{5}\"",
-				RsyncAuthentication,
+				ResolvedRsyncAuthentication,
 				RemoteDir,
 				ConvertPathToCygwin(LocalPath),
 				RSyncUsername,
@@ -484,8 +674,8 @@ namespace UnrealBuildTool
 
 		static public bool DownloadFile(string RemotePath, string LocalPath)
 		{
-			// get the executable dir for Rsync
-			string ExeDir = Path.GetDirectoryName(RSyncExe);
+			// get the executable dir for SSH, so Rsync can call it easily
+			string ExeDir = Path.GetDirectoryName(ResolvedSSHExe);
 
 			Process RsyncProcess = new Process();
 			if (ExeDir != "")
@@ -497,10 +687,10 @@ namespace UnrealBuildTool
 			Directory.CreateDirectory(Path.GetDirectoryName(LocalPath));
 
 			// make simple rsync commandline to send a file
-			RsyncProcess.StartInfo.FileName = RSyncExe;
+			RsyncProcess.StartInfo.FileName = ResolvedRSyncExe;
 			RsyncProcess.StartInfo.Arguments = string.Format(
-				"-zae \"{0}\" {2}@{3}:\"{4}\" \"{1}\"",
-				RsyncAuthentication,
+				"-zae {0} {2}@{3}:\"{4}\" \"{1}\"",
+				ResolvedRsyncAuthentication,
 				ConvertPathToCygwin(LocalPath),
 				RSyncUsername,
 				RemoteServerName,
@@ -516,9 +706,9 @@ namespace UnrealBuildTool
 			return Utils.RunLocalProcess(RsyncProcess) == 0;
 		}
 
-		static public Hashtable SSHCommand(string WorkingDirectory, string Command, string CommandArgs, string RemoteOutputPath)
+		static public Hashtable SSHCommand(string WorkingDirectory, string Command, string RemoteOutputPath)
 		{
-			Console.WriteLine("Doing {0} {1}", Command, CommandArgs);
+			Console.WriteLine("Doing {0}", Command);
 
 			// make the commandline for other end
 			string RemoteCommandline = "cd " + WorkingDirectory;
@@ -528,7 +718,7 @@ namespace UnrealBuildTool
 			}
 
 			// get the executable dir for SSH
-			string ExeDir = Path.GetDirectoryName(SSHExe);
+			string ExeDir = Path.GetDirectoryName(ResolvedSSHExe);
 
 			Process SSHProcess = new Process();
 			if (ExeDir != "")
@@ -537,11 +727,11 @@ namespace UnrealBuildTool
 			}
 
 			// long commands go as a file
-			if (CommandArgs.Length > 256)
+			if (Command.Length > 1024)
 			{
 				// upload the commandline text file
 				string CommandLineFile = Path.GetTempFileName();
-				File.WriteAllText(CommandLineFile, Command + " " + CommandArgs);
+				File.WriteAllText(CommandLineFile, Command);
 
 				string RemoteCommandlineDir = "/var/tmp/" + Environment.MachineName;
 				string RemoteCommandlinePath = RemoteCommandlineDir + "/" + Path.GetFileName(CommandLineFile);
@@ -555,17 +745,19 @@ namespace UnrealBuildTool
 			}
 			else
 			{
-				RemoteCommandline += " && " + Command + " " + CommandArgs;
+				RemoteCommandline += " && " + Command;
 			}
 
-			SSHProcess.StartInfo.FileName = SSHExe;
+			SSHProcess.StartInfo.FileName = ResolvedSSHExe;
 			SSHProcess.StartInfo.Arguments = string.Format(
 				"{0} {1}@{2} \"{3}\"",
-				SSHAuthentication,
+				ResolvedSSHAuthentication,
 				RSyncUsername,
 				RemoteServerName,
 				RemoteCommandline.Replace("\"", "\\\""));
 
+			Hashtable Return = new Hashtable();
+
 			// add this process to the map
 			SSHOutputMap[SSHProcess] = new StringBuilder("");
 			SSHProcess.OutputDataReceived += new DataReceivedEventHandler(OutputReceivedForSSH);
@@ -576,7 +768,6 @@ namespace UnrealBuildTool
 			Console.WriteLine("Execute took {0}", (DateTime.Now - Start).ToString());
 
 			// now we have enough to fill out the HashTable
-			Hashtable Return = new Hashtable();
 			Return["CommandOutput"] = SSHOutputMap[SSHProcess].ToString();
 			Return["ExitCode"] = (object)ExitCode;
 
