Add more error printout to AutoSDK.  Extra checks for OutputEnvVars

diff --git a/Engine/Source/Programs/UnrealBuildTool/Configuration/UEBuildPlatform.cs b/Engine/Source/Programs/UnrealBuildTool/Configuration/UEBuildPlatform.cs
--- a/Engine/Source/Programs/UnrealBuildTool/Configuration/UEBuildPlatform.cs
+++ b/Engine/Source/Programs/UnrealBuildTool/Configuration/UEBuildPlatform.cs
@@ -1110,6 +1110,14 @@ namespace UnrealBuildTool
 
                     return true;
                 }
+                else
+                {
+                    Console.WriteLine("File {0} does not exist", HookExe);
+                }
+            }
+            else
+            {
+                Console.WriteLine("Version string is blank for {0}. Can't determine {1} hook.", PlatformSDKRoot, Hook.ToString());
             }
 
             return bHookCanBeNonExistent;
@@ -1246,6 +1254,10 @@ namespace UnrealBuildTool
                     return true;
                 }
             }
+            else
+            {
+                Console.WriteLine("Cannot set up environment for {1} because command file {1} does not exist.", PlatformSDKRoot, EnvVarFile);
+            }
 
             return false;
         }
@@ -1395,6 +1407,15 @@ namespace UnrealBuildTool
                         RunAutoSDKHooks(AutoSDKRoot, GetRequiredSDKString(), SDKHookType.Uninstall, false);
                         return;
                     }
+
+                    string EnvVarFile = Path.Combine(AutoSDKRoot, SDKEnvironmentVarsFile);
+                    if (!File.Exists(EnvVarFile))
+                    {
+                        Console.WriteLine("Installation of required SDK {0}.  Did not generate Environment file {1}", GetRequiredSDKString(), EnvVarFile);
+                        RunAutoSDKHooks(AutoSDKRoot, GetRequiredSDKString(), SDKHookType.Uninstall, false);
+                        return;
+                    }
+
                     SetCurrentlyInstalledAutoSDKString(GetRequiredSDKString());
                     SetLastRunAutoSDKScriptVersion(GetRequiredScriptVersionString());
                 }
