Prefer AutoSDKs for Linux.

diff --git a/Engine/Source/Programs/UnrealBuildTool/Linux/UEBuildLinux.cs b/Engine/Source/Programs/UnrealBuildTool/Linux/UEBuildLinux.cs
--- a/Engine/Source/Programs/UnrealBuildTool/Linux/UEBuildLinux.cs
+++ b/Engine/Source/Programs/UnrealBuildTool/Linux/UEBuildLinux.cs
@@ -69,6 +69,12 @@ namespace UnrealBuildTool
             return "2.0";
         }
 
+        protected virtual bool PreferAutoSDK()
+        {
+            // having LINUX_ROOT set (for legacy reasons or for convenience of cross-compiling certain third party libs) should not make UBT skip AutoSDKs
+            return true;
+        }
+
         /**
          *	Whether the required external SDKs are installed for this platform
          */
