addition of command to UBT to validate SDKs are installed for the various platforms
#ue4
#ubt

diff --git a/Engine/Source/Programs/UnrealBuildTool/IOS/UEBuildIOS.cs b/Engine/Source/Programs/UnrealBuildTool/IOS/UEBuildIOS.cs
--- a/Engine/Source/Programs/UnrealBuildTool/IOS/UEBuildIOS.cs
+++ b/Engine/Source/Programs/UnrealBuildTool/IOS/UEBuildIOS.cs
@@ -23,7 +23,16 @@ namespace UnrealBuildTool
 
         protected override SDKStatus HasRequiredManualSDKInternal()
         {
-            return SDKStatus.Valid;
+			if (ExternalExecution.GetRuntimePlatform() != UnrealTargetPlatform.Mac)
+			{
+				// check to see if iTunes is installed
+				string dllPath = Microsoft.Win32.Registry.GetValue("HKEY_LOCAL_MACHINE\\SOFTWARE\\Wow6432Node\\Apple Inc.\\Apple Mobile Device Support\\Shared", "iTunesMobileDeviceDLL", null) as string;
+				if (String.IsNullOrEmpty(dllPath) || !File.Exists(dllPath))
+				{
+					return SDKStatus.Invalid;
+				}
+			}
+			return SDKStatus.Valid;
         }
 
         /**
diff --git a/Engine/Source/Programs/UnrealBuildTool/System/UnrealBuildTool.cs b/Engine/Source/Programs/UnrealBuildTool/System/UnrealBuildTool.cs
--- a/Engine/Source/Programs/UnrealBuildTool/System/UnrealBuildTool.cs
+++ b/Engine/Source/Programs/UnrealBuildTool/System/UnrealBuildTool.cs
@@ -647,6 +647,7 @@ namespace UnrealBuildTool
                     var bSpecificModulesOnly = false;
                     var bIgnoreJunk = false;
                     var bAutoSDKOnly = false;
+					var bValidatePlatforms = false;
 
                     // We need to be able to identify if one of the arguments is the platform...
                     // Leverage the existing parser function in UEBuildTarget to get this information.
@@ -813,23 +814,26 @@ namespace UnrealBuildTool
                             // It's the platform set...
                             //PlatformName = Arg;
                         }
-                        
-                        else
-                        {
-                            // This arg may be a game name. Check for the existence of a game folder with this name.
-                            // "Engine" is not a valid game name.
-                            if (LowercaseArg != "engine" && Arg.IndexOfAny(Path.GetInvalidPathChars()) == -1 &&
-                                Directory.Exists(Path.Combine(ProjectFileGenerator.RootRelativePath, Arg, "Config")))
-                            {
-                                GameName = Arg;
-                                Log.TraceVerbose("CommandLine: Found game name '{0}'", GameName);
-                            }
-                            else if (LowercaseArg == "rocket")
-                            {
-                                GameName = Arg;
-                                Log.TraceVerbose("CommandLine: Found game name '{0}'", GameName);
-                            }
-                        }
+						else if (LowercaseArg == "-validateplatform")
+						{
+							bValidatePlatforms = true;
+						}
+						else
+						{
+							// This arg may be a game name. Check for the existence of a game folder with this name.
+							// "Engine" is not a valid game name.
+							if (LowercaseArg != "engine" && Arg.IndexOfAny(Path.GetInvalidPathChars()) == -1 &&
+								Directory.Exists(Path.Combine(ProjectFileGenerator.RootRelativePath, Arg, "Config")))
+							{
+								GameName = Arg;
+								Log.TraceVerbose("CommandLine: Found game name '{0}'", GameName);
+							}
+							else if (LowercaseArg == "rocket")
+							{
+								GameName = Arg;
+								Log.TraceVerbose("CommandLine: Found game name '{0}'", GameName);
+							}
+						}
                     }
 
                     // Send an event with basic usage dimensions
@@ -848,7 +852,9 @@ namespace UnrealBuildTool
                                     ? "RunCopyrightVerification"
                                         : bGenerateMakefiles
                                         ? "GenerateMakefiles"
-                                            : "Build",
+                                            : bValidatePlatforms 
+												? "ValidatePlatfomrs"
+												: "Build",
                         "Platform", CheckPlatform.ToString(),
                         "Configuration", CheckConfiguration.ToString(),
                         "IsRocket", bRunningRocket.ToString(),
@@ -906,56 +912,60 @@ namespace UnrealBuildTool
                         // RegisterAllUBTClasses has already done all the SDK setup.
                         Result = ECompilationResult.Succeeded;
                     }
-                    else
-                    {
-                        // Check if any third party headers are included from public engine headers.
-                        if (bCheckThirdPartyHeaders)
-                        {
-                            ThirdPartyHeaderFinder.FindThirdPartyIncludes(CheckPlatform, CheckConfiguration);
-                        }
-
-                        // Build our project
-                        if (Result == ECompilationResult.Succeeded)
-                        {
-                            if (UEBuildConfiguration.bPrepForDeployment == false)
-                            {
-                                // If we are only prepping for deployment, assume the build already occurred.
-                                Result = RunUBT(Arguments);
-                            }
-                            else
-                            {
-                                var BuildPlatform = UEBuildPlatform.GetBuildPlatform(CheckPlatform, true);
-                                if (BuildPlatform != null)
-                                {
-                                    // Setup environment wasn't called, so set the flag
-                                    BuildConfiguration.bDeployAfterCompile = BuildPlatform.RequiresDeployPrepAfterCompile();
-                                    BuildConfiguration.PlatformIntermediateFolder = Path.Combine(BuildConfiguration.BaseIntermediateFolder, CheckPlatform.ToString(), BuildPlatform.GetActiveArchitecture());
-                                }
-                            }
-
-                            // If we build w/ bXGEExport true, we didn't REALLY build at this point, 
-                            // so don't bother with doing the PrepTargetForDeployment call. 
-                            if ((Result == ECompilationResult.Succeeded) && (BuildConfiguration.bDeployAfterCompile == true) && (BuildConfiguration.bXGEExport == false) &&
-                                (UEBuildConfiguration.bGenerateManifest == false) && (UEBuildConfiguration.bGenerateExternalFileList == false) && (UEBuildConfiguration.bCleanProject == false))
-                            {
-                                var DeployHandler = UEBuildDeploy.GetBuildDeploy(CheckPlatform);
-                                if (DeployHandler != null)
-                                {
-                                    // We need to be able to identify the Target.Type we can derive it from the Arguments.
-                                    BuildConfiguration.bFlushBuildDirOnRemoteMac = false;
-                                    UEBuildTarget CheckTarget = UEBuildTarget.CreateTarget(Arguments);
-                                    CheckTarget.SetupGlobalEnvironment();
-                                    if ((CheckTarget.Rules.Type == TargetRules.TargetType.Game) || 
-                                        (CheckTarget.Rules.Type == TargetRules.TargetType.Server) || 
-                                        (CheckTarget.Rules.Type == TargetRules.TargetType.Client))
-                                    {
-                                        CheckTarget.AppName = CheckTarget.GameName;
-                                    }
-                                    DeployHandler.PrepTargetForDeployment(CheckTarget);
-                                }
-                            }
-                        }
-                    }
+					else if (bValidatePlatforms)
+					{
+						ValidatePlatforms(Arguments);
+					}
+					else
+					{
+						// Check if any third party headers are included from public engine headers.
+						if (bCheckThirdPartyHeaders)
+						{
+							ThirdPartyHeaderFinder.FindThirdPartyIncludes(CheckPlatform, CheckConfiguration);
+						}
+
+						// Build our project
+						if (Result == ECompilationResult.Succeeded)
+						{
+							if (UEBuildConfiguration.bPrepForDeployment == false)
+							{
+								// If we are only prepping for deployment, assume the build already occurred.
+								Result = RunUBT(Arguments);
+							}
+							else
+							{
+								var BuildPlatform = UEBuildPlatform.GetBuildPlatform(CheckPlatform, true);
+								if (BuildPlatform != null)
+								{
+									// Setup environment wasn't called, so set the flag
+									BuildConfiguration.bDeployAfterCompile = BuildPlatform.RequiresDeployPrepAfterCompile();
+									BuildConfiguration.PlatformIntermediateFolder = Path.Combine(BuildConfiguration.BaseIntermediateFolder, CheckPlatform.ToString(), BuildPlatform.GetActiveArchitecture());
+								}
+							}
+
+							// If we build w/ bXGEExport true, we didn't REALLY build at this point, 
+							// so don't bother with doing the PrepTargetForDeployment call. 
+							if ((Result == ECompilationResult.Succeeded) && (BuildConfiguration.bDeployAfterCompile == true) && (BuildConfiguration.bXGEExport == false) &&
+								(UEBuildConfiguration.bGenerateManifest == false) && (UEBuildConfiguration.bGenerateExternalFileList == false) && (UEBuildConfiguration.bCleanProject == false))
+							{
+								var DeployHandler = UEBuildDeploy.GetBuildDeploy(CheckPlatform);
+								if (DeployHandler != null)
+								{
+									// We need to be able to identify the Target.Type we can derive it from the Arguments.
+									BuildConfiguration.bFlushBuildDirOnRemoteMac = false;
+									UEBuildTarget CheckTarget = UEBuildTarget.CreateTarget(Arguments);
+									CheckTarget.SetupGlobalEnvironment();
+									if ((CheckTarget.Rules.Type == TargetRules.TargetType.Game) ||
+										(CheckTarget.Rules.Type == TargetRules.TargetType.Server) ||
+										(CheckTarget.Rules.Type == TargetRules.TargetType.Client))
+									{
+										CheckTarget.AppName = CheckTarget.GameName;
+									}
+									DeployHandler.PrepTargetForDeployment(CheckTarget);
+								}
+							}
+						}
+					}
                     // Print some performance info
                     var BuildDuration = (DateTime.UtcNow - StartTime - MutexWaitTime).TotalSeconds;
                     if (BuildConfiguration.bPrintPerformanceInfo)
@@ -1025,6 +1035,90 @@ namespace UnrealBuildTool
 
 
 
+		/// <summary>
+		/// Validates the various platforms to determine if they are ready for building
+		/// </summary>
+		public static void ValidatePlatforms(string[] Arguments)
+		{
+			List<UnrealTargetPlatform> Platforms = new List<UnrealTargetPlatform>(); 
+			foreach (var CurArgument in Arguments)
+			{
+				if (CurArgument.StartsWith("-"))
+				{
+					if (CurArgument.StartsWith("-Platforms=", StringComparison.InvariantCultureIgnoreCase))
+					{
+						// Parse the list... will be in Foo+Bar+New format
+						string PlatformList = CurArgument.Substring(11);
+						while (PlatformList.Length > 0)
+						{
+							string PlatformString = PlatformList;
+							Int32 PlusIdx = PlatformList.IndexOf("+");
+							if (PlusIdx != -1)
+							{
+								PlatformString = PlatformList.Substring(0, PlusIdx);
+								PlatformList = PlatformList.Substring(PlusIdx + 1);
+							}
+							else
+							{
+								// We are on the last platform... clear the list to exit the loop
+								PlatformList = "";
+							}
+
+							// Is the string a valid platform? If so, add it to the list
+							UnrealTargetPlatform SpecifiedPlatform = UnrealTargetPlatform.Unknown;
+							foreach (UnrealTargetPlatform PlatformParam in Enum.GetValues(typeof(UnrealTargetPlatform)))
+							{
+								if (PlatformString.Equals(PlatformParam.ToString(), StringComparison.InvariantCultureIgnoreCase))
+								{
+									SpecifiedPlatform = PlatformParam;
+									break;
+								}
+							}
+
+							if (SpecifiedPlatform != UnrealTargetPlatform.Unknown)
+							{
+								if (Platforms.Contains(SpecifiedPlatform) == false)
+								{
+									Platforms.Add(SpecifiedPlatform);
+								}
+							}
+							else
+							{
+								Log.TraceWarning("ValidatePlatforms invalid platform specified: {0}", PlatformString);
+							}
+						}
+					}
+					else switch (CurArgument.ToUpperInvariant())
+					{
+						case "-ALLPLATFORMS":
+							foreach (UnrealTargetPlatform platform in Enum.GetValues(typeof(UnrealTargetPlatform)))
+							{
+								if (platform != UnrealTargetPlatform.Unknown)
+								{
+									if (Platforms.Contains(platform) == false)
+									{
+										Platforms.Add(platform);
+									}
+								}
+							}
+							break;
+					}
+				}
+			}
+
+			foreach (UnrealTargetPlatform platform in Platforms)
+			{
+				var BuildPlatform = UEBuildPlatform.GetBuildPlatform(platform, true);
+				if (BuildPlatform != null && BuildPlatform.HasRequiredSDKsInstalled() == SDKStatus.Valid)
+				{
+					Console.WriteLine("##PlatformValidate: {0} VALID", platform.ToString());
+				}
+				else
+				{
+					Console.WriteLine("##PlatformValidate: {0} INVALID", platform.ToString());
+				}
+			}
+		}
 
         // @todo: Ideally get rid of RunUBT() and all of the Clear/Reset stuff!
         public static ECompilationResult RunUBT(string[] Arguments)
