Proper fix for UEBuildWinRT trying to check for Visual Studio on Mono

diff --git a/Engine/Source/Programs/UnrealBuildTool/WinRT/UEBuildWinRT.cs b/Engine/Source/Programs/UnrealBuildTool/WinRT/UEBuildWinRT.cs
--- a/Engine/Source/Programs/UnrealBuildTool/WinRT/UEBuildWinRT.cs
+++ b/Engine/Source/Programs/UnrealBuildTool/WinRT/UEBuildWinRT.cs
@@ -45,7 +45,7 @@ namespace UnrealBuildTool
          */
         protected override SDKStatus HasRequiredManualSDKInternal()
         {
-            return IsVisualStudioInstalled() && !Utils.IsRunningOnMono ? SDKStatus.Valid : SDKStatus.Invalid;
+            return !Utils.IsRunningOnMono && IsVisualStudioInstalled() ? SDKStatus.Valid : SDKStatus.Invalid;
         }
 
         /**
