Fix child UBT processes incorrectly detecting manual SDKs

diff --git a/Engine/Source/Programs/UnrealBuildTool/Configuration/UEBuildPlatform.cs b/Engine/Source/Programs/UnrealBuildTool/Configuration/UEBuildPlatform.cs
--- a/Engine/Source/Programs/UnrealBuildTool/Configuration/UEBuildPlatform.cs
+++ b/Engine/Source/Programs/UnrealBuildTool/Configuration/UEBuildPlatform.cs
@@ -828,6 +828,8 @@ namespace UnrealBuildTool
 
         protected static readonly string SDKRootEnvVar = "UE_SDKS_ROOT";
 
+        protected static string AutoSetupEnvVar = "AutoSDKSetup";
+
         /** 
          * Whether platform supports switching SDKs during runtime
          * 
@@ -898,6 +900,17 @@ namespace UnrealBuildTool
             return SDKPath;
         }
 
+        /**
+         * Because most ManualSDK determination depends on reading env vars, if this process is spawned by a process that ALREADY set up
+         * AutoSDKs then all the SDK env vars will exist, and we will spuriously detect a Manual SDK. (children inherit the environment of the parent process).
+         * Therefore we write out an env var to set in the command file (OutputEnvVars.txt) such that child processes can determine if their manual SDK detection
+         * is bogus.  Make it platform specific so that platforms can be in different states.
+         */
+        protected string GetPlatformAutoSDKSetupEnvVar()
+        {
+            return GetSDKTargetPlatformName() + AutoSetupEnvVar;
+        }
+
         /**
          * Gets currently installed version
          * 
@@ -1144,6 +1157,9 @@ namespace UnrealBuildTool
 
                     List<string> EnvVarNames = new List<string>();
                     List<string> EnvVarValues = new List<string>();
+
+                    bool bNeedsToWriteAutoSetupEnvVar = true;
+                    String PlatformSetupEnvVar = GetPlatformAutoSDKSetupEnvVar();
                     for (; ; )
                     {
                         string VariableString = Reader.ReadLine();
@@ -1165,11 +1181,15 @@ namespace UnrealBuildTool
                             PathRemoves.Add(Parts[1]);
                         }
                         else if (String.Compare(Parts[0], "addpath", true) == 0)
-                        {
+                        {                            
                             PathAdds.Add(Parts[1]);
                         }
                         else
                         {
+                            if (String.Compare(Parts[0], PlatformSetupEnvVar) == 0)
+                            {
+                                bNeedsToWriteAutoSetupEnvVar = false;
+                            }
                             // convenience for setup.bat writers.  Trim any accidental whitespace from var names/values.
                             EnvVarNames.Add(Parts[0].Trim());
                             EnvVarValues.Add(Parts[1].Trim());
@@ -1248,8 +1268,20 @@ namespace UnrealBuildTool
                     String ModifiedPath = String.Join(PathDelimiter, ModifiedPathVars);
                     Environment.SetEnvironmentVariable("PATH", ModifiedPath);
 
+                    Reader.Close();
+
+                    // write out env var command so any process using this commandfile will mark itself as having had autosdks set up.
+                    // avoids child processes spuriously detecting manualsdks.
+                    if (bNeedsToWriteAutoSetupEnvVar)
+                    {
+                        using (StreamWriter Writer = File.AppendText(EnvVarFile))
+                        {
+                            Writer.WriteLine("{0}=1", PlatformSetupEnvVar);
+                        }
+                    }
+
                     // make sure we know that we've modified the local environment, invalidating manual installs for this run.
-                    bProcessEnvSetupAutoSDK = true;
+                    bLocalProcessSetupAutoSDK = true;
 
                     return true;
                 }
@@ -1325,16 +1357,24 @@ namespace UnrealBuildTool
         private Int32 SDKCheckStatus = -1;
 
         // true if we've ever overridden the process's environment with AutoSDK data.  After that, manual installs cannot be considered valid ever again.
-        private bool bProcessEnvSetupAutoSDK = false;
+        private bool bLocalProcessSetupAutoSDK = false;
 
         protected bool HasSetupAutoSDK()
         {
-            return bProcessEnvSetupAutoSDK;
+            bool bParentProcessSetupAutoSDK = false;
+            String AutoSDKSetupVarName = GetPlatformAutoSDKSetupEnvVar();
+            String AutoSDKSetupVar = Environment.GetEnvironmentVariable(AutoSDKSetupVarName);
+            if (!String.IsNullOrEmpty(AutoSDKSetupVar))
+            {                
+                bParentProcessSetupAutoSDK = true;
+            }
+
+            return bLocalProcessSetupAutoSDK || bParentProcessSetupAutoSDK;
         }
 
         protected SDKStatus HasRequiredManualSDK()
         {
-            if (bProcessEnvSetupAutoSDK)
+            if (HasSetupAutoSDK())
             {
                 return SDKStatus.Invalid;
             }
