Fixed UBT Linux compile error

diff --git a/Engine/Source/Programs/UnrealBuildTool/IOS/UEBuildIOS.cs b/Engine/Source/Programs/UnrealBuildTool/IOS/UEBuildIOS.cs
--- a/Engine/Source/Programs/UnrealBuildTool/IOS/UEBuildIOS.cs
+++ b/Engine/Source/Programs/UnrealBuildTool/IOS/UEBuildIOS.cs
@@ -23,7 +23,7 @@ namespace UnrealBuildTool
 
         protected override SDKStatus HasRequiredManualSDKInternal()
         {
-			if (ExternalExecution.GetRuntimePlatform() != UnrealTargetPlatform.Mac)
+			if (!Utils.IsRunningOnMono)
 			{
 				// check to see if iTunes is installed
 				string dllPath = Microsoft.Win32.Registry.GetValue("HKEY_LOCAL_MACHINE\\SOFTWARE\\Wow6432Node\\Apple Inc.\\Apple Mobile Device Support\\Shared", "iTunesMobileDeviceDLL", null) as string;
