UBT: Fix unnecessary relinking on Linux.

diff --git a/Engine/Source/Programs/UnrealBuildTool/Configuration/UEBuildBinary.cs b/Engine/Source/Programs/UnrealBuildTool/Configuration/UEBuildBinary.cs
--- a/Engine/Source/Programs/UnrealBuildTool/Configuration/UEBuildBinary.cs
+++ b/Engine/Source/Programs/UnrealBuildTool/Configuration/UEBuildBinary.cs
@@ -580,7 +580,7 @@ namespace UnrealBuildTool
 					// Mark the link environment as cross-referenced.
 					BinaryLinkEnvironment.Config.bIsCrossReferenced = true;
 
-					if (BinaryLinkEnvironment.Config.Target.Platform != CPPTargetPlatform.Mac)
+					if (BinaryLinkEnvironment.Config.Target.Platform != CPPTargetPlatform.Mac && BinaryLinkEnvironment.Config.Target.Platform != CPPTargetPlatform.Linux)
 					{
 						// Create the import library.
 						OutputFiles.Add(BinaryLinkEnvironment.LinkExecutable(true));
