fix for PlatformerGame not building
#ue4
#mac

diff --git a/Engine/Source/Programs/UnrealBuildTool/Mac/MacToolChain.cs b/Engine/Source/Programs/UnrealBuildTool/Mac/MacToolChain.cs
--- a/Engine/Source/Programs/UnrealBuildTool/Mac/MacToolChain.cs
+++ b/Engine/Source/Programs/UnrealBuildTool/Mac/MacToolChain.cs
@@ -562,7 +562,7 @@ namespace UnrealBuildTool
 		{
 			string LibraryDir = Path.GetDirectoryName(Library);
 			string ExeDir = Path.GetDirectoryName(ExeAbsolutePath);
-			if ((Library.Contains("/Plugins/") || Library.Contains("/Binaries/ThirdParty/")) && LibraryDir != ExeDir)
+			if ((Library.Contains("/Plugins/") || Library.Contains("/Binaries/ThirdParty/")) && Library.EndsWith("dylib") && LibraryDir != ExeDir)
 			{
 				string RelativePath = Utils.MakePathRelativeTo(LibraryDir, ExeDir);
 				if (!RelativePath.Contains(LibraryDir) && !RPaths.Contains(RelativePath))
