- Moved the remote execution init to later, since it was causing too much setup for all UBT runs, not just IOS builds

diff --git a/Engine/Source/Programs/UnrealBuildTool/ToolChain/RemoteToolChain.cs b/Engine/Source/Programs/UnrealBuildTool/ToolChain/RemoteToolChain.cs
--- a/Engine/Source/Programs/UnrealBuildTool/ToolChain/RemoteToolChain.cs
+++ b/Engine/Source/Programs/UnrealBuildTool/ToolChain/RemoteToolChain.cs
@@ -19,13 +19,9 @@ namespace UnrealBuildTool
 		{
 			RemoteToolChainPlatform = InPlatform;
 
-			// make sure no error happened during init time
-			if (InitializeRemoteExecution() == 0)
-			{
-				// Register this tool chain for IOS
-				Log.TraceVerbose("        Registered for {0}", CPPPlatform.ToString());
-				UEToolChain.RegisterPlatformToolChain(CPPPlatform, this);
-			}
+			// Register this tool chain for IOS
+			Log.TraceVerbose("        Registered for {0}", CPPPlatform.ToString());
+			UEToolChain.RegisterPlatformToolChain(CPPPlatform, this);
 		}
 
 		/** These two variables will be loaded from XML config file in XmlConfigLoader.Init() */
@@ -263,47 +259,50 @@ namespace UnrealBuildTool
 					InitializationErrorCode = 99;
 				}
 
-				// we need the RemoteServerName and the Username to find the private key
-				ResolvedRSyncUsername = ResolveString(RSyncUsername, false);
-
-				// if the override path is set, just use it directly
-				if (!string.IsNullOrEmpty(SSHPrivateKeyOverridePath))
+				if (!bUseRPCUtil)
 				{
-					ResolvedSSHPrivateKey = ResolveString(SSHPrivateKeyOverridePath, true);
-					// make sure it exists
-					if (!File.Exists(ResolvedSSHPrivateKey))
+					// we need the RemoteServerName and the Username to find the private key
+					ResolvedRSyncUsername = ResolveString(RSyncUsername, false);
+
+					// if the override path is set, just use it directly
+					if (!string.IsNullOrEmpty(SSHPrivateKeyOverridePath))
 					{
-						throw new BuildException("An SSHKey override was specified [" + SSHPrivateKeyOverridePath + "] but it doesn't exist. Can't continue...");
+						ResolvedSSHPrivateKey = ResolveString(SSHPrivateKeyOverridePath, true);
+						// make sure it exists
+						if (!File.Exists(ResolvedSSHPrivateKey))
+						{
+							throw new BuildException("An SSHKey override was specified [" + SSHPrivateKeyOverridePath + "] but it doesn't exist. Can't continue...");
+						}
 					}
-				}
-				else
-				{
-					// all the places to look for a key
-					string[] Locations = new string[] {
-						Path.Combine(Environment.GetFolderPath(Environment.SpecialFolder.ApplicationData), "Unreal Engine", "UnrealBuildTool"),
-						Path.Combine(Environment.GetFolderPath(Environment.SpecialFolder.Personal), "Unreal Engine", "UnrealBuildTool"),
-						Path.Combine(BuildConfiguration.RelativeEnginePath, "Build", "NotForLicensees"),
-						Path.Combine(BuildConfiguration.RelativeEnginePath, "Build", "NoRedist"),
-						Path.Combine(BuildConfiguration.RelativeEnginePath, "Build"),
-					};
-
-					// look for a key file
-					foreach (string Location in Locations)
+					else
 					{
-						string KeyPath = Path.Combine(Location, "SSHKeys", RemoteServerName, ResolvedRSyncUsername, "RemoteToolChainPrivate.key");
-						if (File.Exists(KeyPath))
+						// all the places to look for a key
+						string[] Locations = new string[] {
+							Path.Combine(Environment.GetFolderPath(Environment.SpecialFolder.ApplicationData), "Unreal Engine", "UnrealBuildTool"),
+							Path.Combine(Environment.GetFolderPath(Environment.SpecialFolder.Personal), "Unreal Engine", "UnrealBuildTool"),
+							Path.Combine(BuildConfiguration.RelativeEnginePath, "Build", "NotForLicensees"),
+							Path.Combine(BuildConfiguration.RelativeEnginePath, "Build", "NoRedist"),
+							Path.Combine(BuildConfiguration.RelativeEnginePath, "Build"),
+						};
+
+						// look for a key file
+						foreach (string Location in Locations)
 						{
-							ResolvedSSHPrivateKey = KeyPath;
-							break;
+							string KeyPath = Path.Combine(Location, "SSHKeys", RemoteServerName, ResolvedRSyncUsername, "RemoteToolChainPrivate.key");
+							if (File.Exists(KeyPath))
+							{
+								ResolvedSSHPrivateKey = KeyPath;
+								break;
+							}
 						}
 					}
-				}
 
-				// resolve the rest of the strings
-				ResolvedRSyncExe = ResolveString(RSyncExe, true);
-				ResolvedSSHExe = ResolveString(SSHExe, true);
-				ResolvedRsyncAuthentication = ResolveString(RsyncAuthentication, false);
-				ResolvedSSHAuthentication = ResolveString(SSHAuthentication, false);
+					// resolve the rest of the strings
+					ResolvedRSyncExe = ResolveString(RSyncExe, true);
+					ResolvedSSHExe = ResolveString(SSHExe, true);
+					ResolvedRsyncAuthentication = ResolveString(RsyncAuthentication, false);
+					ResolvedSSHAuthentication = ResolveString(SSHAuthentication, false);
+				}
 
 				// start up remote communication and record if it succeeds
 				InitializationErrorCode = RPCUtilHelper.Initialize(RemoteServerName);
@@ -344,6 +343,13 @@ namespace UnrealBuildTool
 			return InitializationErrorCode;
         }
 
+		public override void SetUpGlobalEnvironment()
+		{
+			base.SetUpGlobalEnvironment();
+
+			// connect to server
+			InitializeRemoteExecution();
+		}
 
         /** Converts the passed in path from UBT host to compiler native format. */
         public override String ConvertPath(String OriginalPath)
