Force all builders into a clean state for autosdks

diff --git a/Engine/Source/Programs/UnrealBuildTool/Android/UEBuildAndroid.cs b/Engine/Source/Programs/UnrealBuildTool/Android/UEBuildAndroid.cs
--- a/Engine/Source/Programs/UnrealBuildTool/Android/UEBuildAndroid.cs
+++ b/Engine/Source/Programs/UnrealBuildTool/Android/UEBuildAndroid.cs
@@ -50,7 +50,7 @@ namespace UnrealBuildTool
 
         protected override String GetRequiredScriptVersionString()
         {
-            return "2.0";
+            return "3.0";
         }
 
         // prefer auto sdk on android as correct 'manual' sdk detection isn't great at the moment.
diff --git a/Engine/Source/Programs/UnrealBuildTool/Configuration/UEBuildPlatform.cs b/Engine/Source/Programs/UnrealBuildTool/Configuration/UEBuildPlatform.cs
--- a/Engine/Source/Programs/UnrealBuildTool/Configuration/UEBuildPlatform.cs
+++ b/Engine/Source/Programs/UnrealBuildTool/Configuration/UEBuildPlatform.cs
@@ -873,12 +873,13 @@ namespace UnrealBuildTool
         }
 
         /**
-        * Gets the version number of the SDK setup script itself.  The version in the base should ALWAYS be 1.0.  If you need to force a rebuild for a given platform, override this for the given platform.
+        * Gets the version number of the SDK setup script itself.  The version in the base should ALWAYS be the master revision from the last refactor.  
+        * If you need to force a rebuild for a given platform, override this for the given platform.
         * @return Setup script version
         */
         protected virtual String GetRequiredScriptVersionString()
         {
-            return "1.0";
+            return "3.0";
         }
 
         /** 
@@ -1401,7 +1402,7 @@ namespace UnrealBuildTool
         // platforms can choose if they prefer a correct the the AutoSDK install over the manual install.
         protected virtual bool PreferAutoSDK()
         {
-            return false;
+            return true;
         }
 
         // some platforms don't support parallel SDK installs.  AutoSDK on these platforms will
diff --git a/Engine/Source/Programs/UnrealBuildTool/HTML5/UEBuildHTML5.cs b/Engine/Source/Programs/UnrealBuildTool/HTML5/UEBuildHTML5.cs
--- a/Engine/Source/Programs/UnrealBuildTool/HTML5/UEBuildHTML5.cs
+++ b/Engine/Source/Programs/UnrealBuildTool/HTML5/UEBuildHTML5.cs
@@ -53,7 +53,7 @@ namespace UnrealBuildTool
 
 		protected override String GetRequiredScriptVersionString()
 		{
-			return "1.0";
+			return "3.0";
 		}
 		
 		// The current architecture - affects everything about how UBT operates on HTML5
diff --git a/Engine/Source/Programs/UnrealBuildTool/Linux/UEBuildLinux.cs b/Engine/Source/Programs/UnrealBuildTool/Linux/UEBuildLinux.cs
--- a/Engine/Source/Programs/UnrealBuildTool/Linux/UEBuildLinux.cs
+++ b/Engine/Source/Programs/UnrealBuildTool/Linux/UEBuildLinux.cs
@@ -66,7 +66,7 @@ namespace UnrealBuildTool
 
         protected override String GetRequiredScriptVersionString()
         {
-            return "2.0";
+            return "3.0";
         }
 
         protected override bool PreferAutoSDK()
