Fix XboxOne AutoSDK builds through UAT.

diff --git a/Engine/Source/Programs/UnrealBuildTool/Configuration/UEBuildPlatform.cs b/Engine/Source/Programs/UnrealBuildTool/Configuration/UEBuildPlatform.cs
--- a/Engine/Source/Programs/UnrealBuildTool/Configuration/UEBuildPlatform.cs
+++ b/Engine/Source/Programs/UnrealBuildTool/Configuration/UEBuildPlatform.cs
@@ -8,155 +8,155 @@ using System.IO;
 
 namespace UnrealBuildTool
 {
-	public enum SDKStatus
-	{
-		Valid,			// Desired SDK is installed and set up.
-		Invalid,		// Could not find the desired SDK, SDK setup failed, etc.		
-	};
-
-	public interface IUEBuildPlatform
-	{
-		/**
-		 * Whether the required external SDKs are installed for this platform
-		 */
-		SDKStatus HasRequiredSDKsInstalled();
-		
-		/**
-		 * If this platform can be compiled with XGE
-		 */
-		bool CanUseXGE();
-		
-		/**
-		 * Register the platform with the UEBuildPlatform class
-		 */
-		void RegisterBuildPlatform();
+    public enum SDKStatus
+    {
+        Valid,			// Desired SDK is installed and set up.
+        Invalid,		// Could not find the desired SDK, SDK setup failed, etc.		
+    };
+
+    public interface IUEBuildPlatform
+    {
+        /**
+         * Whether the required external SDKs are installed for this platform
+         */
+        SDKStatus HasRequiredSDKsInstalled();
+
+        /**
+         * If this platform can be compiled with XGE
+         */
+        bool CanUseXGE();
+
+        /**
+         * Register the platform with the UEBuildPlatform class
+         */
+        void RegisterBuildPlatform();
 
         /**
          * Attempt to set up AutoSDK for this paltform
          */
         void ManageAndValidateSDK();
-		
-		/**
-		 * Retrieve the CPPTargetPlatform for the given UnrealTargetPlatform
-		 *
-		 * @param InUnrealTargetPlatform The UnrealTargetPlatform being build
-		 * 
-		 * @return CPPTargetPlatform The CPPTargetPlatform to compile for
-		 */
-		CPPTargetPlatform GetCPPTargetPlatform(UnrealTargetPlatform InUnrealTargetPlatform);
-		
-		/**
-		 * Get the extension to use for the given binary type
-		 * 
-		 * @param InBinaryType The binary type being built
-		 * 
-		 * @return string The binary extension (i.e. 'exe' or 'dll')
-		 */
-		string GetBinaryExtension(UEBuildBinaryType InBinaryType);
-		
-		/**
-		 * Get the extension to use for debug info for the given binary type
-		 * 
-		 * @param InBinaryType The binary type being built
-		 * 
-		 * @return string The debug info extension (i.e. 'pdb')
-		 */
-		string GetDebugInfoExtension( UEBuildBinaryType InBinaryType );
-		
-		/**
-		 * Whether incremental linking should be used
-		 * 
-		 * @param InPlatform The CPPTargetPlatform being built
-		 * @param InConfiguration The CPPTargetConfiguration being built
-		 * 
-		 * @return bool true if incremental linking should be used, false if not
-		 */
-		bool ShouldUseIncrementalLinking(CPPTargetPlatform Platform, CPPTargetConfiguration Configuration);
-		
-		/**
-		 * Whether PDB files should be used
-		 * 
-		 * @param InPlatform The CPPTargetPlatform being built
-		 * @param InConfiguration The CPPTargetConfiguration being built
-		 * @param bInCreateDebugInfo true if debug info is getting create, false if not
-		 * 
-		 * @return bool true if PDB files should be used, false if not
-		 */
-		bool ShouldUsePDBFiles(CPPTargetPlatform Platform, CPPTargetConfiguration Configuration, bool bCreateDebugInfo);
-		
-		/**
-		 * Whether PCH files should be used
-		 * 
-		 * @param InPlatform The CPPTargetPlatform being built
-		 * @param InConfiguration The CPPTargetConfiguration being built
-		 * 
-		 * @return bool true if PCH files should be used, false if not
-		 */
-		bool ShouldUsePCHFiles(CPPTargetPlatform Platform, CPPTargetConfiguration Configuration);
-		
-		/**
-		 * Whether the editor should be built for this platform or not
-		 * 
-		 * @param InPlatform The UnrealTargetPlatform being built
-		 * @param InConfiguration The UnrealTargetConfiguration being built
-		 * @return bool true if the editor should be built, false if not
-		 */
-		bool ShouldNotBuildEditor(UnrealTargetPlatform InPlatform, UnrealTargetConfiguration InConfiguration);
-		
-		/**
-		 * Whether this build should support ONLY cooked data or not
-		 * 
-		 * @param InPlatform The UnrealTargetPlatform being built
-		 * @param InConfiguration The UnrealTargetConfiguration being built
-		 * @return bool true if the editor should be built, false if not
-		 */
-		bool BuildRequiresCookedData(UnrealTargetPlatform InPlatform, UnrealTargetConfiguration InConfiguration);
-		
-		/**
-		 * Whether the platform requires the extra UnityCPPWriter
-		 * This is used to add an extra file for UBT to get the #include dependencies from
-		 * 
-		 * @return bool true if it is required, false if not
-		 */
-		bool RequiresExtraUnityCPPWriter();
-		
-		/**
-		 * Whether this platform should build a monolithic binary
-		 */
-		bool ShouldCompileMonolithicBinary(UnrealTargetPlatform InPlatform);
-		
-		/**
-		 * Get a list of extra modules the platform requires.
-		 * This is to allow undisclosed platforms to add modules they need without exposing information about the platform.
-		 * 
-		 * @param Target The target being build
-		 * @param BuildTarget The UEBuildTarget getting build
-		 * @param PlatformExtraModules OUTPUT the list of extra modules the platform needs to add to the target
-		 */
-		void GetExtraModules(TargetInfo Target, UEBuildTarget BuildTarget, ref List<string> PlatformExtraModules);
-		
-		/**
-		 * Modify the newly created module passed in for this platform.
-		 * This is not required - but allows for hiding details of a
-		 * particular platform.
-		 * 
-		 * @param InModule The newly loaded module
-		 * @param Target The target being build
-		 */
-		void ModifyNewlyLoadedModule(UEBuildModule InModule, TargetInfo Target);
-		
-		/**
-		 * Setup the target environment for building
-		 * 
-		 * @param InBuildTarget The target being built
-		 */
-		void SetUpEnvironment(UEBuildTarget InBuildTarget);
-		
-		/**
-		 * Allow the platform to set an optional architecture
-		 */
-		string GetActiveArchitecture();
-		
+
+        /**
+         * Retrieve the CPPTargetPlatform for the given UnrealTargetPlatform
+         *
+         * @param InUnrealTargetPlatform The UnrealTargetPlatform being build
+         * 
+         * @return CPPTargetPlatform The CPPTargetPlatform to compile for
+         */
+        CPPTargetPlatform GetCPPTargetPlatform(UnrealTargetPlatform InUnrealTargetPlatform);
+
+        /**
+         * Get the extension to use for the given binary type
+         * 
+         * @param InBinaryType The binary type being built
+         * 
+         * @return string The binary extension (i.e. 'exe' or 'dll')
+         */
+        string GetBinaryExtension(UEBuildBinaryType InBinaryType);
+
+        /**
+         * Get the extension to use for debug info for the given binary type
+         * 
+         * @param InBinaryType The binary type being built
+         * 
+         * @return string The debug info extension (i.e. 'pdb')
+         */
+        string GetDebugInfoExtension(UEBuildBinaryType InBinaryType);
+
+        /**
+         * Whether incremental linking should be used
+         * 
+         * @param InPlatform The CPPTargetPlatform being built
+         * @param InConfiguration The CPPTargetConfiguration being built
+         * 
+         * @return bool true if incremental linking should be used, false if not
+         */
+        bool ShouldUseIncrementalLinking(CPPTargetPlatform Platform, CPPTargetConfiguration Configuration);
+
+        /**
+         * Whether PDB files should be used
+         * 
+         * @param InPlatform The CPPTargetPlatform being built
+         * @param InConfiguration The CPPTargetConfiguration being built
+         * @param bInCreateDebugInfo true if debug info is getting create, false if not
+         * 
+         * @return bool true if PDB files should be used, false if not
+         */
+        bool ShouldUsePDBFiles(CPPTargetPlatform Platform, CPPTargetConfiguration Configuration, bool bCreateDebugInfo);
+
+        /**
+         * Whether PCH files should be used
+         * 
+         * @param InPlatform The CPPTargetPlatform being built
+         * @param InConfiguration The CPPTargetConfiguration being built
+         * 
+         * @return bool true if PCH files should be used, false if not
+         */
+        bool ShouldUsePCHFiles(CPPTargetPlatform Platform, CPPTargetConfiguration Configuration);
+
+        /**
+         * Whether the editor should be built for this platform or not
+         * 
+         * @param InPlatform The UnrealTargetPlatform being built
+         * @param InConfiguration The UnrealTargetConfiguration being built
+         * @return bool true if the editor should be built, false if not
+         */
+        bool ShouldNotBuildEditor(UnrealTargetPlatform InPlatform, UnrealTargetConfiguration InConfiguration);
+
+        /**
+         * Whether this build should support ONLY cooked data or not
+         * 
+         * @param InPlatform The UnrealTargetPlatform being built
+         * @param InConfiguration The UnrealTargetConfiguration being built
+         * @return bool true if the editor should be built, false if not
+         */
+        bool BuildRequiresCookedData(UnrealTargetPlatform InPlatform, UnrealTargetConfiguration InConfiguration);
+
+        /**
+         * Whether the platform requires the extra UnityCPPWriter
+         * This is used to add an extra file for UBT to get the #include dependencies from
+         * 
+         * @return bool true if it is required, false if not
+         */
+        bool RequiresExtraUnityCPPWriter();
+
+        /**
+         * Whether this platform should build a monolithic binary
+         */
+        bool ShouldCompileMonolithicBinary(UnrealTargetPlatform InPlatform);
+
+        /**
+         * Get a list of extra modules the platform requires.
+         * This is to allow undisclosed platforms to add modules they need without exposing information about the platform.
+         * 
+         * @param Target The target being build
+         * @param BuildTarget The UEBuildTarget getting build
+         * @param PlatformExtraModules OUTPUT the list of extra modules the platform needs to add to the target
+         */
+        void GetExtraModules(TargetInfo Target, UEBuildTarget BuildTarget, ref List<string> PlatformExtraModules);
+
+        /**
+         * Modify the newly created module passed in for this platform.
+         * This is not required - but allows for hiding details of a
+         * particular platform.
+         * 
+         * @param InModule The newly loaded module
+         * @param Target The target being build
+         */
+        void ModifyNewlyLoadedModule(UEBuildModule InModule, TargetInfo Target);
+
+        /**
+         * Setup the target environment for building
+         * 
+         * @param InBuildTarget The target being built
+         */
+        void SetUpEnvironment(UEBuildTarget InBuildTarget);
+
+        /**
+         * Allow the platform to set an optional architecture
+         */
+        string GetActiveArchitecture();
+
         /**
          * Allow the platform to apply architecture-specific name according to its rules
          * 
@@ -166,213 +166,213 @@ namespace UnrealBuildTool
          */
         string ApplyArchitectureName(string BinaryName);
 
-		/**
-		 * Setup the configuration environment for building
-		 * 
-		 * @param InBuildTarget The target being built
-		 */
-		void SetUpConfigurationEnvironment(UEBuildTarget InBuildTarget);
-		
-		/**
-		 * Whether this platform should create debug information or not
-		 * 
-		 * @param InPlatform The UnrealTargetPlatform being built
-		 * @param InConfiguration The UnrealTargetConfiguration being built
-		 * 
-		 * @return bool true if debug info should be generated, false if not
-		 */
-		bool ShouldCreateDebugInfo(UnrealTargetPlatform Platform, UnrealTargetConfiguration Configuration);
-		
-		/**
-		 * Gives the platform a chance to 'override' the configuration settings 
-		 * that are overridden on calls to RunUBT.
-		 * 
-		 * @param InPlatform The UnrealTargetPlatform being built
-		 * @param InConfiguration The UnrealTargetConfiguration being built
-		 * 
-		 * @return bool true if debug info should be generated, false if not
-		 */
-		void ResetBuildConfiguration(UnrealTargetPlatform InPlatform, UnrealTargetConfiguration InConfiguration);
-		
-		/**
-		 * Validate the UEBuildConfiguration for this platform
-		 * This is called BEFORE calling UEBuildConfiguration to allow setting 
-		 * various fields used in that function such as CompileLeanAndMean...
-		 */
-		void ValidateUEBuildConfiguration();
-		
-		/**
-		 * Validate configuration for this platform
-		 * NOTE: This function can/will modify BuildConfiguration!
-		 * 
-		 * @param InPlatform The CPPTargetPlatform being built
-		 * @param InConfiguration The CPPTargetConfiguration being built
-		 * @param bInCreateDebugInfo true if debug info is getting create, false if not
-		 */
-		void ValidateBuildConfiguration(CPPTargetConfiguration Configuration, CPPTargetPlatform Platform, bool bCreateDebugInfo);
-		
-		/**
-		 * Return whether this platform has uniquely named binaries across multiple games
-		 */
-		bool HasUniqueBinaries();
-		
-		/**
-		 * Return whether we wish to have this platform's binaries in our builds
-		 */
-		bool IsBuildRequired();
-		
-		/**
-		 * Return whether we wish to have this platform's binaries in our CIS tests
-		 */
-		bool IsCISRequired();
-		
-		/// <summary>
-		/// Whether the build platform requires deployment prep
-		/// </summary>
-		/// <returns></returns>
-		bool RequiresDeployPrepAfterCompile();
-		
-		/**
-		 * Return all valid configurations for this platform
-		 * 
-		 *  Typically, this is always Debug, Development, and Shipping - but Test is a likely future addition for some platforms
-		 */
-		List<UnrealTargetConfiguration> GetConfigurations(UnrealTargetPlatform InUnrealTargetPlatform, bool bIncludeDebug);
-		
-		/**
-		 * Setup the binaries for this specific platform.
-		 * 
-		 * @param InBuildTarget The target being built
-		 */
-		void SetupBinaries(UEBuildTarget InBuildTarget);
-	}
-
-	public abstract partial class UEBuildPlatform : IUEBuildPlatform
-	{
-		public static Dictionary<UnrealTargetPlatform, IUEBuildPlatform> BuildPlatformDictionary = new Dictionary<UnrealTargetPlatform, IUEBuildPlatform>();
-
-		// a mapping of a group to the platforms in the group (ie, Microsoft contains Win32 and Win64)
-		static Dictionary<UnrealPlatformGroup, List<UnrealTargetPlatform>> PlatformGroupDictionary = new Dictionary<UnrealPlatformGroup, List<UnrealTargetPlatform>>();        
-        
+        /**
+         * Setup the configuration environment for building
+         * 
+         * @param InBuildTarget The target being built
+         */
+        void SetUpConfigurationEnvironment(UEBuildTarget InBuildTarget);
+
+        /**
+         * Whether this platform should create debug information or not
+         * 
+         * @param InPlatform The UnrealTargetPlatform being built
+         * @param InConfiguration The UnrealTargetConfiguration being built
+         * 
+         * @return bool true if debug info should be generated, false if not
+         */
+        bool ShouldCreateDebugInfo(UnrealTargetPlatform Platform, UnrealTargetConfiguration Configuration);
+
+        /**
+         * Gives the platform a chance to 'override' the configuration settings 
+         * that are overridden on calls to RunUBT.
+         * 
+         * @param InPlatform The UnrealTargetPlatform being built
+         * @param InConfiguration The UnrealTargetConfiguration being built
+         * 
+         * @return bool true if debug info should be generated, false if not
+         */
+        void ResetBuildConfiguration(UnrealTargetPlatform InPlatform, UnrealTargetConfiguration InConfiguration);
+
+        /**
+         * Validate the UEBuildConfiguration for this platform
+         * This is called BEFORE calling UEBuildConfiguration to allow setting 
+         * various fields used in that function such as CompileLeanAndMean...
+         */
+        void ValidateUEBuildConfiguration();
+
+        /**
+         * Validate configuration for this platform
+         * NOTE: This function can/will modify BuildConfiguration!
+         * 
+         * @param InPlatform The CPPTargetPlatform being built
+         * @param InConfiguration The CPPTargetConfiguration being built
+         * @param bInCreateDebugInfo true if debug info is getting create, false if not
+         */
+        void ValidateBuildConfiguration(CPPTargetConfiguration Configuration, CPPTargetPlatform Platform, bool bCreateDebugInfo);
+
+        /**
+         * Return whether this platform has uniquely named binaries across multiple games
+         */
+        bool HasUniqueBinaries();
+
+        /**
+         * Return whether we wish to have this platform's binaries in our builds
+         */
+        bool IsBuildRequired();
+
+        /**
+         * Return whether we wish to have this platform's binaries in our CIS tests
+         */
+        bool IsCISRequired();
+
+        /// <summary>
+        /// Whether the build platform requires deployment prep
+        /// </summary>
+        /// <returns></returns>
+        bool RequiresDeployPrepAfterCompile();
+
+        /**
+         * Return all valid configurations for this platform
+         * 
+         *  Typically, this is always Debug, Development, and Shipping - but Test is a likely future addition for some platforms
+         */
+        List<UnrealTargetConfiguration> GetConfigurations(UnrealTargetPlatform InUnrealTargetPlatform, bool bIncludeDebug);
+
+        /**
+         * Setup the binaries for this specific platform.
+         * 
+         * @param InBuildTarget The target being built
+         */
+        void SetupBinaries(UEBuildTarget InBuildTarget);
+    }
+
+    public abstract partial class UEBuildPlatform : IUEBuildPlatform
+    {
+        public static Dictionary<UnrealTargetPlatform, IUEBuildPlatform> BuildPlatformDictionary = new Dictionary<UnrealTargetPlatform, IUEBuildPlatform>();
+
+        // a mapping of a group to the platforms in the group (ie, Microsoft contains Win32 and Win64)
+        static Dictionary<UnrealPlatformGroup, List<UnrealTargetPlatform>> PlatformGroupDictionary = new Dictionary<UnrealPlatformGroup, List<UnrealTargetPlatform>>();
+
         /**
          * Attempt to convert a string to an UnrealTargetPlatform enum entry
          * 
          * @return UnrealTargetPlatform.Unknown on failure (the platform didn't match the enum)
          */
-		public static UnrealTargetPlatform ConvertStringToPlatform(string InPlatformName)
-		{
-			// special case x64, to not break anything
-			// @todo: Is it possible to remove this hack?
-			if (InPlatformName.Equals("X64", StringComparison.InvariantCultureIgnoreCase))
-			{
-				return UnrealTargetPlatform.Win64;
-			}
+        public static UnrealTargetPlatform ConvertStringToPlatform(string InPlatformName)
+        {
+            // special case x64, to not break anything
+            // @todo: Is it possible to remove this hack?
+            if (InPlatformName.Equals("X64", StringComparison.InvariantCultureIgnoreCase))
+            {
+                return UnrealTargetPlatform.Win64;
+            }
 
-			// we can't parse the string into an enum because Enum.Parse is case sensitive, so we loop over the enum
-			// looking for matches
-			foreach (string PlatformName in Enum.GetNames(typeof(UnrealTargetPlatform)))
-			{
-				if (InPlatformName.Equals(PlatformName, StringComparison.InvariantCultureIgnoreCase))
-				{
-					// convert the known good enum string back to the enum value
-					return (UnrealTargetPlatform)Enum.Parse(typeof(UnrealTargetPlatform), PlatformName);
-				}
-			}
-			return UnrealTargetPlatform.Unknown;
-		}
-
-		/**
-		 *	Register the given platforms UEBuildPlatform instance
-		 *	
-		 *	@param	InPlatform			The UnrealTargetPlatform to register with
-		 *	@param	InBuildPlatform		The UEBuildPlatform instance to use for the InPlatform
-		 */
-		public static void RegisterBuildPlatform(UnrealTargetPlatform InPlatform, UEBuildPlatform InBuildPlatform)
-		{
-			if (BuildPlatformDictionary.ContainsKey(InPlatform) == true)
-			{
-				Log.TraceWarning("RegisterBuildPlatform Warning: Registering build platform {0} for {1} when it is already set to {2}",
-					InBuildPlatform.ToString(), InPlatform.ToString(), BuildPlatformDictionary[InPlatform].ToString());
-				BuildPlatformDictionary[InPlatform] = InBuildPlatform;
-			}
-			else
-			{
-				BuildPlatformDictionary.Add(InPlatform, InBuildPlatform);
-			}
-		}
-
-		/**
-		 *	Unregister the given platform
-		 */
-		public static void UnregisterBuildPlatform(UnrealTargetPlatform InPlatform)
-		{
-			if (BuildPlatformDictionary.ContainsKey(InPlatform) == true)
-			{
-				BuildPlatformDictionary.Remove(InPlatform);
-			}
-		}
-
-		/**
-		 * Assign a platform as a member of the given group
-		 */
-		public static void RegisterPlatformWithGroup(UnrealTargetPlatform InPlatform, UnrealPlatformGroup InGroup)
-		{
-			// find or add the list of groups for this platform
-			PlatformGroupDictionary.GetOrAddNew(InGroup).Add(InPlatform);
-		}
-
-		/**
-		 * Retrieve the list of platforms in this group (if any)
-		 */
-		public static List<UnrealTargetPlatform> GetPlatformsInGroup(UnrealPlatformGroup InGroup)
-		{
-			List<UnrealTargetPlatform> PlatformList;
-			PlatformGroupDictionary.TryGetValue(InGroup, out PlatformList);
-			return PlatformList;
-		}
-
-		/**
-		 *	Retrieve the IUEBuildPlatform instance for the given TargetPlatform
-		 *	
-		 *	@param	InPlatform			The UnrealTargetPlatform being built
-		 *	@param	bInAllowFailure		If true, do not throw an exception and return null
-		 *	
-		 *	@return	UEBuildPlatform		The instance of the build platform
-		 */
-		public static IUEBuildPlatform GetBuildPlatform(UnrealTargetPlatform InPlatform, bool bInAllowFailure = false)
-		{
-			if (BuildPlatformDictionary.ContainsKey(InPlatform) == true)
-			{
-				return BuildPlatformDictionary[InPlatform];
-			}
-			if (bInAllowFailure == true)
-			{
-				return null;
-			}
-			throw new BuildException("GetBuildPlatform: No BuildPlatform found for {0}", InPlatform.ToString());
-		}
-
-		/**
-		 *	Retrieve the IUEBuildPlatform instance for the given CPPTargetPlatform
-		 *	
-		 *	@param	InPlatform			The CPPTargetPlatform being built
-		 *	@param	bInAllowFailure		If true, do not throw an exception and return null
-		 *	
-		 *	@return	UEBuildPlatform		The instance of the build platform
-		 */
-		public static IUEBuildPlatform GetBuildPlatformForCPPTargetPlatform(CPPTargetPlatform InPlatform, bool bInAllowFailure = false)
-		{
-			UnrealTargetPlatform UTPlatform = UEBuildTarget.CPPTargetPlatformToUnrealTargetPlatform(InPlatform);
-			if (BuildPlatformDictionary.ContainsKey(UTPlatform) == true)
-			{
-				return BuildPlatformDictionary[UTPlatform];
-			}
-			if (bInAllowFailure == true)
-			{
-				return null;
-			}
-			throw new BuildException("UEBuildPlatform::GetBuildPlatformForCPPTargetPlatform: No BuildPlatform found for {0}", InPlatform.ToString());
-		}
+            // we can't parse the string into an enum because Enum.Parse is case sensitive, so we loop over the enum
+            // looking for matches
+            foreach (string PlatformName in Enum.GetNames(typeof(UnrealTargetPlatform)))
+            {
+                if (InPlatformName.Equals(PlatformName, StringComparison.InvariantCultureIgnoreCase))
+                {
+                    // convert the known good enum string back to the enum value
+                    return (UnrealTargetPlatform)Enum.Parse(typeof(UnrealTargetPlatform), PlatformName);
+                }
+            }
+            return UnrealTargetPlatform.Unknown;
+        }
+
+        /**
+         *	Register the given platforms UEBuildPlatform instance
+         *	
+         *	@param	InPlatform			The UnrealTargetPlatform to register with
+         *	@param	InBuildPlatform		The UEBuildPlatform instance to use for the InPlatform
+         */
+        public static void RegisterBuildPlatform(UnrealTargetPlatform InPlatform, UEBuildPlatform InBuildPlatform)
+        {
+            if (BuildPlatformDictionary.ContainsKey(InPlatform) == true)
+            {
+                Log.TraceWarning("RegisterBuildPlatform Warning: Registering build platform {0} for {1} when it is already set to {2}",
+                    InBuildPlatform.ToString(), InPlatform.ToString(), BuildPlatformDictionary[InPlatform].ToString());
+                BuildPlatformDictionary[InPlatform] = InBuildPlatform;
+            }
+            else
+            {
+                BuildPlatformDictionary.Add(InPlatform, InBuildPlatform);
+            }
+        }
+
+        /**
+         *	Unregister the given platform
+         */
+        public static void UnregisterBuildPlatform(UnrealTargetPlatform InPlatform)
+        {
+            if (BuildPlatformDictionary.ContainsKey(InPlatform) == true)
+            {
+                BuildPlatformDictionary.Remove(InPlatform);
+            }
+        }
+
+        /**
+         * Assign a platform as a member of the given group
+         */
+        public static void RegisterPlatformWithGroup(UnrealTargetPlatform InPlatform, UnrealPlatformGroup InGroup)
+        {
+            // find or add the list of groups for this platform
+            PlatformGroupDictionary.GetOrAddNew(InGroup).Add(InPlatform);
+        }
+
+        /**
+         * Retrieve the list of platforms in this group (if any)
+         */
+        public static List<UnrealTargetPlatform> GetPlatformsInGroup(UnrealPlatformGroup InGroup)
+        {
+            List<UnrealTargetPlatform> PlatformList;
+            PlatformGroupDictionary.TryGetValue(InGroup, out PlatformList);
+            return PlatformList;
+        }
+
+        /**
+         *	Retrieve the IUEBuildPlatform instance for the given TargetPlatform
+         *	
+         *	@param	InPlatform			The UnrealTargetPlatform being built
+         *	@param	bInAllowFailure		If true, do not throw an exception and return null
+         *	
+         *	@return	UEBuildPlatform		The instance of the build platform
+         */
+        public static IUEBuildPlatform GetBuildPlatform(UnrealTargetPlatform InPlatform, bool bInAllowFailure = false)
+        {
+            if (BuildPlatformDictionary.ContainsKey(InPlatform) == true)
+            {
+                return BuildPlatformDictionary[InPlatform];
+            }
+            if (bInAllowFailure == true)
+            {
+                return null;
+            }
+            throw new BuildException("GetBuildPlatform: No BuildPlatform found for {0}", InPlatform.ToString());
+        }
+
+        /**
+         *	Retrieve the IUEBuildPlatform instance for the given CPPTargetPlatform
+         *	
+         *	@param	InPlatform			The CPPTargetPlatform being built
+         *	@param	bInAllowFailure		If true, do not throw an exception and return null
+         *	
+         *	@return	UEBuildPlatform		The instance of the build platform
+         */
+        public static IUEBuildPlatform GetBuildPlatformForCPPTargetPlatform(CPPTargetPlatform InPlatform, bool bInAllowFailure = false)
+        {
+            UnrealTargetPlatform UTPlatform = UEBuildTarget.CPPTargetPlatformToUnrealTargetPlatform(InPlatform);
+            if (BuildPlatformDictionary.ContainsKey(UTPlatform) == true)
+            {
+                return BuildPlatformDictionary[UTPlatform];
+            }
+            if (bInAllowFailure == true)
+            {
+                return null;
+            }
+            throw new BuildException("UEBuildPlatform::GetBuildPlatformForCPPTargetPlatform: No BuildPlatform found for {0}", InPlatform.ToString());
+        }
 
         /**
          *	Allow all registered build platforms to modify the newly created module 
@@ -384,15 +384,15 @@ namespace UnrealBuildTool
          *	@param	Only			If this is not unknown, then only run that platform
          */
         public static void PlatformModifyNewlyLoadedModule(UEBuildModule InModule, TargetInfo Target, UnrealTargetPlatform Only = UnrealTargetPlatform.Unknown)
-		{
-			foreach (var PlatformEntry in BuildPlatformDictionary)
-			{
+        {
+            foreach (var PlatformEntry in BuildPlatformDictionary)
+            {
                 if (Only == UnrealTargetPlatform.Unknown || PlatformEntry.Key == Only || PlatformEntry.Key == Target.Platform)
                 {
                     PlatformEntry.Value.ModifyNewlyLoadedModule(InModule, Target);
                 }
-			}
-		}                      
+            }
+        }
 
         /**
          * Returns the delimiter used to separate paths in the PATH environment variable for the platform we are executing on.
@@ -413,7 +413,7 @@ namespace UnrealBuildTool
             }
         }
 
-        
+
 
         /**
          *	If this platform can be compiled with XGE
@@ -423,9 +423,9 @@ namespace UnrealBuildTool
             return true;
         }
 
-		/**
-		 *	Register the platform with the UEBuildPlatform class
-		 */
+        /**
+         *	Register the platform with the UEBuildPlatform class
+         */
         public void RegisterBuildPlatform()
         {
             ManageAndValidateSDK();
@@ -434,189 +434,189 @@ namespace UnrealBuildTool
 
         protected abstract void RegisterBuildPlatformInternal();
 
-		/**
-		 *	Retrieve the CPPTargetPlatform for the given UnrealTargetPlatform
-		 *
-		 *	@param	InUnrealTargetPlatform		The UnrealTargetPlatform being build
-		 *	
-		 *	@return	CPPTargetPlatform			The CPPTargetPlatform to compile for
-		 */
-		public abstract CPPTargetPlatform GetCPPTargetPlatform(UnrealTargetPlatform InUnrealTargetPlatform);
-
-		/// <summary>
-		/// Return whether the given platform requires a monolithic build
-		/// </summary>
-		/// <param name="InPlatform">The platform of interest</param>
-		/// <param name="InConfiguration">The configuration of interest</param>
-		/// <returns></returns>
-		public static bool PlatformRequiresMonolithicBuilds(UnrealTargetPlatform InPlatform, UnrealTargetConfiguration InConfiguration)
-		{
-			// Some platforms require monolithic builds...
-			var BuildPlatform = GetBuildPlatform(InPlatform, true);
-			if (BuildPlatform != null)
-			{
-				return BuildPlatform.ShouldCompileMonolithicBinary(InPlatform);
-			}
-
-			// We assume it does not
-			return false;
-		}
-
-		/**
-		 *	Get the extension to use for the given binary type
-		 *	
-		 *	@param	InBinaryType		The binary type being built
-		 *	
-		 *	@return	string				The binary extension (i.e. 'exe' or 'dll')
-		 */
-		public virtual string GetBinaryExtension(UEBuildBinaryType InBinaryType)
-		{
-			throw new BuildException("GetBinaryExtensiton for {0} not handled in {1}", InBinaryType.ToString(), this.ToString());
-		}
-
-		/**
-		 *	Get the extension to use for debug info for the given binary type
-		 *	
-		 *	@param	InBinaryType		The binary type being built
-		 *	
-		 *	@return	string				The debug info extension (i.e. 'pdb')
-		 */
-		public virtual string GetDebugInfoExtension( UEBuildBinaryType InBinaryType )
-		{
-			throw new BuildException( "GetDebugInfoExtension for {0} not handled in {1}", InBinaryType.ToString(), this.ToString() );
-		}
-
-		/**
-		 *	Whether incremental linking should be used
-		 *	
-		 *	@param	InPlatform			The CPPTargetPlatform being built
-		 *	@param	InConfiguration		The CPPTargetConfiguration being built
-		 *	
-		 *	@return	bool	true if incremental linking should be used, false if not
-		 */
-		public virtual bool ShouldUseIncrementalLinking(CPPTargetPlatform Platform, CPPTargetConfiguration Configuration)
-		{
-			return false;
-		}
-
-		/**
-		 *	Whether PDB files should be used
-		 *	
-		 *	@param	InPlatform			The CPPTargetPlatform being built
-		 *	@param	InConfiguration		The CPPTargetConfiguration being built
-		 *	@param	bInCreateDebugInfo	true if debug info is getting create, false if not
-		 *	
-		 *	@return	bool	true if PDB files should be used, false if not
-		 */
-		public virtual bool ShouldUsePDBFiles(CPPTargetPlatform Platform, CPPTargetConfiguration Configuration, bool bCreateDebugInfo)
-		{
-			return false;
-		}
-
-		/**
-		 *	Whether PCH files should be used
-		 *	
-		 *	@param	InPlatform			The CPPTargetPlatform being built
-		 *	@param	InConfiguration		The CPPTargetConfiguration being built
-		 *	
-		 *	@return	bool				true if PCH files should be used, false if not
-		 */
-		public virtual bool ShouldUsePCHFiles(CPPTargetPlatform Platform, CPPTargetConfiguration Configuration)
-		{
-			return BuildConfiguration.bUsePCHFiles;
-		}
-
-		/**
-		 *	Whether the editor should be built for this platform or not
-		 *	
-		 *	@param	InPlatform		The UnrealTargetPlatform being built
-		 *	@param	InConfiguration	The UnrealTargetConfiguration being built
-		 *	@return	bool			true if the editor should be built, false if not
-		 */
-		public virtual bool ShouldNotBuildEditor(UnrealTargetPlatform InPlatform, UnrealTargetConfiguration InConfiguration)
-		{
-			return false;
-		}
-
-		/**
-		 *	Whether this build should support ONLY cooked data or not
-		 *	
-		 *	@param	InPlatform		The UnrealTargetPlatform being built
-		 *	@param	InConfiguration	The UnrealTargetConfiguration being built
-		 *	@return	bool			true if the editor should be built, false if not
-		 */
-		public virtual bool BuildRequiresCookedData(UnrealTargetPlatform InPlatform, UnrealTargetConfiguration InConfiguration)
-		{
-			return false;
-		}
-
-		/**
-		 *	Whether the platform requires the extra UnityCPPWriter
-		 *	This is used to add an extra file for UBT to get the #include dependencies from
-		 *	
-		 *	@return	bool	true if it is required, false if not
-		 */
-		public virtual bool RequiresExtraUnityCPPWriter()
-		{
-			return false;
-		}
-
-		/**
-		 * Whether this platform should build a monolithic binary
-		 */
-		public virtual bool ShouldCompileMonolithicBinary(UnrealTargetPlatform InPlatform)
-		{
-			return false;
-		}
-
-		/**
-		 *	Get a list of extra modules the platform requires.
-		 *	This is to allow undisclosed platforms to add modules they need without exposing information about the platform.
-		 *	
-		 *	@param	Target						The target being build
-		 *	@param	BuildTarget					The UEBuildTarget getting build
-		 *	@param	PlatformExtraModules		OUTPUT the list of extra modules the platform needs to add to the target
-		 */
-		public virtual void GetExtraModules(TargetInfo Target, UEBuildTarget BuildTarget, ref List<string> PlatformExtraModules)
-		{
-		}
-
-		/**
-		 *	Modify the newly created module passed in for this platform.
-		 *	This is not required - but allows for hiding details of a
-		 *	particular platform.
-		 *	
-		 *	@param	InModule		The newly loaded module
-		 *	@param	Target			The target being build
-		 */
-		public virtual void ModifyNewlyLoadedModule(UEBuildModule InModule, TargetInfo Target)
-		{
-		}
-
-		/**
-		 *	Setup the target environment for building
-		 *	
-		 *	@param	InBuildTarget		The target being built
-		 */
-		public abstract void SetUpEnvironment(UEBuildTarget InBuildTarget);
-
-		/**
-		 * Allow the platform to set an optional architecture
-		 */
-		public virtual string GetActiveArchitecture()
-		{
-			// by default, use an empty architecture (which is really just a modifer to the platform for some paths/names)
-			return "";
-		}
-
-		/**
-		 * Allow the platform to override the NMake output name
-		 */
-		public virtual string ModifyNMakeOutput(string ExeName)
-		{
-			// by default, use original
-			return ExeName;
-		}
+        /**
+         *	Retrieve the CPPTargetPlatform for the given UnrealTargetPlatform
+         *
+         *	@param	InUnrealTargetPlatform		The UnrealTargetPlatform being build
+         *	
+         *	@return	CPPTargetPlatform			The CPPTargetPlatform to compile for
+         */
+        public abstract CPPTargetPlatform GetCPPTargetPlatform(UnrealTargetPlatform InUnrealTargetPlatform);
+
+        /// <summary>
+        /// Return whether the given platform requires a monolithic build
+        /// </summary>
+        /// <param name="InPlatform">The platform of interest</param>
+        /// <param name="InConfiguration">The configuration of interest</param>
+        /// <returns></returns>
+        public static bool PlatformRequiresMonolithicBuilds(UnrealTargetPlatform InPlatform, UnrealTargetConfiguration InConfiguration)
+        {
+            // Some platforms require monolithic builds...
+            var BuildPlatform = GetBuildPlatform(InPlatform, true);
+            if (BuildPlatform != null)
+            {
+                return BuildPlatform.ShouldCompileMonolithicBinary(InPlatform);
+            }
+
+            // We assume it does not
+            return false;
+        }
+
+        /**
+         *	Get the extension to use for the given binary type
+         *	
+         *	@param	InBinaryType		The binary type being built
+         *	
+         *	@return	string				The binary extension (i.e. 'exe' or 'dll')
+         */
+        public virtual string GetBinaryExtension(UEBuildBinaryType InBinaryType)
+        {
+            throw new BuildException("GetBinaryExtensiton for {0} not handled in {1}", InBinaryType.ToString(), this.ToString());
+        }
+
+        /**
+         *	Get the extension to use for debug info for the given binary type
+         *	
+         *	@param	InBinaryType		The binary type being built
+         *	
+         *	@return	string				The debug info extension (i.e. 'pdb')
+         */
+        public virtual string GetDebugInfoExtension(UEBuildBinaryType InBinaryType)
+        {
+            throw new BuildException("GetDebugInfoExtension for {0} not handled in {1}", InBinaryType.ToString(), this.ToString());
+        }
+
+        /**
+         *	Whether incremental linking should be used
+         *	
+         *	@param	InPlatform			The CPPTargetPlatform being built
+         *	@param	InConfiguration		The CPPTargetConfiguration being built
+         *	
+         *	@return	bool	true if incremental linking should be used, false if not
+         */
+        public virtual bool ShouldUseIncrementalLinking(CPPTargetPlatform Platform, CPPTargetConfiguration Configuration)
+        {
+            return false;
+        }
+
+        /**
+         *	Whether PDB files should be used
+         *	
+         *	@param	InPlatform			The CPPTargetPlatform being built
+         *	@param	InConfiguration		The CPPTargetConfiguration being built
+         *	@param	bInCreateDebugInfo	true if debug info is getting create, false if not
+         *	
+         *	@return	bool	true if PDB files should be used, false if not
+         */
+        public virtual bool ShouldUsePDBFiles(CPPTargetPlatform Platform, CPPTargetConfiguration Configuration, bool bCreateDebugInfo)
+        {
+            return false;
+        }
+
+        /**
+         *	Whether PCH files should be used
+         *	
+         *	@param	InPlatform			The CPPTargetPlatform being built
+         *	@param	InConfiguration		The CPPTargetConfiguration being built
+         *	
+         *	@return	bool				true if PCH files should be used, false if not
+         */
+        public virtual bool ShouldUsePCHFiles(CPPTargetPlatform Platform, CPPTargetConfiguration Configuration)
+        {
+            return BuildConfiguration.bUsePCHFiles;
+        }
+
+        /**
+         *	Whether the editor should be built for this platform or not
+         *	
+         *	@param	InPlatform		The UnrealTargetPlatform being built
+         *	@param	InConfiguration	The UnrealTargetConfiguration being built
+         *	@return	bool			true if the editor should be built, false if not
+         */
+        public virtual bool ShouldNotBuildEditor(UnrealTargetPlatform InPlatform, UnrealTargetConfiguration InConfiguration)
+        {
+            return false;
+        }
+
+        /**
+         *	Whether this build should support ONLY cooked data or not
+         *	
+         *	@param	InPlatform		The UnrealTargetPlatform being built
+         *	@param	InConfiguration	The UnrealTargetConfiguration being built
+         *	@return	bool			true if the editor should be built, false if not
+         */
+        public virtual bool BuildRequiresCookedData(UnrealTargetPlatform InPlatform, UnrealTargetConfiguration InConfiguration)
+        {
+            return false;
+        }
+
+        /**
+         *	Whether the platform requires the extra UnityCPPWriter
+         *	This is used to add an extra file for UBT to get the #include dependencies from
+         *	
+         *	@return	bool	true if it is required, false if not
+         */
+        public virtual bool RequiresExtraUnityCPPWriter()
+        {
+            return false;
+        }
+
+        /**
+         * Whether this platform should build a monolithic binary
+         */
+        public virtual bool ShouldCompileMonolithicBinary(UnrealTargetPlatform InPlatform)
+        {
+            return false;
+        }
+
+        /**
+         *	Get a list of extra modules the platform requires.
+         *	This is to allow undisclosed platforms to add modules they need without exposing information about the platform.
+         *	
+         *	@param	Target						The target being build
+         *	@param	BuildTarget					The UEBuildTarget getting build
+         *	@param	PlatformExtraModules		OUTPUT the list of extra modules the platform needs to add to the target
+         */
+        public virtual void GetExtraModules(TargetInfo Target, UEBuildTarget BuildTarget, ref List<string> PlatformExtraModules)
+        {
+        }
+
+        /**
+         *	Modify the newly created module passed in for this platform.
+         *	This is not required - but allows for hiding details of a
+         *	particular platform.
+         *	
+         *	@param	InModule		The newly loaded module
+         *	@param	Target			The target being build
+         */
+        public virtual void ModifyNewlyLoadedModule(UEBuildModule InModule, TargetInfo Target)
+        {
+        }
+
+        /**
+         *	Setup the target environment for building
+         *	
+         *	@param	InBuildTarget		The target being built
+         */
+        public abstract void SetUpEnvironment(UEBuildTarget InBuildTarget);
+
+        /**
+         * Allow the platform to set an optional architecture
+         */
+        public virtual string GetActiveArchitecture()
+        {
+            // by default, use an empty architecture (which is really just a modifer to the platform for some paths/names)
+            return "";
+        }
+
+        /**
+         * Allow the platform to override the NMake output name
+         */
+        public virtual string ModifyNMakeOutput(string ExeName)
+        {
+            // by default, use original
+            return ExeName;
+        }
 
         /**
          * Allow the platform to apply architecture-specific name according to its rules
@@ -636,180 +636,180 @@ namespace UnrealBuildTool
          *	
          *	@param	InBuildTarget		The target being built
          */
-		public virtual void SetUpConfigurationEnvironment(UEBuildTarget InBuildTarget)
-		{
-			// Determine the C++ compile/link configuration based on the Unreal configuration.
-			CPPTargetConfiguration CompileConfiguration;
-			UnrealTargetConfiguration CheckConfig = InBuildTarget.Configuration;
-			//@todo SAS: Add a true Debug mode!
-			if (UnrealBuildTool.RunningRocket())
-			{
-				if (Utils.IsFileUnderDirectory( InBuildTarget.OutputPath, UnrealBuildTool.GetUProjectPath() ))
-				{
-					if (CheckConfig == UnrealTargetConfiguration.Debug)
-					{
-						CheckConfig = UnrealTargetConfiguration.DebugGame;
-					}
-				}
-				else
-				{
-					// Only Development and Shipping are supported for engine modules
-					if( CheckConfig != UnrealTargetConfiguration.Development && CheckConfig != UnrealTargetConfiguration.Shipping )
-					{
-						CheckConfig = UnrealTargetConfiguration.Development;
-					}
-				}
-			}
-			switch (CheckConfig)
-			{
-				default:
-				case UnrealTargetConfiguration.Debug:
-					CompileConfiguration = CPPTargetConfiguration.Debug;
-					if( BuildConfiguration.bDebugBuildsActuallyUseDebugCRT )
-					{
-						InBuildTarget.GlobalCompileEnvironment.Config.Definitions.Add("_DEBUG=1"); // the engine doesn't use this, but lots of 3rd party stuff does
-					}
-					else
-					{
-						InBuildTarget.GlobalCompileEnvironment.Config.Definitions.Add("NDEBUG=1"); // the engine doesn't use this, but lots of 3rd party stuff does
-					}
-					InBuildTarget.GlobalCompileEnvironment.Config.Definitions.Add("UE_BUILD_DEBUG=1");
-					break;
-				case UnrealTargetConfiguration.DebugGame:
-					// Individual game modules can be switched to be compiled in debug as necessary. By default, everything is compiled in development.
-				case UnrealTargetConfiguration.Development:
-					CompileConfiguration = CPPTargetConfiguration.Development;
-					InBuildTarget.GlobalCompileEnvironment.Config.Definitions.Add("NDEBUG=1"); // the engine doesn't use this, but lots of 3rd party stuff does
-					InBuildTarget.GlobalCompileEnvironment.Config.Definitions.Add("UE_BUILD_DEVELOPMENT=1");
-					break;
-				case UnrealTargetConfiguration.Shipping:
-					CompileConfiguration = CPPTargetConfiguration.Shipping;
-					InBuildTarget.GlobalCompileEnvironment.Config.Definitions.Add("NDEBUG=1"); // the engine doesn't use this, but lots of 3rd party stuff does
-					InBuildTarget.GlobalCompileEnvironment.Config.Definitions.Add("UE_BUILD_SHIPPING=1");
-					break;
-				case UnrealTargetConfiguration.Test:
-					CompileConfiguration = CPPTargetConfiguration.Shipping;
-					InBuildTarget.GlobalCompileEnvironment.Config.Definitions.Add("NDEBUG=1"); // the engine doesn't use this, but lots of 3rd party stuff does
-					InBuildTarget.GlobalCompileEnvironment.Config.Definitions.Add("UE_BUILD_TEST=1");
-					break;
-			}
-
-			// Set up the global C++ compilation and link environment.
-			InBuildTarget.GlobalCompileEnvironment.Config.Target.Configuration = CompileConfiguration;
-			InBuildTarget.GlobalLinkEnvironment.Config.Target.Configuration = CompileConfiguration;
-
-			// Create debug info based on the heuristics specified by the user.
-			InBuildTarget.GlobalCompileEnvironment.Config.bCreateDebugInfo =
-				!BuildConfiguration.bDisableDebugInfo && ShouldCreateDebugInfo(InBuildTarget.Platform, CheckConfig);
-			InBuildTarget.GlobalLinkEnvironment.Config.bCreateDebugInfo = InBuildTarget.GlobalCompileEnvironment.Config.bCreateDebugInfo;
-		}
-
-		/**
-		 *	Whether this platform should create debug information or not
-		 *	
-		 *	@param	InPlatform			The UnrealTargetPlatform being built
-		 *	@param	InConfiguration		The UnrealTargetConfiguration being built
-		 *	
-		 *	@return	bool				true if debug info should be generated, false if not
-		 */
-		public abstract bool ShouldCreateDebugInfo(UnrealTargetPlatform Platform, UnrealTargetConfiguration Configuration);
-
-		/**
-		 *	Gives the platform a chance to 'override' the configuration settings 
-		 *	that are overridden on calls to RunUBT.
-		 *	
-		 *	@param	InPlatform			The UnrealTargetPlatform being built
-		 *	@param	InConfiguration		The UnrealTargetConfiguration being built
-		 *	
-		 *	@return	bool				true if debug info should be generated, false if not
-		 */
-		public virtual void ResetBuildConfiguration(UnrealTargetPlatform InPlatform, UnrealTargetConfiguration InConfiguration)
-		{
-		}
-
-		/**
-		 *	Validate the UEBuildConfiguration for this platform
-		 *	This is called BEFORE calling UEBuildConfiguration to allow setting 
-		 *	various fields used in that function such as CompileLeanAndMean...
-		 */
-		public virtual void ValidateUEBuildConfiguration()
-		{
-		}
-
-		/**
-		 *	Validate configuration for this platform
-		 *	NOTE: This function can/will modify BuildConfiguration!
-		 *	
-		 *	@param	InPlatform			The CPPTargetPlatform being built
-		 *	@param	InConfiguration		The CPPTargetConfiguration being built
-		 *	@param	bInCreateDebugInfo	true if debug info is getting create, false if not
-		 */
-		public virtual void ValidateBuildConfiguration(CPPTargetConfiguration Configuration, CPPTargetPlatform Platform, bool bCreateDebugInfo)
-		{
-		}
-
-		/**
-		 *	Return whether this platform has uniquely named binaries across multiple games
-		 */
-		public virtual bool HasUniqueBinaries()
-		{
-			return true;
-		}
-
-		/**
-		 *	Return whether we wish to have this platform's binaries in our builds
-		 */
-		public virtual bool IsBuildRequired()
-		{
-			return true;
-		}
-
-		/**
-		 *	Return whether we wish to have this platform's binaries in our CIS tests
-		 */
-		public virtual bool IsCISRequired()
-		{
-			return true;
-		}
-
-		/// <summary>
-		/// Whether the build platform requires deployment prep
-		/// </summary>
-		/// <returns></returns>
-		public virtual bool RequiresDeployPrepAfterCompile()
-		{
-			return false;
-		}
-
-		/**
-		 *	Return all valid configurations for this platform
-		 *	
-		 *  Typically, this is always Debug, Development, and Shipping - but Test is a likely future addition for some platforms
-		 */
-		public virtual List<UnrealTargetConfiguration> GetConfigurations( UnrealTargetPlatform InUnrealTargetPlatform, bool bIncludeDebug )
-		{
-			List<UnrealTargetConfiguration> Configurations = new List<UnrealTargetConfiguration>()
+        public virtual void SetUpConfigurationEnvironment(UEBuildTarget InBuildTarget)
+        {
+            // Determine the C++ compile/link configuration based on the Unreal configuration.
+            CPPTargetConfiguration CompileConfiguration;
+            UnrealTargetConfiguration CheckConfig = InBuildTarget.Configuration;
+            //@todo SAS: Add a true Debug mode!
+            if (UnrealBuildTool.RunningRocket())
+            {
+                if (Utils.IsFileUnderDirectory(InBuildTarget.OutputPath, UnrealBuildTool.GetUProjectPath()))
+                {
+                    if (CheckConfig == UnrealTargetConfiguration.Debug)
+                    {
+                        CheckConfig = UnrealTargetConfiguration.DebugGame;
+                    }
+                }
+                else
+                {
+                    // Only Development and Shipping are supported for engine modules
+                    if (CheckConfig != UnrealTargetConfiguration.Development && CheckConfig != UnrealTargetConfiguration.Shipping)
+                    {
+                        CheckConfig = UnrealTargetConfiguration.Development;
+                    }
+                }
+            }
+            switch (CheckConfig)
+            {
+                default:
+                case UnrealTargetConfiguration.Debug:
+                    CompileConfiguration = CPPTargetConfiguration.Debug;
+                    if (BuildConfiguration.bDebugBuildsActuallyUseDebugCRT)
+                    {
+                        InBuildTarget.GlobalCompileEnvironment.Config.Definitions.Add("_DEBUG=1"); // the engine doesn't use this, but lots of 3rd party stuff does
+                    }
+                    else
+                    {
+                        InBuildTarget.GlobalCompileEnvironment.Config.Definitions.Add("NDEBUG=1"); // the engine doesn't use this, but lots of 3rd party stuff does
+                    }
+                    InBuildTarget.GlobalCompileEnvironment.Config.Definitions.Add("UE_BUILD_DEBUG=1");
+                    break;
+                case UnrealTargetConfiguration.DebugGame:
+                // Individual game modules can be switched to be compiled in debug as necessary. By default, everything is compiled in development.
+                case UnrealTargetConfiguration.Development:
+                    CompileConfiguration = CPPTargetConfiguration.Development;
+                    InBuildTarget.GlobalCompileEnvironment.Config.Definitions.Add("NDEBUG=1"); // the engine doesn't use this, but lots of 3rd party stuff does
+                    InBuildTarget.GlobalCompileEnvironment.Config.Definitions.Add("UE_BUILD_DEVELOPMENT=1");
+                    break;
+                case UnrealTargetConfiguration.Shipping:
+                    CompileConfiguration = CPPTargetConfiguration.Shipping;
+                    InBuildTarget.GlobalCompileEnvironment.Config.Definitions.Add("NDEBUG=1"); // the engine doesn't use this, but lots of 3rd party stuff does
+                    InBuildTarget.GlobalCompileEnvironment.Config.Definitions.Add("UE_BUILD_SHIPPING=1");
+                    break;
+                case UnrealTargetConfiguration.Test:
+                    CompileConfiguration = CPPTargetConfiguration.Shipping;
+                    InBuildTarget.GlobalCompileEnvironment.Config.Definitions.Add("NDEBUG=1"); // the engine doesn't use this, but lots of 3rd party stuff does
+                    InBuildTarget.GlobalCompileEnvironment.Config.Definitions.Add("UE_BUILD_TEST=1");
+                    break;
+            }
+
+            // Set up the global C++ compilation and link environment.
+            InBuildTarget.GlobalCompileEnvironment.Config.Target.Configuration = CompileConfiguration;
+            InBuildTarget.GlobalLinkEnvironment.Config.Target.Configuration = CompileConfiguration;
+
+            // Create debug info based on the heuristics specified by the user.
+            InBuildTarget.GlobalCompileEnvironment.Config.bCreateDebugInfo =
+                !BuildConfiguration.bDisableDebugInfo && ShouldCreateDebugInfo(InBuildTarget.Platform, CheckConfig);
+            InBuildTarget.GlobalLinkEnvironment.Config.bCreateDebugInfo = InBuildTarget.GlobalCompileEnvironment.Config.bCreateDebugInfo;
+        }
+
+        /**
+         *	Whether this platform should create debug information or not
+         *	
+         *	@param	InPlatform			The UnrealTargetPlatform being built
+         *	@param	InConfiguration		The UnrealTargetConfiguration being built
+         *	
+         *	@return	bool				true if debug info should be generated, false if not
+         */
+        public abstract bool ShouldCreateDebugInfo(UnrealTargetPlatform Platform, UnrealTargetConfiguration Configuration);
+
+        /**
+         *	Gives the platform a chance to 'override' the configuration settings 
+         *	that are overridden on calls to RunUBT.
+         *	
+         *	@param	InPlatform			The UnrealTargetPlatform being built
+         *	@param	InConfiguration		The UnrealTargetConfiguration being built
+         *	
+         *	@return	bool				true if debug info should be generated, false if not
+         */
+        public virtual void ResetBuildConfiguration(UnrealTargetPlatform InPlatform, UnrealTargetConfiguration InConfiguration)
+        {
+        }
+
+        /**
+         *	Validate the UEBuildConfiguration for this platform
+         *	This is called BEFORE calling UEBuildConfiguration to allow setting 
+         *	various fields used in that function such as CompileLeanAndMean...
+         */
+        public virtual void ValidateUEBuildConfiguration()
+        {
+        }
+
+        /**
+         *	Validate configuration for this platform
+         *	NOTE: This function can/will modify BuildConfiguration!
+         *	
+         *	@param	InPlatform			The CPPTargetPlatform being built
+         *	@param	InConfiguration		The CPPTargetConfiguration being built
+         *	@param	bInCreateDebugInfo	true if debug info is getting create, false if not
+         */
+        public virtual void ValidateBuildConfiguration(CPPTargetConfiguration Configuration, CPPTargetPlatform Platform, bool bCreateDebugInfo)
+        {
+        }
+
+        /**
+         *	Return whether this platform has uniquely named binaries across multiple games
+         */
+        public virtual bool HasUniqueBinaries()
+        {
+            return true;
+        }
+
+        /**
+         *	Return whether we wish to have this platform's binaries in our builds
+         */
+        public virtual bool IsBuildRequired()
+        {
+            return true;
+        }
+
+        /**
+         *	Return whether we wish to have this platform's binaries in our CIS tests
+         */
+        public virtual bool IsCISRequired()
+        {
+            return true;
+        }
+
+        /// <summary>
+        /// Whether the build platform requires deployment prep
+        /// </summary>
+        /// <returns></returns>
+        public virtual bool RequiresDeployPrepAfterCompile()
+        {
+            return false;
+        }
+
+        /**
+         *	Return all valid configurations for this platform
+         *	
+         *  Typically, this is always Debug, Development, and Shipping - but Test is a likely future addition for some platforms
+         */
+        public virtual List<UnrealTargetConfiguration> GetConfigurations(UnrealTargetPlatform InUnrealTargetPlatform, bool bIncludeDebug)
+        {
+            List<UnrealTargetConfiguration> Configurations = new List<UnrealTargetConfiguration>()
 			{
 				UnrealTargetConfiguration.Development, 
 			};
 
-			if( bIncludeDebug )
-			{
-				Configurations.Insert( 0, UnrealTargetConfiguration.Debug );
-			}
-
-			return Configurations;
-		}
-
-		/**
-		 *	Setup the binaries for this specific platform.
-		 *	
-		 *	@param	InBuildTarget		The target being built
-		 */
-		public virtual void SetupBinaries(UEBuildTarget InBuildTarget)
-		{
-		}
-	}
+            if (bIncludeDebug)
+            {
+                Configurations.Insert(0, UnrealTargetConfiguration.Debug);
+            }
+
+            return Configurations;
+        }
+
+        /**
+         *	Setup the binaries for this specific platform.
+         *	
+         *	@param	InBuildTarget		The target being built
+         */
+        public virtual void SetupBinaries(UEBuildTarget InBuildTarget)
+        {
+        }
+    }
 
     // AutoSDKs handling portion
     public abstract partial class UEBuildPlatform : IUEBuildPlatform
@@ -889,10 +889,10 @@ namespace UnrealBuildTool
          */
         protected string GetPathToPlatformAutoSDKs()
         {
-            string SDKPath = "";            
+            string SDKPath = "";
             string SDKRoot = Environment.GetEnvironmentVariable(SDKRootEnvVar);
             if (SDKRoot != null)
-            {                
+            {
                 if (SDKRoot != "")
                 {
                     SDKPath = Path.Combine(SDKRoot, "Host" + ExternalExecution.GetRuntimePlatform(), GetSDKTargetPlatformName());
@@ -1033,9 +1033,9 @@ namespace UnrealBuildTool
                     using (StreamWriter Writer = File.CreateText(VersionFilename))
                     {
                         Writer.WriteLine(InstalledSDKVersionString);
-                        Writer.WriteLine("ManualSDK");                        
+                        Writer.WriteLine("ManualSDK");
                     }
-                }                
+                }
             }
         }
 
@@ -1105,7 +1105,7 @@ namespace UnrealBuildTool
                     Console.WriteLine("Running {0} hook {1}", Hook, HookExe);
 
                     // run it
-                    Process HookProcess = new Process();                                        
+                    Process HookProcess = new Process();
                     HookProcess.StartInfo.WorkingDirectory = SDKDirectory;
                     HookProcess.StartInfo.FileName = HookExe;
                     HookProcess.StartInfo.Arguments = "";
@@ -1118,7 +1118,7 @@ namespace UnrealBuildTool
 
                     if (HookProcess.ExitCode != 0)
                     {
-                        Console.WriteLine("Hook exited uncleanly (returned {0}), considering it failed.", HookProcess.ExitCode);                        
+                        Console.WriteLine("Hook exited uncleanly (returned {0}), considering it failed.", HookProcess.ExitCode);
                         return false;
                     }
 
@@ -1182,7 +1182,7 @@ namespace UnrealBuildTool
                             PathRemoves.Add(Parts[1]);
                         }
                         else if (String.Compare(Parts[0], "addpath", true) == 0)
-                        {                            
+                        {
                             PathAdds.Add(Parts[1]);
                         }
                         else
@@ -1279,6 +1279,8 @@ namespace UnrealBuildTool
                         {
                             Writer.WriteLine("{0}=1", PlatformSetupEnvVar);
                         }
+                        // set the var in the local environment in case this process spawns any others.
+                        Environment.SetEnvironmentVariable(PlatformSetupEnvVar, "1");
                     }
 
                     // make sure we know that we've modified the local environment, invalidating manual installs for this run.
@@ -1317,7 +1319,7 @@ namespace UnrealBuildTool
                 {
                     File.Delete(EnvVarFile);
                 }
-            }            
+            }
         }
 
         /** 
@@ -1327,8 +1329,8 @@ namespace UnrealBuildTool
         protected SDKStatus HasRequiredAutoSDKInstalled()
         {
             if (PlatformSupportsAutoSDKs() && HasAutoSDKSystemEnabled())
-            {                
-                string AutoSDKRoot = GetPathToPlatformAutoSDKs();                
+            {
+                string AutoSDKRoot = GetPathToPlatformAutoSDKs();
                 if (AutoSDKRoot != "")
                 {
                     // check script version so script fixes can be propagated without touching every build machine's CurrentlyInstalled file manually.
@@ -1342,14 +1344,14 @@ namespace UnrealBuildTool
                     // check to make sure OutputEnvVars doesn't need regenerating
                     string EnvVarFile = Path.Combine(AutoSDKRoot, SDKEnvironmentVarsFile);
                     bool bEnvVarFileExists = File.Exists(EnvVarFile);
-                    
+
                     string CurrentSDKString;
                     if (bEnvVarFileExists && GetCurrentlyInstalledSDKString(AutoSDKRoot, out CurrentSDKString) && CurrentSDKString == GetRequiredSDKString() && bScriptVersionMatches)
                     {
                         return SDKStatus.Valid;
                     }
                     return SDKStatus.Invalid;
-                }                
+                }
             }
             return SDKStatus.Invalid;
         }
@@ -1361,16 +1363,20 @@ namespace UnrealBuildTool
         private bool bLocalProcessSetupAutoSDK = false;
 
         protected bool HasSetupAutoSDK()
+        {
+            return bLocalProcessSetupAutoSDK || HasParentProcessSetupAutoSDK();
+        }
+
+        protected bool HasParentProcessSetupAutoSDK()
         {
             bool bParentProcessSetupAutoSDK = false;
             String AutoSDKSetupVarName = GetPlatformAutoSDKSetupEnvVar();
             String AutoSDKSetupVar = Environment.GetEnvironmentVariable(AutoSDKSetupVarName);
             if (!String.IsNullOrEmpty(AutoSDKSetupVar))
-            {                
+            {
                 bParentProcessSetupAutoSDK = true;
             }
-
-            return bLocalProcessSetupAutoSDK || bParentProcessSetupAutoSDK;
+            return bParentProcessSetupAutoSDK;
         }
 
         protected SDKStatus HasRequiredManualSDK()
@@ -1462,9 +1468,9 @@ namespace UnrealBuildTool
                 }
 
                 // fixup process environment to match autosdk
-                SetupEnvironmentFromAutoSDK();                
+                SetupEnvironmentFromAutoSDK();
             }
-        }        
+        }
 
         #endregion
 
@@ -1485,13 +1491,13 @@ namespace UnrealBuildTool
         public virtual string GetSDKTargetPlatformName()
         {
             return "";
-        }  
+        }
 
         /* Whether or not we should try to automatically switch SDKs when asked to validate the platform's SDK state. */
-        public static bool bAllowAutoSDKSwitching = true;               
+        public static bool bAllowAutoSDKSwitching = true;
 
         public SDKStatus SetupEnvironmentFromAutoSDK()
-        {            
+        {
             string PlatformSDKRoot = GetPathToPlatformAutoSDKs();
 
             // load environment variables from current SDK
@@ -1500,7 +1506,7 @@ namespace UnrealBuildTool
                 Console.WriteLine("Failed to load environment from required SDK {0}", GetRequiredSDKString());
                 InvalidateCurrentlyInstalledAutoSDK();
                 return SDKStatus.Invalid;
-            }            
+            }
             return SDKStatus.Valid;
         }
 
@@ -1512,10 +1518,10 @@ namespace UnrealBuildTool
         {
             // avoid redundant potentially expensive SDK checks.
             if (SDKCheckStatus == -1)
-            {                
+            {
                 bool bHasManualSDK = HasRequiredManualSDK() == SDKStatus.Valid;
                 bool bHasAutoSDK = HasRequiredAutoSDKInstalled() == SDKStatus.Valid;
-                
+
                 // Per-Platform implementations can choose how to handle non-Auto SDK detection / handling.
                 SDKCheckStatus = (bHasManualSDK || bHasAutoSDK) ? 1 : 0;
             }
@@ -1524,8 +1530,10 @@ namespace UnrealBuildTool
 
         // Arbitrates between manual SDKs and setting up AutoSDK based on program options and platform preferences.
         public void ManageAndValidateSDK()
-        {            
-            if (bAllowAutoSDKSwitching)
+        {
+            // do not modify installed manifests if parent process has already set everything up.
+            // this avoids problems with determining IsAutoSDKSafe and doing an incorrect invalidate.
+            if (bAllowAutoSDKSwitching && !HasParentProcessSetupAutoSDK())
             {
                 bool bHasRequiredManualSDK = HasRequiredManualSDK() == SDKStatus.Valid;
                 if (IsAutoSDKSafe() && (PreferAutoSDK() || !bHasRequiredManualSDK))
@@ -1540,18 +1548,18 @@ namespace UnrealBuildTool
                 {
                     InvalidateCurrentlyInstalledAutoSDK();
                 }
+            }
 
-                if (BuildConfiguration.bPrintDebugInfo)
-                {
-                    PrintSDKInfo();
-                }
+            if (BuildConfiguration.bPrintDebugInfo)
+            {
+                PrintSDKInfo();
             }
         }
 
         public void PrintSDKInfo()
         {
-            
-            
+
+
             if (HasRequiredSDKsInstalled() == SDKStatus.Valid)
             {
                 bool bHasRequiredManualSDK = HasRequiredManualSDK() == SDKStatus.Valid;
