Fix for incorrectly-formed generated paths when passing a relative .uproject path to UBT.

#change UProjectPath made absolute when set via the command line.

diff --git a/Engine/Source/Programs/UnrealBuildTool/System/UnrealBuildTool.cs b/Engine/Source/Programs/UnrealBuildTool/System/UnrealBuildTool.cs
--- a/Engine/Source/Programs/UnrealBuildTool/System/UnrealBuildTool.cs
+++ b/Engine/Source/Programs/UnrealBuildTool/System/UnrealBuildTool.cs
@@ -487,27 +487,29 @@ namespace UnrealBuildTool
             {
                 bHasUProjectFile = true;
     
-                UProjectPath = Path.GetDirectoryName(UProjectFile);
+                string ProjectPath = Path.GetDirectoryName(UProjectFile);
                 OutGameName = Path.GetFileNameWithoutExtension(UProjectFile);
 
-                if (Path.IsPathRooted(UProjectPath) == false)
+                if (Path.IsPathRooted(ProjectPath) == false)
                 {
-                    if (Directory.Exists(UProjectPath) == false)
+                    if (Directory.Exists(ProjectPath) == false)
                     {
                         // If it is *not* rooted, then we potentially received a path that is 
                         // relative to Engine/Binaries/[PLATFORM] rather than Engine/Source
-                        if (UProjectPath.StartsWith("../") || UProjectPath.StartsWith("..\\"))
+                        if (ProjectPath.StartsWith("../") || ProjectPath.StartsWith("..\\"))
                         {
-                            string AlternativeProjectpath = UProjectPath.Substring(3);
+                            string AlternativeProjectpath = ProjectPath.Substring(3);
                             if (Directory.Exists(AlternativeProjectpath))
                             {
-                                UProjectPath = AlternativeProjectpath;
+                                ProjectPath = AlternativeProjectpath;
                                 UProjectFile = UProjectFile.Substring(3);
                                 Debug.Assert(UProjectFile.Length > 0);
                             }
                         }
                     }
                 }
+
+                UProjectPath = Path.GetFullPath(ProjectPath);
             }
 
             return true;
