Initial support for specifying additional files or folders that need to be copied to Mac app bundle in module's Build.cs file. For now only used for Mac CoreAudio's RadioEffectUnit plugin. More files and support for iOS bundles will follow soon.

diff --git a/Engine/Source/Programs/UnrealBuildTool/Configuration/UEBuildModule.cs b/Engine/Source/Programs/UnrealBuildTool/Configuration/UEBuildModule.cs
--- a/Engine/Source/Programs/UnrealBuildTool/Configuration/UEBuildModule.cs
+++ b/Engine/Source/Programs/UnrealBuildTool/Configuration/UEBuildModule.cs
@@ -150,6 +150,7 @@ namespace UnrealBuildTool
 		protected readonly List<string> PublicWeakFrameworks;
 		protected readonly List<UEBuildFramework> PublicAdditionalFrameworks;
 		protected readonly List<string> PublicAdditionalShadowFiles;
+		protected readonly List<UEBuildBundleResource> PublicAdditionalBundleResources;
 
 		/** Names of modules with header files that this module's public interface needs access to. */
 		protected List<string> PublicIncludePathModuleNames;
@@ -195,6 +196,7 @@ namespace UnrealBuildTool
 			IEnumerable<string> InPublicWeakFrameworks = null,
 			IEnumerable<UEBuildFramework> InPublicAdditionalFrameworks = null,
 			IEnumerable<string> InPublicAdditionalShadowFiles = null,
+			IEnumerable<UEBuildBundleResource> InPublicAdditionalBundleResources = null,
 			IEnumerable<string> InPublicIncludePathModuleNames = null,
 			IEnumerable<string> InPublicDependencyModuleNames = null,
 			IEnumerable<string> InPublicDelayLoadDLLs = null,
@@ -220,6 +222,7 @@ namespace UnrealBuildTool
 			PublicWeakFrameworks = ListFromOptionalEnumerableStringParameter(InPublicWeakFrameworks);
 			PublicAdditionalFrameworks = InPublicAdditionalFrameworks == null ? new List<UEBuildFramework>() : new List<UEBuildFramework>(InPublicAdditionalFrameworks);
 			PublicAdditionalShadowFiles = ListFromOptionalEnumerableStringParameter(InPublicAdditionalShadowFiles);
+			PublicAdditionalBundleResources = InPublicAdditionalBundleResources == null ? new List<UEBuildBundleResource>() : new List<UEBuildBundleResource>(InPublicAdditionalBundleResources);
 			PublicIncludePathModuleNames = ListFromOptionalEnumerableStringParameter( InPublicIncludePathModuleNames );
 			PublicDependencyModuleNames = ListFromOptionalEnumerableStringParameter(InPublicDependencyModuleNames);
 			PublicDelayLoadDLLs = ListFromOptionalEnumerableStringParameter(InPublicDelayLoadDLLs);
@@ -324,6 +327,19 @@ namespace UnrealBuildTool
 			PublicAdditionalShadowFiles.Add(FileName);
 		}
 
+		/** Adds a file or folder that will be copied to Mac or iOS app bundle */
+		public virtual void AddAdditionalBundleResource(UEBuildBundleResource Resource, bool bCheckForDuplicates = true)
+		{
+			if (bCheckForDuplicates == true)
+			{
+				if ( PublicAdditionalBundleResources.Contains( Resource ) )
+				{
+					return;
+				}
+			}
+			PublicAdditionalBundleResources.Add( Resource );
+		}
+
 		/** Adds a public include path module. */
 		public virtual void AddPublicIncludePathModule(string ModuleName, bool bCheckForDuplicates = true)
 		{
@@ -638,6 +654,7 @@ namespace UnrealBuildTool
 			ref List<string> WeakFrameworks,
 			ref List<UEBuildFramework> AdditionalFrameworks,
 			ref List<string> AdditionalShadowFiles,
+			ref List<UEBuildBundleResource> AdditionalBundleResources,
 			ref List<string> DelayLoadDLLs,
 			ref List<UEBuildBinary> BinaryDependencies,
 			ref Dictionary<UEBuildModule, bool> VisitedModules
@@ -678,7 +695,7 @@ namespace UnrealBuildTool
 							if (bIsExternalModule || bIsInStaticLibrary)
 							{
 								DependencyModule.SetupPublicLinkEnvironment(SourceBinary, ref LibraryPaths, ref AdditionalLibraries, ref Frameworks, ref WeakFrameworks,
-									ref AdditionalFrameworks, ref AdditionalShadowFiles, ref DelayLoadDLLs, ref BinaryDependencies, ref VisitedModules);
+									ref AdditionalFrameworks, ref AdditionalShadowFiles, ref AdditionalBundleResources, ref DelayLoadDLLs, ref BinaryDependencies, ref VisitedModules);
 							}
 						}
 					}
@@ -688,6 +705,7 @@ namespace UnrealBuildTool
 					AdditionalLibraries.AddRange(PublicAdditionalLibraries);
 					Frameworks.AddRange(PublicFrameworks);
 					WeakFrameworks.AddRange(PublicWeakFrameworks);
+					AdditionalBundleResources.AddRange(PublicAdditionalBundleResources);
 					// Remember the module so we can refer to it when needed
 					foreach ( var Framework in PublicAdditionalFrameworks )
 					{
@@ -709,7 +727,7 @@ namespace UnrealBuildTool
 		{
 			// Allow the module's public dependencies to add library paths and additional libraries to the link environment.
 			SetupPublicLinkEnvironment(Binary,ref LinkEnvironment.Config.LibraryPaths,ref LinkEnvironment.Config.AdditionalLibraries,ref LinkEnvironment.Config.Frameworks,ref LinkEnvironment.Config.WeakFrameworks,
-				ref LinkEnvironment.Config.AdditionalFrameworks,ref LinkEnvironment.Config.AdditionalShadowFiles,ref LinkEnvironment.Config.DelayLoadDLLs,ref BinaryDependencies,ref VisitedModules);
+				ref LinkEnvironment.Config.AdditionalFrameworks,ref LinkEnvironment.Config.AdditionalShadowFiles, ref LinkEnvironment.Config.AdditionalBundleResources,ref LinkEnvironment.Config.DelayLoadDLLs,ref BinaryDependencies,ref VisitedModules);
 
 			// Also allow the module's public and private dependencies to modify the link environment.
 			List<string> AllDependencyModuleNames = new List<string>(PrivateDependencyModuleNames);
@@ -719,7 +737,7 @@ namespace UnrealBuildTool
 			{
 				var DependencyModule = Target.GetModuleByName(DependencyName);
 				DependencyModule.SetupPublicLinkEnvironment(Binary,ref LinkEnvironment.Config.LibraryPaths,ref LinkEnvironment.Config.AdditionalLibraries,ref LinkEnvironment.Config.Frameworks,ref LinkEnvironment.Config.WeakFrameworks,
-					ref LinkEnvironment.Config.AdditionalFrameworks,ref LinkEnvironment.Config.AdditionalShadowFiles,ref LinkEnvironment.Config.DelayLoadDLLs,ref BinaryDependencies,ref VisitedModules);
+					ref LinkEnvironment.Config.AdditionalFrameworks,ref LinkEnvironment.Config.AdditionalShadowFiles, ref LinkEnvironment.Config.AdditionalBundleResources,ref LinkEnvironment.Config.DelayLoadDLLs,ref BinaryDependencies,ref VisitedModules);
 			}
 		}
 
@@ -791,6 +809,7 @@ namespace UnrealBuildTool
 			IEnumerable<string> InPublicWeakFrameworks = null,
 			IEnumerable<UEBuildFramework> InPublicAdditionalFrameworks = null,
 			IEnumerable<string> InPublicAdditionalShadowFiles = null,
+			IEnumerable<UEBuildBundleResource> InPublicAdditionalBundleResources = null,
 			IEnumerable<string> InPublicDependencyModuleNames = null,
 			IEnumerable<string> InPublicDelayLoadDLLs = null		// Delay loaded DLLs that should be setup when including this module
 			)
@@ -800,7 +819,7 @@ namespace UnrealBuildTool
 			InName:							InName,
 			InModuleDirectory:				InModuleDirectory,
 			InOutputDirectory:				InOutputDirectory,
-			InIsRedistributableOverride:		InIsRedistributableOverride,
+			InIsRedistributableOverride:	InIsRedistributableOverride,
 			InPublicDefinitions:			InPublicDefinitions,
 			InPublicIncludePaths:			InPublicIncludePaths,
 			InPublicSystemIncludePaths:		InPublicSystemIncludePaths,
@@ -810,6 +829,7 @@ namespace UnrealBuildTool
 			InPublicWeakFrameworks:			InPublicWeakFrameworks,
 			InPublicAdditionalFrameworks: 	InPublicAdditionalFrameworks,
 			InPublicAdditionalShadowFiles:  InPublicAdditionalShadowFiles,
+			InPublicAdditionalBundleResources: InPublicAdditionalBundleResources,
 			InPublicDependencyModuleNames:	InPublicDependencyModuleNames,
 			InPublicDelayLoadDLLs:			InPublicDelayLoadDLLs
 			)
@@ -987,6 +1007,7 @@ namespace UnrealBuildTool
 			IEnumerable<string> InPublicWeakFrameworks,
 			IEnumerable<UEBuildFramework> InPublicAdditionalFrameworks,
 			IEnumerable<string> InPublicAdditionalShadowFiles,
+			IEnumerable<UEBuildBundleResource> InPublicAdditionalBundleResources,
 			IEnumerable<string> InPrivateIncludePaths,
 			IEnumerable<string> InPrivateIncludePathModuleNames,
 			IEnumerable<string> InPrivateDependencyModuleNames,
@@ -1018,6 +1039,7 @@ namespace UnrealBuildTool
 					InPublicWeakFrameworks,
 					InPublicAdditionalFrameworks,
 					InPublicAdditionalShadowFiles,
+					InPublicAdditionalBundleResources,
 					InPublicIncludePathModuleNames,
 					InPublicDependencyModuleNames,
 					InPublicDelayLoadDLLs,
@@ -2035,6 +2057,7 @@ namespace UnrealBuildTool
 			IEnumerable<string> InPublicWeakFrameworks,
 			IEnumerable<UEBuildFramework> InPublicAdditionalFrameworks,
 			IEnumerable<string> InPublicAdditionalShadowFiles,
+			IEnumerable<UEBuildBundleResource> InPublicAdditionalBundleResources,
 			IEnumerable<string> InPrivateIncludePaths,
 			IEnumerable<string> InPrivateIncludePathModuleNames,
 			IEnumerable<string> InPrivateDependencyModuleNames,
@@ -2053,7 +2076,7 @@ namespace UnrealBuildTool
 			)
 			: base(InTarget,InName,InType,InModuleDirectory,InOutputDirectory,InIsRedistributableOverride,InIntelliSenseGatherer,
 			InSourceFiles,InPublicIncludePaths,InPublicSystemIncludePaths,InDefinitions,
-			InPublicIncludePathModuleNames,InPublicDependencyModuleNames,InPublicDelayLoadDLLs,InPublicAdditionalLibraries,InPublicFrameworks,InPublicWeakFrameworks,InPublicAdditionalFrameworks,InPublicAdditionalShadowFiles,
+			InPublicIncludePathModuleNames,InPublicDependencyModuleNames,InPublicDelayLoadDLLs,InPublicAdditionalLibraries,InPublicFrameworks,InPublicWeakFrameworks,InPublicAdditionalFrameworks,InPublicAdditionalShadowFiles,InPublicAdditionalBundleResources,
 			InPrivateIncludePaths,InPrivateIncludePathModuleNames,InPrivateDependencyModuleNames,
             InCircularlyReferencedDependentModules, InDynamicallyLoadedModuleNames, InPlatformSpecificDynamicallyLoadedModuleNames, InOptimizeCode,
 			InAllowSharedPCH, InSharedPCHHeaderFile, InUseRTTI, InEnableBufferSecurityChecks, InFasterWithoutUnity, InMinFilesUsingPrecompiledHeaderOverride,
@@ -2118,4 +2141,16 @@ namespace UnrealBuildTool
 		public string			FrameworkZipPath	= null;
 		public string			CopyBundledAssets	= null;
 	}
+
+	public class UEBuildBundleResource
+	{
+		public UEBuildBundleResource(string InResourcePath, string InBundleContentsSubdir = "Resources")
+		{
+			ResourcePath = InResourcePath;
+			BundleContentsSubdir = InBundleContentsSubdir;
+		}
+
+		public string ResourcePath			= null;
+		public string BundleContentsSubdir	= null;
+	}
 }
diff --git a/Engine/Source/Programs/UnrealBuildTool/Configuration/UEBuildTarget.cs b/Engine/Source/Programs/UnrealBuildTool/Configuration/UEBuildTarget.cs
--- a/Engine/Source/Programs/UnrealBuildTool/Configuration/UEBuildTarget.cs
+++ b/Engine/Source/Programs/UnrealBuildTool/Configuration/UEBuildTarget.cs
@@ -1777,6 +1777,7 @@ namespace UnrealBuildTool
 					InPublicFrameworks: new List<string>(),
 					InPublicWeakFrameworks: new List<string>(),
 					InPublicAdditionalShadowFiles: new List<string>(),
+					InPublicAdditionalBundleResources: new List<UEBuildBundleResource>(),
 					InPublicAdditionalFrameworks: new List<UEBuildFramework>(),
 					InPrivateIncludePaths: new List<string>(),
 					InPrivateIncludePathModuleNames: new List<string>(),
@@ -2729,6 +2730,7 @@ namespace UnrealBuildTool
 							InPublicWeakFrameworks: RulesObject.PublicWeakFrameworks,
 							InPublicAdditionalFrameworks: RulesObject.PublicAdditionalFrameworks,
 							InPublicAdditionalShadowFiles: RulesObject.PublicAdditionalShadowFiles,
+							InPublicAdditionalBundleResources: RulesObject.AdditionalBundleResources,
 							InPrivateIncludePaths: RulesObject.PrivateIncludePaths,
 							InPrivateIncludePathModuleNames: RulesObject.PrivateIncludePathModuleNames,
 							InPrivateDependencyModuleNames: RulesObject.PrivateDependencyModuleNames,
@@ -2767,6 +2769,7 @@ namespace UnrealBuildTool
 							InPublicWeakFrameworks: RulesObject.PublicWeakFrameworks,
 							InPublicAdditionalFrameworks: RulesObject.PublicAdditionalFrameworks,
 							InPublicAdditionalShadowFiles: RulesObject.PublicAdditionalShadowFiles,
+							InPublicAdditionalBundleResources: RulesObject.AdditionalBundleResources,
 							InPrivateIncludePaths: RulesObject.PrivateIncludePaths,
 							InPrivateIncludePathModuleNames: RulesObject.PrivateIncludePathModuleNames,
 							InPrivateDependencyModuleNames: RulesObject.PrivateDependencyModuleNames,
@@ -2802,6 +2805,7 @@ namespace UnrealBuildTool
 							InPublicWeakFrameworks: RulesObject.PublicWeakFrameworks,
 							InPublicAdditionalFrameworks: RulesObject.PublicAdditionalFrameworks,
 							InPublicAdditionalShadowFiles: RulesObject.PublicAdditionalShadowFiles,
+							InPublicAdditionalBundleResources: RulesObject.AdditionalBundleResources,
 							InPublicDependencyModuleNames: RulesObject.PublicDependencyModuleNames,
 							InPublicDelayLoadDLLs: RulesObject.PublicDelayLoadDLLs
 						);
diff --git a/Engine/Source/Programs/UnrealBuildTool/Mac/MacToolChain.cs b/Engine/Source/Programs/UnrealBuildTool/Mac/MacToolChain.cs
--- a/Engine/Source/Programs/UnrealBuildTool/Mac/MacToolChain.cs
+++ b/Engine/Source/Programs/UnrealBuildTool/Mac/MacToolChain.cs
@@ -53,8 +53,6 @@ namespace UnrealBuildTool
 
 		public List<string> BuiltBinaries = new List<string>();
 
-		private static List<FileItem> CrossReferencedFiles = new List<FileItem>();
-
 		public override void SetUpGlobalEnvironment()
 		{
 			base.SetUpGlobalEnvironment();
@@ -808,10 +806,6 @@ namespace UnrealBuildTool
 					// Tell linker to ignore unresolved symbols, so we don't have a problem with cross dependent dylibs that do not exist yet.
 					// This is fixed in later step, FixDylibDependencies.
 					LinkCommand += string.Format(" -undefined dynamic_lookup");
-					if (Utils.IsRunningOnMono)
-					{
-						CrossReferencedFiles.Add( OutputFile );
-					}
 				}
 			}
 
@@ -881,8 +875,8 @@ namespace UnrealBuildTool
 
 				FixDylibDepsScript.Close();
 
-				// Prepare a script that will be called by CreateAppBundle.sh to copy all necessary third party dylibs to the app bundle
-				// This is done this way as CreateAppBundle.sh script can be created before all the libraries are processed, so
+				// Prepare a script that will be called by FinalizeAppBundle.sh to copy all necessary third party dylibs to the app bundle
+				// This is done this way as FinalizeAppBundle.sh script can be created before all the libraries are processed, so
 				// at the time of it's creation we don't have the full list of third party dylibs all modules need.
 				string DylibCopyScriptPath = Path.Combine(LinkEnvironment.Config.IntermediateDirectory, "DylibCopy.sh");
 				if (!bHasWipedCopyDylibScript)
@@ -905,26 +899,24 @@ namespace UnrealBuildTool
 				}
 				DylibCopyScript.Close();
 
-				// For non-console application, prepare a script that will create the app bundle. It'll be run by CreateAppBundle action
+				// For non-console application, prepare a script that will create the app bundle. It'll be run by FinalizeAppBundle action
 				if (bIsBuildingAppBundle)
 				{
-					string CreateAppBundleScriptPath = Path.Combine(LinkEnvironment.Config.IntermediateDirectory, "CreateAppBundle.sh");
-					StreamWriter CreateAppBundleScript = File.CreateText(CreateAppBundleScriptPath);
-					AppendMacLine(CreateAppBundleScript, "#!/bin/sh");
+					string FinalizeAppBundleScriptPath = Path.Combine(LinkEnvironment.Config.IntermediateDirectory, "FinalizeAppBundle.sh");
+					StreamWriter FinalizeAppBundleScript = File.CreateText(FinalizeAppBundleScriptPath);
+					AppendMacLine(FinalizeAppBundleScript, "#!/bin/sh");
 					string BinariesPath = Path.GetDirectoryName(OutputFile.AbsolutePath);
 					BinariesPath = Path.GetDirectoryName(BinariesPath.Substring(0, BinariesPath.IndexOf(".app")));
-					AppendMacLine(CreateAppBundleScript, "cd \"{0}\"", ConvertPath(BinariesPath).Replace("$", "\\$"));
+					AppendMacLine(FinalizeAppBundleScript, "cd \"{0}\"", ConvertPath(BinariesPath).Replace("$", "\\$"));
 
 					string ExeName = Path.GetFileName(OutputFile.AbsolutePath);
 					string[] ExeNameParts = ExeName.Split('-');
 					string GameName = ExeNameParts[0];
 
-					AppendMacLine(CreateAppBundleScript, "mkdir -p \"{0}.app/Contents/MacOS\"", ExeName);
-					AppendMacLine(CreateAppBundleScript, "mkdir -p \"{0}.app/Contents/Resources/RadioEffectUnit.component/Contents/MacOS\"", ExeName);
-					AppendMacLine(CreateAppBundleScript, "mkdir -p \"{0}.app/Contents/Resources/RadioEffectUnit.component/Contents/Resources/English.lproj\"", ExeName);
+					AppendMacLine(FinalizeAppBundleScript, "mkdir -p \"{0}.app/Contents/MacOS\"", ExeName);
 
 					// Copy third party dylibs by calling additional script prepared earlier
-					AppendMacLine(CreateAppBundleScript, "sh \"{0}\" \"{1}\"", ConvertPath(DylibCopyScriptPath).Replace("$", "\\$"), ExeName);
+					AppendMacLine(FinalizeAppBundleScript, "sh \"{0}\" \"{1}\"", ConvertPath(DylibCopyScriptPath).Replace("$", "\\$"), ExeName);
 
 					string IconName = "UE4";
 					string BundleVersion = ExeName.StartsWith("UnrealEngineLauncher") ? LoadLauncherDisplayVersion() : LoadEngineDisplayVersion();
@@ -955,15 +947,11 @@ namespace UnrealBuildTool
 					{
 						CustomIcon = EngineSourcePath + "/Runtime/Launch/Resources/Mac/" + IconName + ".icns";
 					}
-					AppendMacLine(CreateAppBundleScript, "cp -f \"{0}\" \"{2}.app/Contents/Resources/{1}.icns\"", CustomIcon, IconName, ExeName);
+					AppendMacLine(FinalizeAppBundleScript, "cp -f \"{0}\" \"{2}.app/Contents/Resources/{1}.icns\"", CustomIcon, IconName, ExeName);
 
 					if (ExeName.StartsWith("UE4Editor"))
 					{
-					}
-
-					if (ExeName.StartsWith("UE4Editor"))
-					{
-						AppendMacLine(CreateAppBundleScript, "cp -f \"{0}/Runtime/Launch/Resources/Mac/UProject.icns\" \"{1}.app/Contents/Resources/UProject.icns\"", EngineSourcePath, ExeName);
+						AppendMacLine(FinalizeAppBundleScript, "cp -f \"{0}/Runtime/Launch/Resources/Mac/UProject.icns\" \"{1}.app/Contents/Resources/UProject.icns\"", EngineSourcePath, ExeName);
 					}
 
 					string InfoPlistFile = CustomResourcesPath + "/Info.plist";
@@ -971,36 +959,28 @@ namespace UnrealBuildTool
 					{
 						InfoPlistFile = EngineSourcePath + "/Runtime/Launch/Resources/Mac/" + (GameName.EndsWith("Editor") ? "Info-Editor.plist" : "Info.plist");
 					}
-					AppendMacLine(CreateAppBundleScript, "cp -f \"{0}\" \"{1}.app/Contents/Info.plist\"", InfoPlistFile, ExeName);
+					AppendMacLine(FinalizeAppBundleScript, "cp -f \"{0}\" \"{1}.app/Contents/Info.plist\"", InfoPlistFile, ExeName);
 
 					// Fix contents of Info.plist
-					AppendMacLine(CreateAppBundleScript, "sed -i \"\" \"s/\\${0}/{1}/g\" \"{1}.app/Contents/Info.plist\"", "{EXECUTABLE_NAME}", ExeName);
-					AppendMacLine(CreateAppBundleScript, "sed -i \"\" \"s/\\${0}/{1}/g\" \"{2}.app/Contents/Info.plist\"", "{APP_NAME}", GameName, ExeName);
-					AppendMacLine(CreateAppBundleScript, "sed -i \"\" \"s/\\${0}/{1}/g\" \"{2}.app/Contents/Info.plist\"", "{MACOSX_DEPLOYMENT_TARGET}", MinMacOSVersion, ExeName);
-					AppendMacLine(CreateAppBundleScript, "sed -i \"\" \"s/\\${0}/{1}/g\" \"{2}.app/Contents/Info.plist\"", "{ICON_NAME}", IconName, ExeName);
-					AppendMacLine(CreateAppBundleScript, "sed -i \"\" \"s/\\${0}/{1}/g\" \"{2}.app/Contents/Info.plist\"", "{BUNDLE_VERSION}", BundleVersion, ExeName);
+					AppendMacLine(FinalizeAppBundleScript, "sed -i \"\" \"s/\\${0}/{1}/g\" \"{1}.app/Contents/Info.plist\"", "{EXECUTABLE_NAME}", ExeName);
+					AppendMacLine(FinalizeAppBundleScript, "sed -i \"\" \"s/\\${0}/{1}/g\" \"{2}.app/Contents/Info.plist\"", "{APP_NAME}", GameName, ExeName);
+					AppendMacLine(FinalizeAppBundleScript, "sed -i \"\" \"s/\\${0}/{1}/g\" \"{2}.app/Contents/Info.plist\"", "{MACOSX_DEPLOYMENT_TARGET}", MinMacOSVersion, ExeName);
+					AppendMacLine(FinalizeAppBundleScript, "sed -i \"\" \"s/\\${0}/{1}/g\" \"{2}.app/Contents/Info.plist\"", "{ICON_NAME}", IconName, ExeName);
+					AppendMacLine(FinalizeAppBundleScript, "sed -i \"\" \"s/\\${0}/{1}/g\" \"{2}.app/Contents/Info.plist\"", "{BUNDLE_VERSION}", BundleVersion, ExeName);
 
 					// Generate PkgInfo file
-					AppendMacLine(CreateAppBundleScript, "echo 'echo -n \"APPL????\"' | bash > \"{0}.app/Contents/PkgInfo\"", ExeName);
-
-					// Copy RadioEffect component
-					AppendMacLine(CreateAppBundleScript, "cp -f \"{0}/ThirdParty/Mac/RadioEffectUnit/RadioEffectUnit.component/Contents/MacOS/RadioEffectUnit\" \"{1}.app/Contents/Resources/RadioEffectUnit.component/Contents/MacOS/RadioEffectUnit\"", EngineSourcePath, ExeName);
-					AppendMacLine(CreateAppBundleScript, "cp -f \"{0}/ThirdParty/Mac/RadioEffectUnit/RadioEffectUnit.component/Contents/Resources/English.lproj/Localizable.strings\" \"{1}.app/Contents/Resources/RadioEffectUnit.component/Contents/Resources/English.lproj/Localizable.strings\"", EngineSourcePath, ExeName);
-					AppendMacLine(CreateAppBundleScript, "cp -f \"{0}/ThirdParty/Mac/RadioEffectUnit/RadioEffectUnit.component/Contents/Info.plist\" \"{1}.app/Contents/Resources/RadioEffectUnit.component/Contents/Info.plist\"", EngineSourcePath, ExeName);
+					AppendMacLine(FinalizeAppBundleScript, "echo 'echo -n \"APPL????\"' | bash > \"{0}.app/Contents/PkgInfo\"", ExeName);
 
 					// Make sure OS X knows the bundle was updated
-					AppendMacLine(CreateAppBundleScript, "touch -c \"{0}.app\"", ExeName);
+					AppendMacLine(FinalizeAppBundleScript, "touch -c \"{0}.app\"", ExeName);
 
-					CreateAppBundleScript.Close();
+					FinalizeAppBundleScript.Close();
 
 					// copy over some needed files
 					// @todo mac: Make a QueueDirectoryForBatchUpload
 					QueueFileForBatchUpload(FileItem.GetItemByFullPath(Path.GetFullPath("../../Engine/Source/Runtime/Launch/Resources/Mac/UE4.icns")));
 					QueueFileForBatchUpload(FileItem.GetItemByFullPath(Path.GetFullPath("../../Engine/Source/Runtime/Launch/Resources/Mac/UProject.icns")));
 					QueueFileForBatchUpload(FileItem.GetItemByFullPath(Path.GetFullPath("../../Engine/Source/Runtime/Launch/Resources/Mac/Info.plist")));
-					QueueFileForBatchUpload(FileItem.GetItemByFullPath(Path.GetFullPath("../../Engine/Source/ThirdParty/Mac/RadioEffectUnit/RadioEffectUnit.component/Contents/MacOS/RadioEffectUnit")));
-					QueueFileForBatchUpload(FileItem.GetItemByFullPath(Path.GetFullPath("../../Engine/Source/ThirdParty/Mac/RadioEffectUnit/RadioEffectUnit.component/Contents/Resources/English.lproj/Localizable.strings")));
-					QueueFileForBatchUpload(FileItem.GetItemByFullPath(Path.GetFullPath("../../Engine/Source/ThirdParty/Mac/RadioEffectUnit/RadioEffectUnit.component/Contents/Info.plist")));
 					QueueFileForBatchUpload(FileItem.GetItemByFullPath(Path.Combine(LinkEnvironment.Config.IntermediateDirectory, "DylibCopy.sh")));
 				}
 			}
@@ -1100,7 +1080,7 @@ namespace UnrealBuildTool
 		 * 
 		 * @param Executable FileItem describing the executable to generate app bundle for
 		 */
-		FileItem CreateAppBundle(LinkEnvironment LinkEnvironment, FileItem Executable, FileItem FixDylibOutputFile)
+		FileItem FinalizeAppBundle(LinkEnvironment LinkEnvironment, FileItem Executable, FileItem FixDylibOutputFile)
 		{
 			// Make a file item for the source and destination files
 			string FullDestPath = Executable.AbsolutePath.Substring(0, Executable.AbsolutePath.IndexOf(".app") + 4);
@@ -1108,30 +1088,63 @@ namespace UnrealBuildTool
 			FileItem RemoteDestFile = LocalToRemoteFileItem(DestFile, false);
 
 			// Make the compile action
-			Action CreateAppBundleAction = new Action(ActionType.CreateAppBundle);
-			CreateAppBundleAction.WorkingDirectory = Path.GetFullPath(".");
-			CreateAppBundleAction.CommandPath = "/bin/sh";
-			CreateAppBundleAction.CommandDescription = "";
+			Action FinalizeAppBundleAction = new Action(ActionType.CreateAppBundle);
+			FinalizeAppBundleAction.WorkingDirectory = Path.GetFullPath(".");
+			FinalizeAppBundleAction.CommandPath = "/bin/sh";
+			FinalizeAppBundleAction.CommandDescription = "";
 
 			// make path to the script
-			FileItem BundleScript = FileItem.GetItemByFullPath(Path.Combine(LinkEnvironment.Config.IntermediateDirectory, "CreateAppBundle.sh"));
+			FileItem BundleScript = FileItem.GetItemByFullPath(Path.Combine(LinkEnvironment.Config.IntermediateDirectory, "FinalizeAppBundle.sh"));
 			FileItem RemoteBundleScript = LocalToRemoteFileItem(BundleScript, true);
 
-		
 			if (ExternalExecution.GetRuntimePlatform() != UnrealTargetPlatform.Mac)
 			{
-				CreateAppBundleAction.ActionHandler = new Action.BlockingActionHandler(RPCUtilHelper.RPCActionHandler);
+				FinalizeAppBundleAction.ActionHandler = new Action.BlockingActionHandler(RPCUtilHelper.RPCActionHandler);
 			}
 
-			CreateAppBundleAction.CommandArguments = "\"" + RemoteBundleScript.AbsolutePath + "\"";
-			CreateAppBundleAction.PrerequisiteItems.Add(FixDylibOutputFile);
-			CreateAppBundleAction.ProducedItems.Add(RemoteDestFile);
-			CreateAppBundleAction.StatusDescription = string.Format("Creating app bundle: {0}.app", Path.GetFileName(Executable.AbsolutePath));
-			CreateAppBundleAction.bCanExecuteRemotely = false;
+			FinalizeAppBundleAction.CommandArguments = "\"" + RemoteBundleScript.AbsolutePath + "\"";
+			FinalizeAppBundleAction.PrerequisiteItems.Add(FixDylibOutputFile);
+			FinalizeAppBundleAction.ProducedItems.Add(RemoteDestFile);
+			FinalizeAppBundleAction.StatusDescription = string.Format("Finalizing app bundle: {0}.app", Path.GetFileName(Executable.AbsolutePath));
+			FinalizeAppBundleAction.bCanExecuteRemotely = false;
 
 			return RemoteDestFile;
 		}
 
+		FileItem CopyBundleResource(UEBuildBundleResource Resource, FileItem Executable)
+		{
+			Action CopyAction = new Action(ActionType.CreateAppBundle);
+			CopyAction.WorkingDirectory = Path.GetFullPath(".");
+			CopyAction.CommandPath = "/bin/sh";
+			CopyAction.CommandDescription = "";
+
+			string BundlePath = Executable.AbsolutePath.Substring(0, Executable.AbsolutePath.IndexOf(".app") + 4);
+			string SourcePath = Path.Combine(CopyAction.WorkingDirectory, Resource.ResourcePath);
+			string TargetPath = Path.Combine(BundlePath, "Contents", Resource.BundleContentsSubdir, Path.GetFileName(Resource.ResourcePath));
+
+			FileItem TargetItem = LocalToRemoteFileItem(FileItem.GetItemByPath(TargetPath), false);
+
+			CopyAction.CommandArguments = string.Format("-c 'cp -f -R \"{0}\" \"{1}\"'", SourcePath, TargetPath);
+			CopyAction.PrerequisiteItems.Add(Executable);
+			CopyAction.ProducedItems.Add(TargetItem);
+			CopyAction.StatusDescription = string.Format("Copying {0} to app bundle", Path.GetFileName(Resource.ResourcePath));
+			CopyAction.bCanExecuteRemotely = false;
+
+			if (Directory.Exists(Resource.ResourcePath))
+			{
+				foreach (string ResourceFile in Directory.GetFiles(Resource.ResourcePath, "*", SearchOption.AllDirectories))
+				{
+					QueueFileForBatchUpload(FileItem.GetItemByFullPath(Path.GetFullPath(ResourceFile)));
+				}
+			}
+			else
+			{
+				QueueFileForBatchUpload(FileItem.GetItemByFullPath(SourcePath));
+			}
+
+			return TargetItem;
+		}
+
 		static public void SetupBundleDependencies(List<UEBuildBinary> Binaries, string GameName)
 		{
 			foreach (UEBuildBinary Binary in Binaries)
@@ -1211,15 +1224,27 @@ namespace UnrealBuildTool
 				}
 			}
 
+			foreach (UEBuildBundleResource Resource in BinaryLinkEnvironment.Config.AdditionalBundleResources)
+			{
+				if (Directory.Exists(Resource.ResourcePath))
+				{
+					foreach (string ResourceFile in Directory.GetFiles(Resource.ResourcePath, "*", SearchOption.AllDirectories))
+					{
+						Manifest.AddFileName(Path.Combine(BundleContentsDirectory, Resource.BundleContentsSubdir, ResourceFile.Substring(Path.GetDirectoryName(Resource.ResourcePath).Length + 1)));
+					}
+				}
+				else
+				{
+					Manifest.AddFileName(Path.Combine(BundleContentsDirectory, Resource.BundleContentsSubdir, Path.GetFileName(Resource.ResourcePath)));
+				}
+			}
+
 			if (Binary.Config.Type == UEBuildBinaryType.Executable)
 			{
 				// And we also need all the resources
 				Manifest.AddFileName(BundleContentsDirectory + "Info.plist");
 				Manifest.AddFileName(BundleContentsDirectory + "PkgInfo");
 				Manifest.AddFileName(BundleContentsDirectory + "Resources/UE4.icns");
-				Manifest.AddFileName(BundleContentsDirectory + "Resources/RadioEffectUnit.component/Contents/MacOS/RadioEffectUnit");
-				Manifest.AddFileName(BundleContentsDirectory + "Resources/RadioEffectUnit.component/Contents/Resources/English.lproj/Localizable.strings");
-				Manifest.AddFileName(BundleContentsDirectory + "Resources/RadioEffectUnit.component/Contents/Info.plist");
 
 				if (Binary.Config.TargetName.StartsWith("UE4Editor"))
 				{
@@ -1280,24 +1305,33 @@ namespace UnrealBuildTool
 			base.PostBuildSync(Target);
 		}
 
-		public override ICollection<FileItem> PostBuild (FileItem Executable, LinkEnvironment BinaryLinkEnvironment)
+		public override ICollection<FileItem> PostBuild(FileItem Executable, LinkEnvironment BinaryLinkEnvironment)
 		{
-			var OutputFiles = base.PostBuild (Executable, BinaryLinkEnvironment);
+			var OutputFiles = base.PostBuild(Executable, BinaryLinkEnvironment);
 
-			// if building for Mac on a Mac, use actions to finalize the builds (otherwise, we use Deploy)
-			if (ExternalExecution.GetRuntimePlatform () != UnrealTargetPlatform.Mac) {
+			// If building for Mac on a Mac, use actions to finalize the builds (otherwise, we use Deploy)
+			if (ExternalExecution.GetRuntimePlatform() != UnrealTargetPlatform.Mac)
+			{
 				return OutputFiles;
 			}
 
-			if (BinaryLinkEnvironment.Config.bIsBuildingDLL || BinaryLinkEnvironment.Config.bIsBuildingLibrary) {
+			foreach (UEBuildBundleResource Resource in BinaryLinkEnvironment.Config.AdditionalBundleResources)
+			{
+				OutputFiles.Add(CopyBundleResource(Resource, Executable));
+			}
+
+			if (BinaryLinkEnvironment.Config.bIsBuildingDLL || BinaryLinkEnvironment.Config.bIsBuildingLibrary)
+			{
 				return OutputFiles;
 			}
+
 			FileItem FixDylibOutputFile = FixDylibDependencies(BinaryLinkEnvironment, Executable);
 			OutputFiles.Add(FixDylibOutputFile);
-			if (BinaryLinkEnvironment.Config.bIsBuildingConsoleApplication)
-				return OutputFiles;
+			if (!BinaryLinkEnvironment.Config.bIsBuildingConsoleApplication)
+			{
+				OutputFiles.Add(FinalizeAppBundle(BinaryLinkEnvironment, Executable, FixDylibOutputFile));
+			}
 
-			OutputFiles.Add(CreateAppBundle(BinaryLinkEnvironment, Executable, FixDylibOutputFile));
 			return OutputFiles;
 		}
 	};
diff --git a/Engine/Source/Programs/UnrealBuildTool/System/LinkEnvironment.cs b/Engine/Source/Programs/UnrealBuildTool/System/LinkEnvironment.cs
--- a/Engine/Source/Programs/UnrealBuildTool/System/LinkEnvironment.cs
+++ b/Engine/Source/Programs/UnrealBuildTool/System/LinkEnvironment.cs
@@ -44,6 +44,9 @@ namespace UnrealBuildTool
 		public List<string> Frameworks = new List<string>();
 		public List<string> WeakFrameworks = new List<string>();
 
+		/** iOS/Mac resources that should be copied to the app bundle */
+		public List<UEBuildBundleResource> AdditionalBundleResources = new List<UEBuildBundleResource>();
+
 		/**
 		 * A list of the dynamically linked libraries that shouldn't be loaded until they are first called
 		 * into.
@@ -109,6 +112,7 @@ namespace UnrealBuildTool
 			AdditionalShadowFiles.AddRange( InCopyEnvironment.AdditionalShadowFiles );
 			AdditionalFrameworks.AddRange(InCopyEnvironment.AdditionalFrameworks);
 			WeakFrameworks.AddRange(InCopyEnvironment.WeakFrameworks);
+			AdditionalBundleResources.AddRange(InCopyEnvironment.AdditionalBundleResources);
 			DelayLoadDLLs.AddRange(InCopyEnvironment.DelayLoadDLLs);
 			AdditionalArguments = InCopyEnvironment.AdditionalArguments;
 			bCreateDebugInfo = InCopyEnvironment.bCreateDebugInfo;
diff --git a/Engine/Source/Programs/UnrealBuildTool/System/RulesCompiler.cs b/Engine/Source/Programs/UnrealBuildTool/System/RulesCompiler.cs
--- a/Engine/Source/Programs/UnrealBuildTool/System/RulesCompiler.cs
+++ b/Engine/Source/Programs/UnrealBuildTool/System/RulesCompiler.cs
@@ -201,6 +201,9 @@ namespace UnrealBuildTool
 		/// List of addition frameworks - typically used for External (third party) modules on Mac and iOS
 		public List<UEBuildFramework> PublicAdditionalFrameworks = new List<UEBuildFramework>();
 
+		/// List of addition resources that should be copied to the app bundle for Mac or iOS
+		public List<UEBuildBundleResource> AdditionalBundleResources = new List<UEBuildBundleResource>();
+
 		/// For builds that execute on a remote machine (e.g. iOS), this list contains additional files that
 		/// need to be copied over in order for the app to link successfully.  Source/header files and PCHs are
 		/// automatically copied.  Usually this is simply a list of precompiled third party library dependencies.
diff --git a/Engine/Source/Programs/UnrealBuildTool/UnrealBuildTool_Mono.csproj b/Engine/Source/Programs/UnrealBuildTool/UnrealBuildTool_Mono.csproj
--- a/Engine/Source/Programs/UnrealBuildTool/UnrealBuildTool_Mono.csproj
+++ b/Engine/Source/Programs/UnrealBuildTool/UnrealBuildTool_Mono.csproj
@@ -38,7 +38,6 @@
     <DebugType>pdbonly</DebugType>
     <Optimize>true</Optimize>
     <OutputPath>..\..\..\Binaries\DotNET\</OutputPath>
-    <DefineConstants></DefineConstants>
     <Optimize>true</Optimize>
     <DebugType>pdbonly</DebugType>
     <PlatformTarget>AnyCPU</PlatformTarget>
@@ -51,7 +50,6 @@
     <DebugType>pdbonly</DebugType>
     <Optimize>true</Optimize>
     <OutputPath>..\..\..\Binaries\DotNET\</OutputPath>
-    <DefineConstants></DefineConstants>
     <Optimize>true</Optimize>
     <DebugType>pdbonly</DebugType>
     <PlatformTarget>x86</PlatformTarget>
diff --git a/Engine/Source/Runtime/Mac/CoreAudio/CoreAudio.Build.cs b/Engine/Source/Runtime/Mac/CoreAudio/CoreAudio.Build.cs
--- a/Engine/Source/Runtime/Mac/CoreAudio/CoreAudio.Build.cs
+++ b/Engine/Source/Runtime/Mac/CoreAudio/CoreAudio.Build.cs
@@ -23,5 +23,7 @@ public class CoreAudio : ModuleRules
 			);
 
 		PublicFrameworks.AddRange(new string[] { "CoreAudio", "AudioToolbox", "AudioUnit" });
+
+		AdditionalBundleResources.Add(new UEBuildBundleResource("ThirdParty/Mac/RadioEffectUnit/RadioEffectUnit.component"));
 	}
 }
