Fixed copying of folders to Mac bundle resources

diff --git a/Engine/Source/Programs/UnrealBuildTool/Mac/MacToolChain.cs b/Engine/Source/Programs/UnrealBuildTool/Mac/MacToolChain.cs
--- a/Engine/Source/Programs/UnrealBuildTool/Mac/MacToolChain.cs
+++ b/Engine/Source/Programs/UnrealBuildTool/Mac/MacToolChain.cs
@@ -1124,7 +1124,7 @@ namespace UnrealBuildTool
 
 			FileItem TargetItem = LocalToRemoteFileItem(FileItem.GetItemByPath(TargetPath), false);
 
-			CopyAction.CommandArguments = string.Format("-c 'cp -f -R \"{0}\" \"{1}\"'", SourcePath, TargetPath);
+			CopyAction.CommandArguments = string.Format("-c 'cp -f -R \"{0}\" \"{1}\"'", SourcePath, Path.GetDirectoryName(TargetPath) + "/");
 			CopyAction.PrerequisiteItems.Add(Executable);
 			CopyAction.ProducedItems.Add(TargetItem);
 			CopyAction.StatusDescription = string.Format("Copying {0} to app bundle", Path.GetFileName(Resource.ResourcePath));
