Make sure Mac app bundle's Contents/Resources is created before copying icon file to it

diff --git a/Engine/Source/Programs/UnrealBuildTool/Mac/MacToolChain.cs b/Engine/Source/Programs/UnrealBuildTool/Mac/MacToolChain.cs
--- a/Engine/Source/Programs/UnrealBuildTool/Mac/MacToolChain.cs
+++ b/Engine/Source/Programs/UnrealBuildTool/Mac/MacToolChain.cs
@@ -914,6 +914,7 @@ namespace UnrealBuildTool
 					string GameName = ExeNameParts[0];
 
 					AppendMacLine(FinalizeAppBundleScript, "mkdir -p \"{0}.app/Contents/MacOS\"", ExeName);
+					AppendMacLine(FinalizeAppBundleScript, "mkdir -p \"{0}.app/Contents/Resources\"", ExeName);
 
 					// Copy third party dylibs by calling additional script prepared earlier
 					AppendMacLine(FinalizeAppBundleScript, "sh \"{0}\" \"{1}\"", ConvertPath(DylibCopyScriptPath).Replace("$", "\\$"), ExeName);
