#UBT

     #TTP 342328 (HTML5: Packaging Blueprint projects fails)

   -  A compiler generated file was missing from the manifest, fixed
   -  clean up -  instead of calling static functions, call virtual function on tool chain, move platform specific code to derived tool chains.

diff --git a/Engine/Source/Programs/UnrealBuildTool/Configuration/UEBuildTarget.cs b/Engine/Source/Programs/UnrealBuildTool/Configuration/UEBuildTarget.cs
--- a/Engine/Source/Programs/UnrealBuildTool/Configuration/UEBuildTarget.cs
+++ b/Engine/Source/Programs/UnrealBuildTool/Configuration/UEBuildTarget.cs
@@ -1128,7 +1128,7 @@ namespace UnrealBuildTool
 		}
 
 		/** Generates a public manifest file for writing out */
-        public void GenerateManifest(List<UEBuildBinary> Binaries, CPPTargetPlatform Platform, List<string> SpecialRocketLibFilesThatAreBuildProducts)
+        public void GenerateManifest(IUEToolChain ToolChain, List<UEBuildBinary> Binaries, CPPTargetPlatform Platform, List<string> SpecialRocketLibFilesThatAreBuildProducts)
 		{
 			string ManifestPath;
 			if (UnrealBuildTool.RunningRocket())
@@ -1177,24 +1177,7 @@ namespace UnrealBuildTool
 					Manifest.AddBinaryNames(UEBuildBinary.GetAdditionalConsoleAppPath(Binary.Config.OutputFilePath), DebugInfoExtension);
 				}
 
-				if (TargetPlatform == UnrealTargetPlatform.Mac)
-				{
-					// Add all the resources and third party dylibs stored in app bundle
-					MacToolChain.AddAppBundleContentsToManifest(ref Manifest, Binary);
-				}
-				else if (TargetPlatform == UnrealTargetPlatform.IOS)
-				{
-					IOSToolChain.AddStubToManifest(ref Manifest, Binary);
-				}
-				// ok, this is pretty awful, we want the import libraries that go with the editor, only on the PC
-				else if (UnrealBuildTool.BuildingRocket() &&
-					Path.GetFileNameWithoutExtension(Binary.Config.OutputFilePath).StartsWith("UE4Editor-", StringComparison.InvariantCultureIgnoreCase) &&
-					Path.GetExtension(Binary.Config.OutputFilePath).EndsWith("dll", StringComparison.InvariantCultureIgnoreCase) &&
-					Binary.Config.Type == UEBuildBinaryType.DynamicLinkLibrary)
-				{
-					// ok, this is pretty awful, we want the import libraries that go with the editor, only on the PC
-					Manifest.AddBinaryNames(Path.Combine(Binary.Config.IntermediateDirectory, Path.GetFileNameWithoutExtension(Binary.Config.OutputFilePath) + ".lib"), "");
-				}
+                ToolChain.AddFilesToManifest(ref Manifest,Binary);
 			}
 			{
 				string DebugInfoExtension = BuildPlatform.GetDebugInfoExtension(UEBuildBinaryType.StaticLibrary);
@@ -1309,7 +1292,7 @@ namespace UnrealBuildTool
 			// If we're only generating the manifest, return now
 			if (UEBuildConfiguration.bGenerateManifest || UEBuildConfiguration.bCleanProject)
 			{
-                GenerateManifest(AppBinaries, GlobalLinkEnvironment.Config.Target.Platform, SpecialRocketLibFilesThatAreBuildProducts);
+                GenerateManifest(TargetToolChain, AppBinaries, GlobalLinkEnvironment.Config.Target.Platform, SpecialRocketLibFilesThatAreBuildProducts);
                 if (!BuildConfiguration.bXGEExport)
                 {
                     return ECompilationResult.Succeeded;
diff --git a/Engine/Source/Programs/UnrealBuildTool/HTML5/HTML5ToolChain.cs b/Engine/Source/Programs/UnrealBuildTool/HTML5/HTML5ToolChain.cs
--- a/Engine/Source/Programs/UnrealBuildTool/HTML5/HTML5ToolChain.cs
+++ b/Engine/Source/Programs/UnrealBuildTool/HTML5/HTML5ToolChain.cs
@@ -516,5 +516,13 @@ namespace UnrealBuildTool
 		{
 			throw new BuildException("HTML5 cannot compile C# files");
 		}
+
+        public override void AddFilesToManifest(ref FileManifest Manifest, UEBuildBinary Binary)
+        {
+            // we need to include the generated .mem file.  
+            string MemFile = Binary.Config.OutputFilePath + ".mem";
+            // make sure we don't add mem files more than once. 
+            Manifest.AddBinaryNames(MemFile, null);
+        }
 	};
 }
diff --git a/Engine/Source/Programs/UnrealBuildTool/IOS/IOSToolChain.cs b/Engine/Source/Programs/UnrealBuildTool/IOS/IOSToolChain.cs
--- a/Engine/Source/Programs/UnrealBuildTool/IOS/IOSToolChain.cs
+++ b/Engine/Source/Programs/UnrealBuildTool/IOS/IOSToolChain.cs
@@ -170,7 +170,7 @@ namespace UnrealBuildTool
 			}
 		}
 
-		static public void AddStubToManifest(ref FileManifest Manifest, UEBuildBinary Binary)
+		public override void AddFilesToManifest(ref FileManifest Manifest, UEBuildBinary Binary)
 		{
 			if (BuildConfiguration.bCreateStubIPA)
 			{
diff --git a/Engine/Source/Programs/UnrealBuildTool/Mac/MacToolChain.cs b/Engine/Source/Programs/UnrealBuildTool/Mac/MacToolChain.cs
--- a/Engine/Source/Programs/UnrealBuildTool/Mac/MacToolChain.cs
+++ b/Engine/Source/Programs/UnrealBuildTool/Mac/MacToolChain.cs
@@ -1187,7 +1187,7 @@ namespace UnrealBuildTool
 
 		static private string BundleContentsDirectory = "";
 
-		static public void AddAppBundleContentsToManifest(ref FileManifest Manifest, UEBuildBinary Binary)
+        public override void AddFilesToManifest(ref FileManifest Manifest, UEBuildBinary Binary)
 		{
 			if (Binary.Target.GlobalLinkEnvironment.Config.bIsBuildingConsoleApplication)
 			{
diff --git a/Engine/Source/Programs/UnrealBuildTool/ToolChain/UEToolChain.cs b/Engine/Source/Programs/UnrealBuildTool/ToolChain/UEToolChain.cs
--- a/Engine/Source/Programs/UnrealBuildTool/ToolChain/UEToolChain.cs
+++ b/Engine/Source/Programs/UnrealBuildTool/ToolChain/UEToolChain.cs
@@ -37,6 +37,8 @@ namespace UnrealBuildTool
 		ICollection<FileItem> PostBuild(FileItem Executable, LinkEnvironment ExecutableLinkEnvironment);
 
 		void SetUpGlobalEnvironment();
+
+        void AddFilesToManifest(ref FileManifest manifest, UEBuildBinary Binary );
 	}
 
 	public abstract class UEToolChain : IUEToolChain
@@ -152,5 +154,10 @@ namespace UnrealBuildTool
 		{
 
 		}
+
+        public virtual void AddFilesToManifest(ref FileManifest manifest, UEBuildBinary Binary )
+        {
+
+        }
 	};
 }
\ No newline at end of file
diff --git a/Engine/Source/Programs/UnrealBuildTool/Windows/VCToolChain.cs b/Engine/Source/Programs/UnrealBuildTool/Windows/VCToolChain.cs
--- a/Engine/Source/Programs/UnrealBuildTool/Windows/VCToolChain.cs
+++ b/Engine/Source/Programs/UnrealBuildTool/Windows/VCToolChain.cs
@@ -1517,5 +1517,18 @@ namespace UnrealBuildTool
 
 			return BaseVSToolPath;
 		}
+
+        public override void AddFilesToManifest(ref FileManifest Manifest, UEBuildBinary Binary)
+        {
+            // ok, this is pretty awful, we want the import libraries that go with the editor, only on the PC
+            if (UnrealBuildTool.BuildingRocket() &&
+                Path.GetFileNameWithoutExtension(Binary.Config.OutputFilePath).StartsWith("UE4Editor-", StringComparison.InvariantCultureIgnoreCase) &&
+                Path.GetExtension(Binary.Config.OutputFilePath).EndsWith("dll", StringComparison.InvariantCultureIgnoreCase) &&
+                Binary.Config.Type == UEBuildBinaryType.DynamicLinkLibrary)
+            {
+                // ok, this is pretty awful, we want the import libraries that go with the editor, only on the PC
+                Manifest.AddBinaryNames(Path.Combine(Binary.Config.IntermediateDirectory, Path.GetFileNameWithoutExtension(Binary.Config.OutputFilePath) + ".lib"), "");
+            }
+        }
 	};
 }
