Fixed problem with python path having spaces in it if installed to program files.

diff --git a/Engine/Source/Programs/UnrealBuildTool/HTML5/HTML5ToolChain.cs b/Engine/Source/Programs/UnrealBuildTool/HTML5/HTML5ToolChain.cs
--- a/Engine/Source/Programs/UnrealBuildTool/HTML5/HTML5ToolChain.cs
+++ b/Engine/Source/Programs/UnrealBuildTool/HTML5/HTML5ToolChain.cs
@@ -32,7 +32,8 @@ namespace UnrealBuildTool
 					{
 						PythonPath = Utils.IsRunningOnMono ? "python" : "python.exe";
 					}
-                    EMCCPath = "\"" + EMCCPath + "\""; 
+                    EMCCPath = "\"" + EMCCPath + "\"";
+                    PythonPath = "\"" + PythonPath + "\"";
 					// set some environment variable we'll need
 					//Environment.SetEnvironmentVariable("EMCC_DEBUG", "cache");
 					Environment.SetEnvironmentVariable("EMCC_CORES", "8");
@@ -367,9 +368,9 @@ namespace UnrealBuildTool
 
 				CompileAction.WorkingDirectory = Path.GetFullPath(".");
 				CompileAction.CommandPath = PythonPath;
-                
+
 				CompileAction.CommandArguments = EMCCPath + Arguments + FileArguments + CompileEnvironment.Config.AdditionalArguments;
-               
+
                 System.Console.WriteLine(CompileAction.CommandArguments); 
 				CompileAction.StatusDescription = Path.GetFileName(SourceFile.AbsolutePath);
 				CompileAction.StatusDetailedDescription = SourceFile.Description;
