Fixed Xcode projects generated for games with different name than target (for example when you create a project from the Vault in the Launcher and use custom name)

diff --git a/Engine/Source/Programs/UnrealBuildTool/System/XcodeProjectFileGenerator.cs b/Engine/Source/Programs/UnrealBuildTool/System/XcodeProjectFileGenerator.cs
--- a/Engine/Source/Programs/UnrealBuildTool/System/XcodeProjectFileGenerator.cs
+++ b/Engine/Source/Programs/UnrealBuildTool/System/XcodeProjectFileGenerator.cs
@@ -348,7 +348,7 @@ namespace UnrealBuildTool
 			if (Target.Type == XcodeTargetType.Legacy)
 			{
 				string UProjectPath = "";
-				if (UnrealBuildTool.HasUProjectFile() && Target.TargetName.StartsWith(Path.GetFileNameWithoutExtension(UnrealBuildTool.GetUProjectFile())))
+				if (UnrealBuildTool.HasUProjectFile())
 				{
 					if (MasterProjectRelativePath == UnrealBuildTool.GetUProjectPath())
 					{
@@ -356,8 +356,8 @@ namespace UnrealBuildTool
 					}
 					else
 					{
-					UProjectPath = " " + "\\\"" + UnrealBuildTool.GetUProjectFile() + "\\\"";
-				}
+						UProjectPath = " " + "\\\"" + UnrealBuildTool.GetUProjectFile() + "\\\"";
+					}
 				}
 				// Xcode provides $ACTION argument for determining if we are building or cleaning a project
 				Contents.Append("\t\t\tbuildArgumentsString = \"$(ACTION) " + Target.TargetName + " $(PLATFORM_NAME) $(CONFIGURATION)" + UProjectPath + "\";" + ProjectFileGenerator.NewLine);
@@ -1307,7 +1307,7 @@ namespace UnrealBuildTool
 					continue;
 				}
 
-				if (UnrealBuildTool.HasUProjectFile() && Target.TargetName.StartsWith(Path.GetFileNameWithoutExtension(UnrealBuildTool.GetUProjectFile())))
+				if (UnrealBuildTool.HasUProjectFile())
 				{
 					if (Target.TargetName.EndsWith("Editor"))
 					{
@@ -1316,7 +1316,7 @@ namespace UnrealBuildTool
 					else
 					{
 						string ProjectBinariesDir = UnrealBuildTool.GetUProjectPath() + "/Binaries/Mac/";
-						WriteSchemeFile(XcodeProjectPath, Target, ".app", ProjectBinariesDir + Target.TargetName);
+						WriteSchemeFile(XcodeProjectPath, Target, ".app", ProjectBinariesDir + Target.TargetName, new List<string>(new string[] { "&quot;" + UnrealBuildTool.GetUProjectFile() + "&quot;" }));
 					}
 				}
 				else if (Target.TargetName.EndsWith("Game"))
