Fix for pythonpath missing python exe name. TTP 342757.

diff --git a/Engine/Source/Programs/UnrealBuildTool/HTML5/HTML5ToolChain.cs b/Engine/Source/Programs/UnrealBuildTool/HTML5/HTML5ToolChain.cs
--- a/Engine/Source/Programs/UnrealBuildTool/HTML5/HTML5ToolChain.cs
+++ b/Engine/Source/Programs/UnrealBuildTool/HTML5/HTML5ToolChain.cs
@@ -28,12 +28,21 @@ namespace UnrealBuildTool
 					EMCCPath = Path.Combine(BaseSDKPath, "emcc");
 					// also figure out where python lives (if no envvar, assume it's in the path)
 					PythonPath = Environment.GetEnvironmentVariable("PYTHON");
+
+                    string PythonExeName = Utils.IsRunningOnMono ? "python" : "python.exe";
+
 					if (PythonPath == null)
 					{
-						PythonPath = Utils.IsRunningOnMono ? "python" : "python.exe";
+                        PythonPath = PythonExeName;
 					}
+                    else
+                    {
+                        if (!PythonPath.EndsWith(PythonExeName))
+                        {
+                            PythonPath += "/" + PythonExeName;
+                        }
+                    }
                     EMCCPath = "\"" + EMCCPath + "\"";
-                    PythonPath = "\"" + PythonPath + "\"";
 					// set some environment variable we'll need
 					//Environment.SetEnvironmentVariable("EMCC_DEBUG", "cache");
 					Environment.SetEnvironmentVariable("EMCC_CORES", "8");
