Pipe the output from Android ndk-build and ant-build to the log, so it's possible to see error messages when they fail.

diff --git a/Engine/Source/Programs/UnrealBuildTool/Android/UEDeployAndroid.cs b/Engine/Source/Programs/UnrealBuildTool/Android/UEDeployAndroid.cs
--- a/Engine/Source/Programs/UnrealBuildTool/Android/UEDeployAndroid.cs
+++ b/Engine/Source/Programs/UnrealBuildTool/Android/UEDeployAndroid.cs
@@ -466,13 +466,10 @@ namespace UnrealBuildTool.Android
 				NDKBuildInfo.UseShellExecute = true;
 				NDKBuildInfo.WindowStyle = ProcessWindowStyle.Minimized;
 				Console.WriteLine("\nRunning: " + NDKBuildInfo.FileName + " " + NDKBuildInfo.Arguments);
-				Process NDKBuild = new Process();
-				NDKBuild.StartInfo = NDKBuildInfo;
-				NDKBuild.Start();
-				NDKBuild.WaitForExit();
+				int NDKBuildExitCode = Utils.RunLocalProcessAndLogOutput(NDKBuildInfo);
 
 				// ndk build failure
-				if (NDKBuild.ExitCode != 0)
+				if (NDKBuildExitCode != 0)
 				{
 					throw new BuildException("ndk-build failed [{0}]", NDKBuildInfo.Arguments);
 				}
@@ -497,13 +494,10 @@ namespace UnrealBuildTool.Android
 			CallAntStartInfo.Arguments = "/c \"" + AntBuildPath + "\"  " + (bForDistribution ? "release" : "debug");
 			CallAntStartInfo.UseShellExecute = false;
 			Console.WriteLine("\nRunning: " + CallAntStartInfo.Arguments);
-			Process CallAnt = new Process();
-			CallAnt.StartInfo = CallAntStartInfo;
-			CallAnt.Start();
-			CallAnt.WaitForExit();
+			int CallAntExitCode = Utils.RunLocalProcessAndLogOutput(CallAntStartInfo);
 
 			// ant failure
-			if (CallAnt.ExitCode != 0)
+			if (CallAntExitCode != 0)
 			{
 				throw new BuildException("ant.bat failed [{0}]", CallAntStartInfo.Arguments);
 			}
diff --git a/Engine/Source/Programs/UnrealBuildTool/System/Utils.cs b/Engine/Source/Programs/UnrealBuildTool/System/Utils.cs
--- a/Engine/Source/Programs/UnrealBuildTool/System/Utils.cs
+++ b/Engine/Source/Programs/UnrealBuildTool/System/Utils.cs
@@ -383,6 +383,17 @@ namespace UnrealBuildTool
 			return ExitCode;
 		}
 
+		/**
+		 * Runs a local process and pipes the output to the log
+		 */
+		public static int RunLocalProcessAndLogOutput(ProcessStartInfo StartInfo)
+		{
+			Process LocalProcess = new Process();
+			LocalProcess.StartInfo = StartInfo;
+			LocalProcess.OutputDataReceived += (Sender, Line) => { if(Line != null && Line.Data != null) Log.TraceInformation(Line.Data); };
+			LocalProcess.ErrorDataReceived += (Sender, Line) => { if(Line != null && Line.Data != null) Log.TraceError(Line.Data); };
+			return RunLocalProcess(LocalProcess);
+		}
 
 		/// <summary>
 		/// Given a list of supported platforms, returns a list of names of platforms that should not be supported
