Implemented support for Distcc + DMUCs as an alternative to XGE on Mac OS X.
By either manually configuring distcc & DMUCS as per their documentation, or using the DistCode Xcode plugin compilation can be distributed amongst networked Macs.
All Macs must have the same developer tools installed & the same distcc/DMUCS/DistCode and be on the same local network.
This works for compiling for Mac & iOS and it should be possible to configure distcc & DMUCS manually for use on build farms for remote building from Windows.
Linux could also use this for local builds or remote build farms - but I'm not in a position to test this.

diff --git a/Engine/Source/Programs/UnrealBuildTool/Configuration/BuildConfiguration.cs b/Engine/Source/Programs/UnrealBuildTool/Configuration/BuildConfiguration.cs
--- a/Engine/Source/Programs/UnrealBuildTool/Configuration/BuildConfiguration.cs
+++ b/Engine/Source/Programs/UnrealBuildTool/Configuration/BuildConfiguration.cs
@@ -348,6 +348,24 @@ namespace UnrealBuildTool
 		[XmlConfig]
 		public static bool bBreakBuildOnLicenseViolation;
 
+		/// <summary>
+		/// Whether DMUCS/Distcc may be used.
+		/// </summary>
+		[XmlConfig]
+		public static bool bAllowDistcc;
+
+		/// <summary>
+		/// When enabled allows DMUCS/Distcc to fallback to local compilation when remote compiling fails. Defaults to true as separation of pre-process and compile stages can introduce non-fatal errors.
+		/// </summary>
+		[XmlConfig]
+		public static bool bAllowDistccLocalFallback;
+
+		/// <summary>
+		/// Path to the Distcc & DMUCS executables.
+		/// </summary>
+		[XmlConfig]
+		public static string DistccExecutablesPath;
+
 		/// <summary>
 		/// Sets the configuration back to defaults.
 		/// </summary>
@@ -454,6 +472,26 @@ namespace UnrealBuildTool
 			// By default check and stop the build on EULA violation
 			bCheckLicenseViolations = false;
 			bBreakBuildOnLicenseViolation = true;
+
+			// Distcc requires some setup - so by default disable it so we don't break local or remote building
+			bAllowDistcc = false;
+			bAllowDistccLocalFallback = true;
+
+			// The normal Posix place to install distcc/dmucs would be /usr/local so start there
+			DistccExecutablesPath = "/usr/local/bin";
+
+			// The default for normal Mac users should be to use DistCode which installs as an Xcode plugin and provides dynamic host management
+			if (ExternalExecution.GetRuntimePlatform() == UnrealTargetPlatform.Mac)
+			{
+				string UserDir = Environment.GetEnvironmentVariable("HOME");
+				string MacDistccExecutablesPath = UserDir + "/Library/Application Support/Developer/Shared/Xcode/Plug-ins/Distcc 3.2.xcplugin/Contents/usr/bin";
+
+				// But we should use the standard Posix directory when not installed - a build farm would use a static .dmucs hosts file not DistCode.
+				if (System.IO.Directory.Exists (MacDistccExecutablesPath)) 
+				{
+					DistccExecutablesPath = MacDistccExecutablesPath;
+				}
+			}
 		}
 
 		/// <summary>
@@ -514,6 +552,11 @@ namespace UnrealBuildTool
             {
                 bAllowXGE = false;
             }
+
+			if (!BuildPlatform.CanUseDistcc()) 
+			{
+				bAllowDistcc = false;
+			}
 		}
 	}
 }
diff --git a/Engine/Source/Programs/UnrealBuildTool/Configuration/UEBuildPlatform.cs b/Engine/Source/Programs/UnrealBuildTool/Configuration/UEBuildPlatform.cs
--- a/Engine/Source/Programs/UnrealBuildTool/Configuration/UEBuildPlatform.cs
+++ b/Engine/Source/Programs/UnrealBuildTool/Configuration/UEBuildPlatform.cs
@@ -26,6 +26,11 @@ namespace UnrealBuildTool
          */
         bool CanUseXGE();
 
+        /**
+         * If this platform can be compiled with DMUCS/Distcc
+         */
+		bool CanUseDistcc();
+
         /**
          * Register the platform with the UEBuildPlatform class
          */
@@ -423,6 +428,14 @@ namespace UnrealBuildTool
             return true;
         }
 
+		/**
+         *	If this platform can be compiled with DMUCS/Distcc
+         */
+		public virtual bool CanUseDistcc()
+		{
+			return false;
+		}
+
         /**
          *	Register the platform with the UEBuildPlatform class
          */
diff --git a/Engine/Source/Programs/UnrealBuildTool/IOS/UEBuildIOS.cs b/Engine/Source/Programs/UnrealBuildTool/IOS/UEBuildIOS.cs
--- a/Engine/Source/Programs/UnrealBuildTool/IOS/UEBuildIOS.cs
+++ b/Engine/Source/Programs/UnrealBuildTool/IOS/UEBuildIOS.cs
@@ -1,4 +1,4 @@
-﻿// Copyright 1998-2014 Epic Games, Inc. All Rights Reserved.
+// Copyright 1998-2014 Epic Games, Inc. All Rights Reserved.
 
 using System;
 using System.Collections.Generic;
@@ -108,6 +108,11 @@ namespace UnrealBuildTool
             return false;
         }
 
+		public override bool CanUseDistcc()
+		{
+			return true;
+		}
+
         /**
          *	Setup the target environment for building
          *	
diff --git a/Engine/Source/Programs/UnrealBuildTool/Mac/MacToolChain.cs b/Engine/Source/Programs/UnrealBuildTool/Mac/MacToolChain.cs
--- a/Engine/Source/Programs/UnrealBuildTool/Mac/MacToolChain.cs
+++ b/Engine/Source/Programs/UnrealBuildTool/Mac/MacToolChain.cs
@@ -393,69 +393,8 @@ namespace UnrealBuildTool
 				}
 
 				CompileAction.WorkingDirectory = GetMacDevSrcRoot();
-				
-				string CompilerBundleIdentifier = Environment.GetEnvironmentVariable("GCC_VERSION");
-				string StaticAnalysisMode = Environment.GetEnvironmentVariable("CLANG_STATIC_ANALYZER_MODE");
-				if(CompilerBundleIdentifier == null || CompilerBundleIdentifier == "" 
-					|| CompilerBundleIdentifier == "com.apple.compilers.llvm.clang.1_0" 
-					|| (StaticAnalysisMode != null && StaticAnalysisMode != ""))
-				{
-					CompileAction.CommandPath = MacCompiler;
-					if(StaticAnalysisMode != null && StaticAnalysisMode != "")
-					{
-						FileArguments = " --analyze " + FileArguments;
-					}
-				}
-				else
-				{
-					CompileAction.CommandPath = null;
-					string UserDir = Environment.GetEnvironmentVariable("HOME");
-					string XcodePlugins = UserDir + "/Library/Application Support/Developer/Shared/Xcode/Plug-ins/";
-					if(Directory.Exists(XcodePlugins))
-					{
-						IEnumerable<string> Plugins = Directory.EnumerateDirectories(XcodePlugins, "*.xcplugin");
-						foreach(string Plugin in Plugins)
-						{
-							string Resources = Plugin + "/Contents/Resources";
-							string[] CompilerSpecs = Directory.GetFiles(Resources, "*.xcspec");
-							if(CompilerSpecs.Length > 0)
-							{
-								IEnumerable<string> Lines = File.ReadLines(CompilerSpecs[0]);
-								bool bFindExecPath = false;
-								bool bFoundExecPath = false;
-								foreach(string Line in Lines)
-								{
-									if(!bFindExecPath && Line.Contains(CompilerBundleIdentifier))
-									{
-										bFindExecPath = true;
-									}
-									else if(bFindExecPath && Line.Contains("ExecPath"))
-									{
-										char[] CharsToTrim = {';', '"'};
-										string ExecPath = Line.Trim().Substring(12).TrimEnd(CharsToTrim);
-										CompileAction.CommandPath = ExecPath;
-										bFoundExecPath = true;
-										break;
-									}
-								}
-								if(bFoundExecPath)
-								{
-									break;
-								}
-							}
-						}
-					}
-					else
-					{
-						throw new BuildException("Couldn't find Xcode compiler plugins path: {0}", XcodePlugins);
-					}
-					
-					if(CompileAction.CommandPath == null)
-					{
-						throw new BuildException("Couldn't find Xcode compiler: {0}", CompilerBundleIdentifier);
-					}
-				}
-				CompileAction.CommandArguments = Arguments + FileArguments + CompileEnvironment.Config.AdditionalArguments;
+				CompileAction.CommandPath = "xcrun";
+				CompileAction.CommandArguments = MacCompiler + Arguments + FileArguments + CompileEnvironment.Config.AdditionalArguments;
 				CompileAction.CommandDescription = "Compile";
 				CompileAction.StatusDescription = Path.GetFileName(SourceFile.AbsolutePath);
 				CompileAction.StatusDetailedDescription = SourceFile.Description;
diff --git a/Engine/Source/Programs/UnrealBuildTool/Mac/UEBuildMac.cs b/Engine/Source/Programs/UnrealBuildTool/Mac/UEBuildMac.cs
--- a/Engine/Source/Programs/UnrealBuildTool/Mac/UEBuildMac.cs
+++ b/Engine/Source/Programs/UnrealBuildTool/Mac/UEBuildMac.cs
@@ -16,6 +16,11 @@ namespace UnrealBuildTool
             return false;
         }
 
+		public override bool CanUseDistcc()
+		{
+			return true;
+		}
+
 		protected override SDKStatus HasRequiredManualSDKInternal()
 		{
 			return SDKStatus.Valid;
diff --git a/Engine/Source/Programs/UnrealBuildTool/System/ActionGraph.cs b/Engine/Source/Programs/UnrealBuildTool/System/ActionGraph.cs
--- a/Engine/Source/Programs/UnrealBuildTool/System/ActionGraph.cs
+++ b/Engine/Source/Programs/UnrealBuildTool/System/ActionGraph.cs
@@ -268,6 +268,14 @@ namespace UnrealBuildTool
 					}
 				}
 
+				if (!bUsedXGE && BuildConfiguration.bAllowDistcc) 
+				{
+					ExecutorName = "Distcc";
+					Result = Distcc.ExecuteActions (ActionsToExecute);
+					// don't do local compilation
+					bUsedXGE = true;
+				}
+
 				// If XGE is disallowed or unavailable, execute the commands locally.
 				if (!bUsedXGE)
 				{
diff --git a/Engine/Source/Programs/UnrealBuildTool/System/Distcc.cs b/Engine/Source/Programs/UnrealBuildTool/System/Distcc.cs
new file mode 100644
--- /dev/null
+++ b/Engine/Source/Programs/UnrealBuildTool/System/Distcc.cs
@@ -0,0 +1,244 @@
+// Copyright 1998-2014 Epic Games, Inc. All Rights Reserved.
+
+using System;
+using System.Collections.Generic;
+using System.IO;
+using System.Diagnostics;
+using System.Xml;
+using System.Text.RegularExpressions;
+using System.Linq;
+using System.Reflection;
+
+namespace UnrealBuildTool
+{
+	public class Distcc
+	{
+		public static bool ExecuteActions(List<Action> Actions)
+		{
+			bool bDistccResult = true;
+			if (Actions.Count > 0)
+			{
+				// Time to sleep after each iteration of the loop in order to not busy wait.
+				const float LoopSleepTime = 0.1f;
+
+				int MaxActionsToExecuteInParallel = 0;
+				string UserDir = Environment.GetEnvironmentVariable("HOME");
+				string HostsInfo = UserDir + "/.dmucs/hosts-info";
+				System.IO.StreamReader File = new System.IO.StreamReader(HostsInfo);
+				string Line = null;
+				while((Line = File.ReadLine()) != null)
+				{
+					var HostInfo = Line.Split(' ');
+					if (HostInfo.Count () == 3) 
+					{
+						int NumCPUs = 0;
+						if (System.Int32.TryParse (HostInfo [1], out NumCPUs)) 
+						{
+							MaxActionsToExecuteInParallel += NumCPUs;
+						}
+					}
+				}
+				File.Close();
+
+				if (BuildConfiguration.bAllowDistccLocalFallback == false) 
+				{
+					Environment.SetEnvironmentVariable("DISTCC_FALLBACK", "0");
+				}
+
+				string DistccExecutable = BuildConfiguration.DistccExecutablesPath + "/distcc";
+				string GetHostExecutable = BuildConfiguration.DistccExecutablesPath + "/gethost";
+
+				Log.TraceInformation("Performing {0} actions (max {1} parallel jobs)", Actions.Count, MaxActionsToExecuteInParallel, DistccExecutable, GetHostExecutable);
+
+				Dictionary<Action, ActionThread> ActionThreadDictionary = new Dictionary<Action, ActionThread>();
+				int JobNumber = 1;
+				using(ProgressWriter ProgressWriter = new ProgressWriter("Compiling source code...", false))
+				{
+					int ProgressValue = 0;
+					while(true)
+					{
+						// Count the number of pending and still executing actions.
+						int NumUnexecutedActions = 0;
+						int NumExecutingActions = 0;
+						foreach (Action Action in Actions)
+						{
+							ActionThread ActionThread = null;
+							bool bFoundActionProcess = ActionThreadDictionary.TryGetValue(Action, out ActionThread);
+							if (bFoundActionProcess == false)
+							{
+								NumUnexecutedActions++;
+							}
+							else if (ActionThread != null)
+							{
+								if (ActionThread.bComplete == false)
+								{
+									NumUnexecutedActions++;
+									NumExecutingActions++;
+								}
+							}
+						}
+
+						// Update the current progress
+						int NewProgressValue = Actions.Count + 1 - NumUnexecutedActions;
+						if (ProgressValue != NewProgressValue)
+						{
+							ProgressWriter.Write(ProgressValue, Actions.Count + 1);
+							ProgressValue = NewProgressValue;
+						}
+
+						// If there aren't any pending actions left, we're done executing.
+						if (NumUnexecutedActions == 0)
+						{
+							break;
+						}
+
+						// If there are fewer actions executing than the maximum, look for pending actions that don't have any outdated
+						// prerequisites.
+						foreach (Action Action in Actions)
+						{
+							ActionThread ActionProcess = null;
+							bool bFoundActionProcess = ActionThreadDictionary.TryGetValue(Action, out ActionProcess);
+							if (bFoundActionProcess == false)
+							{
+								if (NumExecutingActions < Math.Max(1,MaxActionsToExecuteInParallel))
+								{
+									// Determine whether there are any prerequisites of the action that are outdated.
+									bool bHasOutdatedPrerequisites = false;
+									bool bHasFailedPrerequisites = false;
+									foreach (FileItem PrerequisiteItem in Action.PrerequisiteItems)
+									{
+										if (PrerequisiteItem.ProducingAction != null && Actions.Contains(PrerequisiteItem.ProducingAction))
+										{
+											ActionThread PrerequisiteProcess = null;
+											bool bFoundPrerequisiteProcess = ActionThreadDictionary.TryGetValue( PrerequisiteItem.ProducingAction, out PrerequisiteProcess );
+											if (bFoundPrerequisiteProcess == true)
+											{
+												if (PrerequisiteProcess == null)
+												{
+													bHasFailedPrerequisites = true;
+												}
+												else if (PrerequisiteProcess.bComplete == false)
+												{
+													bHasOutdatedPrerequisites = true;
+												}
+												else if (PrerequisiteProcess.ExitCode != 0)
+												{
+													bHasFailedPrerequisites = true;
+												}
+											}
+											else
+											{
+												bHasOutdatedPrerequisites = true;
+											}
+										}
+									}
+
+									// If there are any failed prerequisites of this action, don't execute it.
+									if (bHasFailedPrerequisites)
+									{
+										// Add a null entry in the dictionary for this action.
+										ActionThreadDictionary.Add( Action, null );
+									}
+									// If there aren't any outdated prerequisites of this action, execute it.
+									else if (!bHasOutdatedPrerequisites)
+									{
+										if ((Action.ActionType == ActionType.Compile || Action.ActionType == ActionType.Link) && DistccExecutable != null && GetHostExecutable != null) 
+										{
+											string NewCommandArguments = "--wait -1 \"" + DistccExecutable + "\" " + Action.CommandPath + " " + Action.CommandArguments;
+											Action.CommandPath = GetHostExecutable;
+											Action.CommandArguments = NewCommandArguments;
+										}
+
+										ActionThread ActionThread = new ActionThread(Action, JobNumber, Actions.Count);
+										JobNumber++;
+										ActionThread.Run();
+
+										ActionThreadDictionary.Add(Action, ActionThread);
+
+										NumExecutingActions++;
+									}
+								}
+							}
+						}
+
+						System.Threading.Thread.Sleep(TimeSpan.FromSeconds(LoopSleepTime));
+					}
+				}
+
+				Log.WriteLineIf(BuildConfiguration.bLogDetailedActionStats, TraceEventType.Information, "-------- Begin Detailed Action Stats ----------------------------------------------------------");
+				Log.WriteLineIf(BuildConfiguration.bLogDetailedActionStats, TraceEventType.Information, "^Action Type^Duration (seconds)^Tool^Task^Using PCH^Description");
+
+				double TotalThreadSeconds = 0;
+
+				// Check whether any of the tasks failed and log action stats if wanted.
+				foreach (KeyValuePair<Action, ActionThread> ActionProcess in ActionThreadDictionary)
+				{
+					Action Action = ActionProcess.Key;
+					ActionThread ActionThread = ActionProcess.Value;
+
+					// Check for pending actions, preemptive failure
+					if (ActionThread == null)
+					{
+						bDistccResult = false;
+						continue;
+					}
+					// Check for executed action but general failure
+					if (ActionThread.ExitCode != 0)
+					{
+						bDistccResult = false;
+					}
+					// Log CPU time, tool and task.
+					double ThreadSeconds = Action.Duration.TotalSeconds;
+
+					Log.WriteLineIf(BuildConfiguration.bLogDetailedActionStats,
+						TraceEventType.Information,
+						"^{0}^{1:0.00}^{2}^{3}^{4}^{5}", 
+						Action.ActionType.ToString(),
+						ThreadSeconds,
+						Path.GetFileName(Action.CommandPath), 
+						Action.StatusDescription,
+						Action.bIsUsingPCH,
+						Action.StatusDetailedDescription);
+
+					// Update statistics
+					switch (Action.ActionType)
+					{
+					case ActionType.BuildProject:
+						UnrealBuildTool.TotalBuildProjectTime += ThreadSeconds;
+						break;
+
+					case ActionType.Compile:
+						UnrealBuildTool.TotalCompileTime += ThreadSeconds;
+						break;
+
+					case ActionType.CreateAppBundle:
+						UnrealBuildTool.TotalCreateAppBundleTime += ThreadSeconds;
+						break;
+
+					case ActionType.GenerateDebugInfo:
+						UnrealBuildTool.TotalGenerateDebugInfoTime += ThreadSeconds;
+						break;
+
+					case ActionType.Link:
+						UnrealBuildTool.TotalLinkTime += ThreadSeconds;
+						break;
+
+					default:
+						UnrealBuildTool.TotalOtherActionsTime += ThreadSeconds;
+						break;
+					}
+
+					// Keep track of total thread seconds spent on tasks.
+					TotalThreadSeconds += ThreadSeconds;
+				}
+
+				Log.TraceInformation("-------- End Detailed Actions Stats -----------------------------------------------------------");
+
+				// Log total CPU seconds and numbers of processors involved in tasks.
+				Log.WriteLineIf(BuildConfiguration.bLogDetailedActionStats || BuildConfiguration.bPrintDebugInfo, 
+					TraceEventType.Information, "Cumulative thread seconds ({0} processors): {1:0.00}", System.Environment.ProcessorCount, TotalThreadSeconds);
+			}
+			return bDistccResult;
+		}
+	}
+}
diff --git a/Engine/Source/Programs/UnrealBuildTool/System/LocalExecutor.cs b/Engine/Source/Programs/UnrealBuildTool/System/LocalExecutor.cs
--- a/Engine/Source/Programs/UnrealBuildTool/System/LocalExecutor.cs
+++ b/Engine/Source/Programs/UnrealBuildTool/System/LocalExecutor.cs
@@ -370,16 +370,6 @@ namespace UnrealBuildTool
 				int MaxActionsAffordedByMemory = (int)(Math.Max(1, (PhysicalRAMAvailableMB) / MinMemoryPerActionMB));
 
 				MaxActionsToExecuteInParallel = Math.Min(MaxActionsToExecuteInParallel, MaxActionsAffordedByMemory);
-
-                if (ExternalExecution.GetRuntimePlatform() == UnrealTargetPlatform.Mac)
-			    {
-				    string NumUBTBuildTasks = Environment.GetEnvironmentVariable("NumUBTBuildTasks");
-				    Int32 MaxUBTBuildTasks = MaxActionsToExecuteInParallel;
-				    if(Int32.TryParse(NumUBTBuildTasks, out MaxUBTBuildTasks))
-				    {
-					    MaxActionsToExecuteInParallel = MaxUBTBuildTasks;
-				    }
-                }
             }
 
             Log.TraceInformation("Performing {0} actions (max {1} parallel jobs)", Actions.Count, MaxActionsToExecuteInParallel);
diff --git a/Engine/Source/Programs/UnrealBuildTool/System/UnrealBuildTool.cs b/Engine/Source/Programs/UnrealBuildTool/System/UnrealBuildTool.cs
--- a/Engine/Source/Programs/UnrealBuildTool/System/UnrealBuildTool.cs
+++ b/Engine/Source/Programs/UnrealBuildTool/System/UnrealBuildTool.cs
@@ -1,4 +1,4 @@
-﻿// Copyright 1998-2014 Epic Games, Inc. All Rights Reserved.
+// Copyright 1998-2014 Epic Games, Inc. All Rights Reserved.
 
 using System;
 using System.Collections.Generic;
@@ -1312,7 +1312,7 @@ namespace UnrealBuildTool
 
             // Figure out how long we took to execute.
             double BuildDuration = (DateTime.UtcNow - StartTime - MutexWaitTime).TotalSeconds;
-            if (ExecutorName == "Local")
+			if (ExecutorName == "Local" || ExecutorName == "Distcc")
             {
                 Log.TraceInformation("Cumulative action seconds ({0} processors): {1:0.00} building projects, {2:0.00} compiling, {3:0.00} creating app bundles, {4:0.00} generating debug info, {5:0.00} linking, {6:0.00} other",
                     System.Environment.ProcessorCount,
diff --git a/Engine/Source/Programs/UnrealBuildTool/UnrealBuildTool.csproj b/Engine/Source/Programs/UnrealBuildTool/UnrealBuildTool.csproj
--- a/Engine/Source/Programs/UnrealBuildTool/UnrealBuildTool.csproj
+++ b/Engine/Source/Programs/UnrealBuildTool/UnrealBuildTool.csproj
@@ -215,6 +215,7 @@
     <Compile Include="System\ActionHistory.cs" />
     <Compile Include="System\BuildException.cs" />
     <Compile Include="Configuration\BuildConfiguration.cs" />
+    <Compile Include="System\Distcc.cs" />
     <Compile Include="System\ModuleProcessingException.cs" />
     <Compile Include="System\MakefileGenerator.cs" />
     <Compile Include="System\VCSolutionOptions.cs" />
diff --git a/Engine/Source/Programs/UnrealBuildTool/UnrealBuildTool_Mono.csproj b/Engine/Source/Programs/UnrealBuildTool/UnrealBuildTool_Mono.csproj
--- a/Engine/Source/Programs/UnrealBuildTool/UnrealBuildTool_Mono.csproj
+++ b/Engine/Source/Programs/UnrealBuildTool/UnrealBuildTool_Mono.csproj
@@ -169,6 +169,7 @@
     <Compile Include="IOS\UEDeployIOS.cs" />
     <Compile Include="System\ModuleProcessingException.cs" />
     <Compile Include="System\VCSolutionOptions.cs" />
+    <Compile Include="System\Distcc.cs" />
   </ItemGroup>
   <ItemGroup>
     <BootstrapperPackage Include="Microsoft.Net.Client.3.5">
