UnrealBuildTool: Experimental fast C++ include dependency scanning
- Adds experimental super-fast C++ outdated file checking
- This feature is turned off for now as we continue to test and improve it
      - You can try it out by enabling "bUseExperimentalFastDependencyScan" option in your BuildConfiguration.xml
- Here is the basic idea:
     - We no longer exhaustively scan all includes and build up a big graph every invocation
     - Instead, source files whose build products are missing have their includes scanned asynchronously while being compiled
     - The flat list of dependent includes for every outdated C++ is saved into a new cache file ("FlatCPPIncludes.bin")
     - On the next run, we quickly load that up and "just know" which files to check timestamps on to determine what is out of date
- Lots of "[REDACTED_HANDLE] fastubt" comments were added to UnrealBuildTool for potential performance optimizations and further improvements on this feature.

UnrealBuildTool: Determination of which modules have UObjects is now faster
- We now cache which modules have UObjects and load those for the next session

UnrealBuildTool: Module "shared" precompiled header determination is now much faster
- We no longer scan all C++ includes for a module to figure out which "shared" PCH to use
- Instead, we use the module dependencies specified in the module's *.Build.cs file
- For example, if your module depends on "Engine" and "UnrealEd", we choose "UnrealEd"'s shared PCH

Other UBT optimizations:
- Reduced calls to string formatting functions when setting up API definitions for all modules
- Added new performance diagnostics when bPrintPerformanceInfo is enabled in BuildConfiguration.xml
- We no longer check for "external" headers when scanning includes (this code didn't work at all)
- Optimized CleanDirectorySeparators() utility function to avoid string copies

Fixed UnrealBuildTool not saving DependencyCache under a platform-named folder

diff --git a/Engine/Source/Programs/UnrealBuildTool/Android/AndroidToolChain.cs b/Engine/Source/Programs/UnrealBuildTool/Android/AndroidToolChain.cs
--- a/Engine/Source/Programs/UnrealBuildTool/Android/AndroidToolChain.cs
+++ b/Engine/Source/Programs/UnrealBuildTool/Android/AndroidToolChain.cs
@@ -395,7 +395,7 @@ namespace UnrealBuildTool
 		static string ArPathx86;
 
 		static private bool bHasPrintedApiLevel = false;
-		public override CPPOutput CompileCPPFiles(CPPEnvironment CompileEnvironment, List<FileItem> SourceFiles, string ModuleName)
+		public override CPPOutput CompileCPPFiles(UEBuildTarget Target, CPPEnvironment CompileEnvironment, List<FileItem> SourceFiles, string ModuleName)
 		{
 			if (!bHasPrintedApiLevel)
 			{
@@ -441,6 +441,8 @@ namespace UnrealBuildTool
 				Arguments += string.Format(" -D \"{0}\"", Definition);
 			}
 
+			var BuildPlatform = UEBuildPlatform.GetBuildPlatformForCPPTargetPlatform(CompileEnvironment.Config.Target.Platform);
+
 			// Create a compile action for each source file.
 			CPPOutput Result = new CPPOutput();
 			foreach (FileItem SourceFile in SourceFiles)
@@ -478,9 +480,9 @@ namespace UnrealBuildTool
 
 				// Add the C++ source file and its included files to the prerequisite item list.
 				CompileAction.PrerequisiteItems.Add(SourceFile);
-				foreach (FileItem IncludedFile in CompileEnvironment.GetIncludeDependencies(SourceFile))
 				{
-					CompileAction.PrerequisiteItems.Add(IncludedFile);
+					var IncludedFileList = CPPEnvironment.FindAndCacheAllIncludedFiles( Target, SourceFile, BuildPlatform, CompileEnvironment.GetIncludesPathsToSearch( SourceFile ), CompileEnvironment.IncludeFileSearchDictionary, bOnlyCachedDependencies:BuildConfiguration.bUseExperimentalFastDependencyScan );
+					CompileAction.PrerequisiteItems.AddRange( IncludedFileList );
 				}
 
 				if (CompileEnvironment.Config.PrecompiledHeaderAction == PrecompiledHeaderAction.Create)
diff --git a/Engine/Source/Programs/UnrealBuildTool/Configuration/BuildConfiguration.cs b/Engine/Source/Programs/UnrealBuildTool/Configuration/BuildConfiguration.cs
--- a/Engine/Source/Programs/UnrealBuildTool/Configuration/BuildConfiguration.cs
+++ b/Engine/Source/Programs/UnrealBuildTool/Configuration/BuildConfiguration.cs
@@ -348,6 +348,15 @@ namespace UnrealBuildTool
 		[XmlConfig]
 		public static bool bBreakBuildOnLicenseViolation;
 
+		/// <summary>
+		/// True if we should use the new 'experimental' C++ header dependency scanning and caching system that allows Unreal Build Tool to
+		/// detect out of date dependencies very quickly.  When enabled, a deep C++ include graph does not have to be generated, and instead
+		/// we only scan and cache indirect includes for after a dependent build product was already found to be out of date.  During the
+		/// next build, we'll load those cached indirect includes and check for outdatedness.
+		/// </summary>
+		[XmlConfig]
+		public static bool bUseExperimentalFastDependencyScan;
+	
 		/// <summary>
 		/// Whether DMUCS/Distcc may be used.
 		/// </summary>
@@ -473,6 +482,9 @@ namespace UnrealBuildTool
 			bCheckLicenseViolations = false;
 			bBreakBuildOnLicenseViolation = true;
 
+			// The new experimental fast dependency scanning is still being tested, so we disable it by default
+			bUseExperimentalFastDependencyScan = false;
+
 			// Distcc requires some setup - so by default disable it so we don't break local or remote building
 			bAllowDistcc = false;
 			bAllowDistccLocalFallback = true;
diff --git a/Engine/Source/Programs/UnrealBuildTool/Configuration/UEBuildBinary.cs b/Engine/Source/Programs/UnrealBuildTool/Configuration/UEBuildBinary.cs
--- a/Engine/Source/Programs/UnrealBuildTool/Configuration/UEBuildBinary.cs
+++ b/Engine/Source/Programs/UnrealBuildTool/Configuration/UEBuildBinary.cs
@@ -371,7 +371,8 @@ namespace UnrealBuildTool
 					var Module = Target.GetModuleByName( ModuleName );
 					ReferencedModules[ ModuleName ] = Module;
 
-					Module.GetAllDependencyModules(ref ReferencedModules, ref OrderedModules, bIncludeDynamicallyLoaded, bForceCircular);
+					bool bOnlyDirectDependencies = false;
+					Module.GetAllDependencyModules(ref ReferencedModules, ref OrderedModules, bIncludeDynamicallyLoaded, bForceCircular, bOnlyDirectDependencies);
 
 					OrderedModules.Add( Module );
 				}
diff --git a/Engine/Source/Programs/UnrealBuildTool/Configuration/UEBuildModule.cs b/Engine/Source/Programs/UnrealBuildTool/Configuration/UEBuildModule.cs
--- a/Engine/Source/Programs/UnrealBuildTool/Configuration/UEBuildModule.cs
+++ b/Engine/Source/Programs/UnrealBuildTool/Configuration/UEBuildModule.cs
@@ -548,19 +548,19 @@ namespace UnrealBuildTool
 					{
 						// When generating IntelliSense files, never add dllimport/dllexport specifiers as it
 						// simply confuses the compiler
-						Definitions.Add( string.Format( "{0}_API=", Name.ToUpperInvariant() ) );
+						Definitions.Add( Name.ToUpperInvariant() + "_API=" );
 					}
 					else if( Binary == SourceBinary )
 					{
 						if( Binary.Config.bAllowExports )
 						{
 							Log.TraceVerbose( "{0}: Exporting {1} from {2}", Path.GetFileNameWithoutExtension( SourceBinary.Config.OutputFilePath ), Name, Path.GetFileNameWithoutExtension( Binary.Config.OutputFilePath ) );
-							Definitions.Add( string.Format( "{0}_API=DLLEXPORT", Name.ToUpperInvariant() ) );
+							Definitions.Add( Name.ToUpperInvariant() + "_API=DLLEXPORT" );
 						}
 						else
 						{
 							Log.TraceVerbose( "{0}: Not importing/exporting {1} (binary: {2})", Path.GetFileNameWithoutExtension( SourceBinary.Config.OutputFilePath ), Name, Path.GetFileNameWithoutExtension( Binary.Config.OutputFilePath ) );
-							Definitions.Add( string.Format( "{0}_API=", Name.ToUpperInvariant() ) );
+							Definitions.Add( Name.ToUpperInvariant() + "_API=" );
 						}
 					}
 					else
@@ -571,17 +571,17 @@ namespace UnrealBuildTool
 						if( bIncludePathsOnly )
 						{
 							Log.TraceVerbose( "{0}: Include paths only for {1} (binary: {2})", Path.GetFileNameWithoutExtension( SourceBinary.Config.OutputFilePath ), Name, Path.GetFileNameWithoutExtension( Binary.Config.OutputFilePath ) );
-							Definitions.Add( string.Format( "{0}_API=", Name.ToUpperInvariant() ) );
+							Definitions.Add( Name.ToUpperInvariant() + "_API=" );
 						}
 						else if (Binary.Config.bAllowExports)
 						{
 							Log.TraceVerbose( "{0}: Importing {1} from {2}", Path.GetFileNameWithoutExtension( SourceBinary.Config.OutputFilePath ), Name, Path.GetFileNameWithoutExtension( Binary.Config.OutputFilePath ) );
-							Definitions.Add( string.Format( "{0}_API=DLLIMPORT", Name.ToUpperInvariant() ) );
+							Definitions.Add( Name.ToUpperInvariant() + "_API=DLLIMPORT" );
 						}
 						else
 						{
 							Log.TraceVerbose("{0}: Not importing/exporting {1} (binary: {2})", Path.GetFileNameWithoutExtension(SourceBinary.Config.OutputFilePath), Name, Path.GetFileNameWithoutExtension(Binary.Config.OutputFilePath));
-							Definitions.Add(string.Format("{0}_API=", Name.ToUpperInvariant()));
+							Definitions.Add( Name.ToUpperInvariant() + "_API=" );
 						}
 					}
 				}
@@ -760,9 +760,10 @@ namespace UnrealBuildTool
 		 * @param	ReferencedModules	Hash of all referenced modules
 		 * @param	OrderedModules		Ordered list of the referenced modules
 		 * @param	bIncludeDynamicallyLoaded	True if dynamically loaded modules (and all of their dependent modules) should be included.
-		 * @param	bForceCircular	True if circular dependecies should be processed
+		 * @param	bForceCircular	True if circular dependencies should be processed
+		 * @param	bOnlyDirectDependencies	True to return only this module's direct dependencies
 		 */
-		public virtual void GetAllDependencyModules(ref Dictionary<string, UEBuildModule> ReferencedModules, ref List<UEBuildModule> OrderedModules, bool bIncludeDynamicallyLoaded, bool bForceCircular)
+		public virtual void GetAllDependencyModules(ref Dictionary<string, UEBuildModule> ReferencedModules, ref List<UEBuildModule> OrderedModules, bool bIncludeDynamicallyLoaded, bool bForceCircular, bool bOnlyDirectDependencies)
 		{
 		}
 
@@ -1105,9 +1106,6 @@ namespace UnrealBuildTool
 				return LinkInputFiles;
 			}
 
-			// Process all of the header file dependencies for this module
-			this.ProcessAllCppDependencies( ModuleCompileEnvironment );
-
 			// Throw an error if the module's source file list referenced any non-existent files.
 			if (SourceFilesToBuild.MissingFiles.Count > 0)
 			{
@@ -1127,17 +1125,33 @@ namespace UnrealBuildTool
 				// Add default PCLaunch.rc file if this module has no own resource file specified
 				if (SourceFilesToBuild.RCFiles.Count <= 0)
 				{
-					string DefRC = Path.Combine(Directory.GetCurrentDirectory(), "Runtime/Launch/Resources/Windows/PCLaunch.rc");
+					string DefRC = Utils.CleanDirectorySeparators( Path.GetFullPath(Path.Combine(Directory.GetCurrentDirectory(), "Runtime/Launch/Resources/Windows/PCLaunch.rc")) );
 					FileItem Item = FileItem.GetItemByFullPath(DefRC);
 					SourceFilesToBuild.RCFiles.Add(Item);
 				}
 
 				// Always compile in the API version resource separately. This is required for the module manager to detect compatible API versions.
-				string ModuleVersionRC = Path.GetFullPath(Path.Combine(Directory.GetCurrentDirectory(), "Runtime/Core/Resources/Windows/ModuleVersionResource.rc.inl"));
+				string ModuleVersionRC = Utils.CleanDirectorySeparators( Path.GetFullPath(Path.Combine(Directory.GetCurrentDirectory(), "Runtime/Core/Resources/Windows/ModuleVersionResource.rc.inl")) );
 				FileItem ModuleVersionItem = FileItem.GetItemByFullPath(ModuleVersionRC);
-				if(!SourceFilesToBuild.RCFiles.Contains(ModuleVersionItem)) SourceFilesToBuild.RCFiles.Add(ModuleVersionItem);
+				if( !SourceFilesToBuild.RCFiles.Contains(ModuleVersionItem) )
+				{
+					SourceFilesToBuild.RCFiles.Add(ModuleVersionItem);
+				}
+			}
+
+
+			{ 
+				// Process all of the header file dependencies for this module
+				this.CachePCHUsageForModuleSourceFiles( ModuleCompileEnvironment );
+
+				// Make sure our RC files have cached includes.  
+				foreach( var RCFile in SourceFilesToBuild.RCFiles )
+				{ 
+					RCFile.CachedCPPEnvironment = ModuleCompileEnvironment;
+				}
 			}
 
+
 			// Check to see if this is an Engine module (including program or plugin modules).  That is, the module is located under the "Engine" folder
 			var IsGameModule = !Utils.IsFileUnderDirectory( this.ModuleDirectory, Path.Combine( ProjectFileGenerator.EngineRelativePath ) );
 
@@ -1178,7 +1192,7 @@ namespace UnrealBuildTool
 			bool bWasModuleCodeCompiled = false;
 			if (BuildPlatform.ShouldUsePCHFiles(CompileEnvironment.Config.Target.Platform, CompileEnvironment.Config.Target.Configuration))
 			{
-				var PCHTimerStart = DateTime.UtcNow;
+				var PCHGenTimerStart = DateTime.UtcNow;
 
 				// The code below will figure out whether this module will either use a "unique PCH" (private PCH that will only be included by
 				// this module's code files), or a "shared PCH" (potentially included by many code files in many modules.)  Only one or the other
@@ -1195,146 +1209,66 @@ namespace UnrealBuildTool
 				}
 
 				bool bUseSharedPCHFiles  = BuildConfiguration.bUseSharedPCHs && (bDisableSharedPCHFiles == false) && !UnrealBuildTool.BuildingRocket();
-				bool bIsASharedPCHModule = bUseSharedPCHFiles && GlobalCompileEnvironment.SharedPCHHeaderFiles.Any( PCH => PCH.Module == this );
-
-				string SharedPCHHeaderFilePath = null;
-				if( bIsASharedPCHModule )
-				{
-					SharedPCHHeaderFilePath = Path.GetFullPath( Path.Combine( ProjectFileGenerator.RootRelativePath, "Engine", "Source", this.SharedPCHHeaderFile ) );
-				}
 
-				// Map from pch header string to the source files that use that PCH
-				var UsageMapPCH = new Dictionary<string, List<FileItem>>( StringComparer.InvariantCultureIgnoreCase );
-
-				// Determine what potential precompiled header is used by each source file.
-				double SharedPCHTime = 0.0;
-				foreach( var CPPFile in SourceFilesToBuild.CPPFiles )
+				if( bUseSharedPCHFiles )
 				{
-					
-					if (CPPFile.AbsolutePath.IndexOf("ActorConstruction.cpp", StringComparison.InvariantCultureIgnoreCase) >= 0)
+					string SharingPCHHeaderFilePath = null;
+					bool bIsASharedPCHModule = bUseSharedPCHFiles && GlobalCompileEnvironment.SharedPCHHeaderFiles.Any( PCH => PCH.Module == this );
+					if( bIsASharedPCHModule )
 					{
-						Log.TraceVerbose("ActorConstruction.cpp");
+						SharingPCHHeaderFilePath = Path.GetFullPath( Path.Combine( ProjectFileGenerator.RootRelativePath, "Engine", "Source", this.SharedPCHHeaderFile ) );
 					}
 
-					if( CPPFile.PrecompiledHeaderIncludeFilename == null )
+					// When compiling in modular mode, we can't use a shared PCH file when compiling a module
+					// with exports, because the shared PCH can only have imports in it to work correctly.
+					bool bCanModuleUseOwnSharedPCH = bAllowSharedPCH && bIsASharedPCHModule && bCompileMonolithic && ProcessedDependencies.UniquePCHHeaderFile.AbsolutePath.Equals( SharingPCHHeaderFilePath, StringComparison.InvariantCultureIgnoreCase );
+					if( bAllowSharedPCH && ( !bIsASharedPCHModule || bCanModuleUseOwnSharedPCH ) )
 					{
-						throw new BuildException( "No PCH usage for file \"{0}\" . Missing #include header?", CPPFile.AbsolutePath );
-					}
-
-					// Create a new entry if not in the pch usage map
-					UsageMapPCH.GetOrAddNew( CPPFile.PrecompiledHeaderIncludeFilename ).Add( CPPFile );
+						// Figure out which shared PCH tier we're in
+						var DirectDependencyModules = new List<UEBuildModule>();
+						{ 
+							var ReferencedModules = new Dictionary<string, UEBuildModule>( StringComparer.InvariantCultureIgnoreCase );
+							this.GetAllDependencyModules( ref ReferencedModules, ref DirectDependencyModules, bIncludeDynamicallyLoaded:false, bForceCircular:false, bOnlyDirectDependencies:true );
+						}
 
-					if (bUseSharedPCHFiles)
-					{
-						var SharedPCHStartTime = DateTime.UtcNow;
-
-						// When compiling in modular mode, we can't use a shared PCH file when compiling a module
-						// with exports, because the shared PCH can only have imports in it to work correctly.
-						// @todo SharedPCH: If we ever have SharedPCH headers that themselves belong to modules which never use DLL Exports, we can avoid
-						// generating TWO PCH files by checking for that here.  For now, we always assume that SharedPCH headers have exports when
-						// compiling in modular mode.
-						bool bCanModuleUseOwnSharedPCH = bAllowSharedPCH && bIsASharedPCHModule && bCompileMonolithic && CPPFile.PrecompiledHeaderIncludeFilename.Equals( SharedPCHHeaderFilePath, StringComparison.InvariantCultureIgnoreCase );
-						if( bAllowSharedPCH && ( !bIsASharedPCHModule || bCanModuleUseOwnSharedPCH ) )
+						int LargestSharedPCHHeaderFileIndex = -1;
+						foreach( var DependencyModule in DirectDependencyModules )
 						{
-							// Figure out which shared PCH tier we're in
-							int LargestSharedPCHHeaderFileIndex = -1;
+							// These Shared PCHs are ordered from least complex to most complex.  We'll start at the last one and search backwards.
+							for( var SharedPCHHeaderFileIndex = GlobalCompileEnvironment.SharedPCHHeaderFiles.Count - 1; SharedPCHHeaderFileIndex > LargestSharedPCHHeaderFileIndex; --SharedPCHHeaderFileIndex )
 							{
-								var AllIncludedFiles = ModuleCompileEnvironment.GetIncludeDependencies( CPPFile );
-								foreach( var IncludedFile in AllIncludedFiles )
-								{
-									// These PCHs are ordered from least complex to most complex.  We'll start at the last one and search backwards.
-									for( var SharedPCHHeaderFileIndex = GlobalCompileEnvironment.SharedPCHHeaderFiles.Count - 1; SharedPCHHeaderFileIndex > LargestSharedPCHHeaderFileIndex; --SharedPCHHeaderFileIndex )
-									{
-										var CurSharedPCHHeaderFile = GlobalCompileEnvironment.SharedPCHHeaderFiles[ SharedPCHHeaderFileIndex ];
-										if( IncludedFile == CurSharedPCHHeaderFile.PCHHeaderFile )
-										{
-											LargestSharedPCHHeaderFileIndex = SharedPCHHeaderFileIndex;
-											break;
-										}
-									}
-
-									if( LargestSharedPCHHeaderFileIndex == GlobalCompileEnvironment.SharedPCHHeaderFiles.Count - 1)
-									{
-										// We've determined that the module is using our most complex PCH header, so we can early-out
-										break;
-									}
-								}
-							}
+								var CurSharedPCHHeaderFile = GlobalCompileEnvironment.SharedPCHHeaderFiles[ SharedPCHHeaderFileIndex ];
 
-							if( LargestSharedPCHHeaderFileIndex > -1 )
-							{
-								var LargestIncludedSharedPCHHeaderFile = GlobalCompileEnvironment.SharedPCHHeaderFiles[LargestSharedPCHHeaderFileIndex];
-								if( SharedPCHHeaderFile == null )
-								{
-									SharedPCHModuleName = LargestIncludedSharedPCHHeaderFile.Module.Name;
-									SharedPCHHeaderFile = LargestIncludedSharedPCHHeaderFile.PCHHeaderFile;
-								}
-								else if( SharedPCHHeaderFile != LargestIncludedSharedPCHHeaderFile.PCHHeaderFile )
+								if( DependencyModule == CurSharedPCHHeaderFile.Module || 
+									( bIsASharedPCHModule && CurSharedPCHHeaderFile.Module == this ) )	// If we ourselves are a shared PCH module, always at least use our own module as our shared PCH header if we can't find anything better
 								{
-									// @todo UBT perf: It is fairly costly to perform this test, as we could easily early-out after we have SharedPCHHeaderFile and not bother testing the rest of the files in the module.  But this can be useful to find abusive modules that include a shared PCH header in the bowels of a non-PCH private source file.
-									Console.WriteLine( "WARNING: File '{0}' doesn't use same 'shared' precompiled header as other files in this module: '{1}' vs '{2}'.  This can greatly impact compile times.  Make sure that your module's private PCH header includes the 'largest' shared PCH header that your module uses", 
-										CPPFile.AbsolutePath, SharedPCHHeaderFile, LargestIncludedSharedPCHHeaderFile.PCHHeaderFile );
+									SharedPCHModuleName = CurSharedPCHHeaderFile.Module.Name;
+									SharedPCHHeaderFile = CurSharedPCHHeaderFile.PCHHeaderFile;
+									LargestSharedPCHHeaderFileIndex = SharedPCHHeaderFileIndex;
+									break;
 								}
 							}
-							else
+
+							if( LargestSharedPCHHeaderFileIndex == GlobalCompileEnvironment.SharedPCHHeaderFiles.Count - 1 )
 							{
-								Log.TraceVerbose("File {0} doesn't use a Shared PCH!", CPPFile.AbsolutePath);
+								// We've determined that the module is using our most complex PCH header, so we can early-out
+								break;
 							}
-
-							SharedPCHTime += ( DateTime.UtcNow - SharedPCHStartTime ).TotalSeconds;
 						}
-						else
-						{
-							Log.TraceVerbose("File '{0}' cannot create or use Shared PCHs, because its module '{1}' needs its own private PCH", CPPFile.AbsolutePath, this.Name);
-						}
-					}
-				}
-				TotalSharedPCHTime += SharedPCHTime;
 
-				if( BuildConfiguration.bPrintDebugInfo )
-				{
-					Log.TraceVerbose( "{0} PCH files for module {1}:", UsageMapPCH.Count, Name );
-					int MostFilesIncluded = 0;
-					foreach( var CurPCH in UsageMapPCH )
-					{
-						if( CurPCH.Value.Count > MostFilesIncluded )
+						// Did we not find a shared PCH header that is being included by this module?  This could happen if the module is not including Core.h, even indirectly.
+						if( String.IsNullOrEmpty( SharedPCHModuleName ) )
 						{
-							MostFilesIncluded = CurPCH.Value.Count;
+							throw new BuildException( "Module {0} doesn't use a Shared PCH!  Please add a dependency on a Shared PCH module to this module's dependency list", this.Name);
 						}
-
-						Log.TraceVerbose("   {0}  ({1} files including it: {2}, ...)", CurPCH.Key, CurPCH.Value.Count, CurPCH.Value[0].AbsolutePath);
-					}
-				}
-
-				if( UsageMapPCH.Count > 1 )
-				{
-					// Keep track of the PCH file that is most used within this module
-					string MostFilesAreIncludingPCH = string.Empty;
-					int MostFilesIncluded = 0;
-					foreach( var CurPCH in UsageMapPCH.Where( PCH => PCH.Value.Count > MostFilesIncluded ) )
-					{
-						MostFilesAreIncludingPCH = CurPCH.Key;
-						MostFilesIncluded        = CurPCH.Value.Count;
 					}
-
-					// Find all of the files that are not including our "best" PCH header
-					var FilesNotIncludingBestPCH = new StringBuilder();
-					foreach( var CurPCH in UsageMapPCH.Where( PCH => PCH.Key != MostFilesAreIncludingPCH ) )
+					else
 					{
-						foreach( var SourceFile in CurPCH.Value )
-						{
-							FilesNotIncludingBestPCH.AppendFormat( "{0} (including {1})\n", SourceFile.AbsolutePath, CurPCH.Key );
-						}
+						Log.TraceVerbose("Module '{0}' cannot create or use Shared PCHs, because it needs its own private PCH", this.Name);
 					}
-
-					// Bail out and let the user know which source files may need to be fixed up
-					throw new BuildException(
-						"All source files in module \"{0}\" must include the same precompiled header first.  Currently \"{1}\" is included by most of the source files.  The following source files are not including \"{1}\" as their first include:\n\n{2}",
-						Name,
-						MostFilesAreIncludingPCH,
-						FilesNotIncludingBestPCH );
 				}
 
+
 				// The precompiled header environment for all source files in this module that use a precompiled header, if we even need one
 				PrecompileHeaderEnvironment ModulePCHEnvironment = null;
 
@@ -1351,7 +1285,7 @@ namespace UnrealBuildTool
 						{
 							// @todo SharedPCH: Ideally we would exhaustively check for a compatible compile environment (definitions, imports/exports, etc)
 							//    Currently, it's possible for the shared PCH to be compiled differently depending on which module UBT happened to have
-							//    include it first during the build phase.  This could create problems with determinstic builds, or turn up compile
+							//    include it first during the build phase.  This could create problems with deterministic builds, or turn up compile
 							//    errors unexpectedly due to compile environment differences.
 							Log.TraceVerbose("Module " + Name + " uses existing Shared PCH '" + ModulePCHEnvironment.PrecompiledHeaderIncludeFilename + "' (from module " + ModulePCHEnvironment.ModuleName + ")");
 						}
@@ -1426,7 +1360,7 @@ namespace UnrealBuildTool
 					if (bModuleUsesUnityBuild)
 					{
 						// unity files generated for only the set of files which share the same PCH environment
-						CPPFilesToBuild = Unity.GenerateUnityCPPs( CPPFilesToBuild, ModulePCHCompileEnvironment, Name );
+						CPPFilesToBuild = Unity.GenerateUnityCPPs( Target, CPPFilesToBuild, ModulePCHCompileEnvironment, Name );
 					}
 
 					// Check if there are enough unity files to warrant pch generation (and we haven't already generated the shared one)
@@ -1451,6 +1385,7 @@ namespace UnrealBuildTool
 							}
 
 							var PCHOutput = PrecompileHeaderEnvironment.GeneratePCHCreationAction( 
+								Target,
 								CPPFilesToBuild[0].PCHHeaderNameInCode,
 								ModulePCHEnvironment.PrecompiledHeaderIncludeFilename,
 								ModuleCompileEnvironment, 
@@ -1480,16 +1415,14 @@ namespace UnrealBuildTool
 						CPPCompileEnvironment = ModulePCHCompileEnvironment;
 					}
 
-					LinkInputFiles.AddRange( CPPCompileEnvironment.CompileFiles( CPPFilesToBuild, Name ).ObjectFiles );
+					LinkInputFiles.AddRange( CPPCompileEnvironment.CompileFiles( Target, CPPFilesToBuild, Name ).ObjectFiles );
 					bWasModuleCodeCompiled = true;
 				}
 
 				if( BuildConfiguration.bPrintPerformanceInfo )
 				{
-					var PCHTime = ( DateTime.UtcNow - PCHTimerStart ).TotalSeconds;
-					Trace.TraceInformation( "PCH time for " + Name + " is " + PCHTime + "s (shared PCHs: " + SharedPCHTime + "s)" );
-
-					TotalPCHTime += PCHTime;
+					var PCHGenTime = ( DateTime.UtcNow - PCHGenTimerStart ).TotalSeconds;
+					TotalPCHGenTime += PCHGenTime;
 				}
 			}
 
@@ -1498,26 +1431,29 @@ namespace UnrealBuildTool
 				var CPPFilesToCompile = SourceFilesToBuild.CPPFiles;
 				if (bModuleUsesUnityBuild)
 				{
-					CPPFilesToCompile = Unity.GenerateUnityCPPs( CPPFilesToCompile, ModuleCompileEnvironment, Name );
+					CPPFilesToCompile = Unity.GenerateUnityCPPs( Target, CPPFilesToCompile, CPPCompileEnvironment, Name );
 				}
-				LinkInputFiles.AddRange( CPPCompileEnvironment.CompileFiles( CPPFilesToCompile, Name ).ObjectFiles );
+				LinkInputFiles.AddRange( CPPCompileEnvironment.CompileFiles( Target, CPPFilesToCompile, Name ).ObjectFiles );
 			}
 
-			if (AutoGenerateCppInfo != null && !ModuleCompileEnvironment.bHackHeaderGenerator)
+			if (AutoGenerateCppInfo != null && !CPPCompileEnvironment.bHackHeaderGenerator)
 			{
 				var GeneratedCppFileItem = FileItem.GetItemByPath( AutoGenerateCppInfo.Filename );
 
-				LinkInputFiles.AddRange( CPPCompileEnvironment.CompileFiles( new List<FileItem>{ GeneratedCppFileItem }, Name ).ObjectFiles );
+				CachePCHUsageForModuleSourceFile( this.Target, CPPCompileEnvironment, GeneratedCppFileItem );
+
+				// @todo fastubt: Check for ALL other places where we might be injecting .cpp or .rc files for compiling without caching CachedCPPEnvironment first (anything platform specfic?)
+				LinkInputFiles.AddRange( CPPCompileEnvironment.CompileFiles( Target, new List<FileItem>{ GeneratedCppFileItem }, Name ).ObjectFiles );
 			}
 
 			// Compile C files directly.
-			LinkInputFiles.AddRange(ModuleCompileEnvironment.CompileFiles(SourceFilesToBuild.CFiles, Name).ObjectFiles);
+			LinkInputFiles.AddRange(CPPCompileEnvironment.CompileFiles( Target, SourceFilesToBuild.CFiles, Name).ObjectFiles);
 
 			// Compile CC files directly.
-			LinkInputFiles.AddRange(ModuleCompileEnvironment.CompileFiles(SourceFilesToBuild.CCFiles, Name).ObjectFiles);
+			LinkInputFiles.AddRange(CPPCompileEnvironment.CompileFiles( Target, SourceFilesToBuild.CCFiles, Name).ObjectFiles);
 
 			// Compile MM files directly.
-			LinkInputFiles.AddRange(ModuleCompileEnvironment.CompileFiles(SourceFilesToBuild.MMFiles, Name).ObjectFiles);
+			LinkInputFiles.AddRange(CPPCompileEnvironment.CompileFiles( Target, SourceFilesToBuild.MMFiles, Name).ObjectFiles);
 
 			// If we're building Rocket, generate a static library for this module
 			if(RedistStaticLibraryPath != null)
@@ -1537,7 +1473,7 @@ namespace UnrealBuildTool
 			}
 
 			// Compile RC files.
-			LinkInputFiles.AddRange(ModuleCompileEnvironment.CompileRCFiles(SourceFilesToBuild.RCFiles).ObjectFiles);
+			LinkInputFiles.AddRange(CPPCompileEnvironment.CompileRCFiles(Target, SourceFilesToBuild.RCFiles).ObjectFiles);
 
 			return LinkInputFiles;
 		}
@@ -1548,7 +1484,9 @@ namespace UnrealBuildTool
 			var SharedPCHHeaderFileCopy = SharedPCHHeaderFile;
 			var SharedPCHEnvironment = GlobalCompileEnvironment.SharedPCHEnvironments.Find(Env => Env.PrecompiledHeaderIncludeFilename == SharedPCHHeaderFileCopy);
 			if (SharedPCHEnvironment == null)
+			{ 
 				return null;
+			}
 
 			// Don't mix CLR modes
 			if (SharedPCHEnvironment.CLRMode != ModuleCompileEnvironment.Config.CLRMode)
@@ -1585,25 +1523,142 @@ namespace UnrealBuildTool
 			return SharedPCHEnvironment;
 		}
 
-		public void ProcessAllCppDependencies(CPPEnvironment ModuleCompileEnvironment)
+		public static FileItem CachePCHUsageForModuleSourceFile(UEBuildTarget Target, CPPEnvironment ModuleCompileEnvironment, FileItem CPPFile)
+		{
+			var PCHCacheTimerStart = DateTime.UtcNow;
+
+			var BuildPlatform = UEBuildPlatform.GetBuildPlatformForCPPTargetPlatform( ModuleCompileEnvironment.Config.Target.Platform );
+			var IncludePathsToSearch = ModuleCompileEnvironment.GetIncludesPathsToSearch( CPPFile );
+
+			// Store the module compile environment along with the .cpp file.  This is so that we can use it later on when looking
+			// for header dependencies
+			CPPFile.CachedCPPEnvironment = ModuleCompileEnvironment;
+
+			var PCHFile = CachePCHUsageForCPPFile( Target, CPPFile, BuildPlatform, IncludePathsToSearch, ModuleCompileEnvironment.IncludeFileSearchDictionary );
+
+			if( BuildConfiguration.bPrintPerformanceInfo )
+			{
+				var PCHCacheTime = ( DateTime.UtcNow - PCHCacheTimerStart ).TotalSeconds;
+				TotalPCHCacheTime += PCHCacheTime;
+			}
+			
+			return PCHFile;
+		}
+
+
+		public void CachePCHUsageForModuleSourceFiles(CPPEnvironment ModuleCompileEnvironment)
 		{
 			if( ProcessedDependencies == null )
 			{
+				var PCHCacheTimerStart = DateTime.UtcNow;
+
+				var BuildPlatform = UEBuildPlatform.GetBuildPlatformForCPPTargetPlatform( ModuleCompileEnvironment.Config.Target.Platform );
+
+				bool bFoundAProblemWithPCHs = false;
+
 				FileItem UniquePCH = null;
 				foreach( var CPPFile in SourceFilesFound.CPPFiles )
 				{
+					// Build a single list of include paths to search.
+					var IncludePathsToSearch = ModuleCompileEnvironment.GetIncludesPathsToSearch( CPPFile );   // @todo fastubt: This call could be made faster (list caching)
+
+					// Store the module compile environment along with the .cpp file.  This is so that we can use it later on when looking
+					// for header dependencies
+					CPPFile.CachedCPPEnvironment = ModuleCompileEnvironment;
+
 					// Find headers used by the source file.
-					var PCH = ProcessDependencies(CPPFile, ModuleCompileEnvironment);
-					UniquePCH = UniquePCH ?? PCH;
+					var PCH = UEBuildModuleCPP.CachePCHUsageForCPPFile(Target, CPPFile, BuildPlatform, IncludePathsToSearch, ModuleCompileEnvironment.IncludeFileSearchDictionary);
+					if( PCH == null )
+					{
+						throw new BuildException( "Source file \"{0}\" is not including any headers.  We expect all modules to include a header file for precompiled header generation.  Please add an #include statement.", CPPFile.AbsolutePath );
+					}
+
+					if( UniquePCH == null )
+					{
+						UniquePCH = PCH;
+					}
+					else if( !UniquePCH.Info.Name.Equals( PCH.Info.Name, StringComparison.InvariantCultureIgnoreCase ) )		// @todo fastubt: We do a string compare on the file name (not path) here, because sometimes the include resolver will pick an Intermediate copy of a PCH header file and throw off our comparisons
+					{
+						// OK, looks like we have multiple source files including a different header file first.  We'll keep track of this and print out
+						// helpful information afterwards.
+						bFoundAProblemWithPCHs = true;
+					}
 				}
 
 				ProcessedDependencies = new ProcessedDependenciesClass{ UniquePCHHeaderFile = UniquePCH };
+
+
+				if( bFoundAProblemWithPCHs )
+				{
+					// Map from pch header string to the source files that use that PCH
+					var UsageMapPCH = new Dictionary<string, List<FileItem>>( StringComparer.InvariantCultureIgnoreCase );
+					foreach( var CPPFile in SourceFilesToBuild.CPPFiles )
+					{
+						// Create a new entry if not in the pch usage map
+						UsageMapPCH.GetOrAddNew( CPPFile.PrecompiledHeaderIncludeFilename ).Add( CPPFile );
+					}
+
+					if( BuildConfiguration.bPrintDebugInfo )
+					{
+						Log.TraceVerbose( "{0} PCH files for module {1}:", UsageMapPCH.Count, Name );
+						int MostFilesIncluded = 0;
+						foreach( var CurPCH in UsageMapPCH )
+						{
+							if( CurPCH.Value.Count > MostFilesIncluded )
+							{
+								MostFilesIncluded = CurPCH.Value.Count;
+							}
+
+							Log.TraceVerbose("   {0}  ({1} files including it: {2}, ...)", CurPCH.Key, CurPCH.Value.Count, CurPCH.Value[0].AbsolutePath);
+						}
+					}
+
+					if( UsageMapPCH.Count > 1 )
+					{
+						// Keep track of the PCH file that is most used within this module
+						string MostFilesAreIncludingPCH = string.Empty;
+						int MostFilesIncluded = 0;
+						foreach( var CurPCH in UsageMapPCH.Where( PCH => PCH.Value.Count > MostFilesIncluded ) )
+						{
+							MostFilesAreIncludingPCH = CurPCH.Key;
+							MostFilesIncluded        = CurPCH.Value.Count;
+						}
+
+						// Find all of the files that are not including our "best" PCH header
+						var FilesNotIncludingBestPCH = new StringBuilder();
+						foreach( var CurPCH in UsageMapPCH.Where( PCH => PCH.Key != MostFilesAreIncludingPCH ) )
+						{
+							foreach( var SourceFile in CurPCH.Value )
+							{
+								FilesNotIncludingBestPCH.AppendFormat( "{0} (including {1})\n", SourceFile.AbsolutePath, CurPCH.Key );
+							}
+						}
+
+						// Bail out and let the user know which source files may need to be fixed up
+						throw new BuildException(
+							"All source files in module \"{0}\" must include the same precompiled header first.  Currently \"{1}\" is included by most of the source files.  The following source files are not including \"{1}\" as their first include:\n\n{2}",
+							Name,
+							MostFilesAreIncludingPCH,
+							FilesNotIncludingBestPCH );
+					}
+				}
+				
+				if( BuildConfiguration.bPrintPerformanceInfo )
+				{
+					var PCHCacheTime = ( DateTime.UtcNow - PCHCacheTimerStart ).TotalSeconds;
+					TotalPCHCacheTime += PCHCacheTime;
+				}
 			}
 		}
 
-		private FileItem ProcessDependencies(FileItem CPPFile, CPPEnvironment ModuleCompileEnvironment)
+
+
+		private static FileItem CachePCHUsageForCPPFile(UEBuildTarget Target, FileItem CPPFile, IUEBuildPlatform BuildPlatform, List<string> IncludePathsToSearch, Dictionary<string, FileItem> IncludeFileSearchDictionary)
 		{
-			List<DependencyInclude> DirectIncludeFilenames = CPPEnvironment.GetDirectIncludeDependencies(CPPFile, ModuleCompileEnvironment.Config.Target.Platform);
+			// @todo fastubt: We don't really need to scan every file looking for PCH headers, just need one.  The rest is just for error checking.
+			// @todo fastubt: We don't need all of the direct includes either.  We just need the first, unless we want to check for errors.
+			bool HasUObjects;
+			List<DependencyInclude> DirectIncludeFilenames = CPPEnvironment.GetDirectIncludeDependencies(Target, CPPFile, BuildPlatform, bOnlyCachedDependencies:false, HasUObjects:out HasUObjects);
 			if (BuildConfiguration.bPrintDebugInfo)
 			{
 				Log.TraceVerbose("Found direct includes for {0}: {1}", Path.GetFileName(CPPFile.AbsolutePath), string.Join(", ", DirectIncludeFilenames.Select(F => F.IncludeName)));
@@ -1632,16 +1687,14 @@ namespace UnrealBuildTool
 				return FileItem.GetItemByFullPath(CPPFile.PrecompiledHeaderIncludeFilename);
 			}
 
-			string SourceFilesDirectory = Path.GetDirectoryName(CPPFile.AbsolutePath);
-
 			// search the include paths to resolve the file.
-			FileItem PrecompiledHeaderIncludeFile = ModuleCompileEnvironment.FindIncludedFile(CPPFile.PCHHeaderNameInCode, !BuildConfiguration.bCheckExternalHeadersForModification, SourceFilesDirectory);
+			FileItem PrecompiledHeaderIncludeFile = CPPEnvironment.FindIncludedFile(CPPFile.PCHHeaderNameInCode, !BuildConfiguration.bCheckExternalHeadersForModification, IncludePathsToSearch, IncludeFileSearchDictionary );
 			if (PrecompiledHeaderIncludeFile == null)
 			{ 
-				throw new BuildException("The first include statement in source file '{0}' is trying to include the file '{1}' as the precompiled header for module '{2}', but that file could not be located in any of the module's include search paths.", CPPFile.AbsolutePath, CPPFile.PCHHeaderNameInCode, this.Name);
+				throw new BuildException("The first include statement in source file '{0}' is trying to include the file '{1}' as the precompiled header, but that file could not be located in any of the module's include search paths.", CPPFile.AbsolutePath, CPPFile.PCHHeaderNameInCode);
 			}
 
-			CPPEnvironment.IncludeDependencyCache.CacheResolvedIncludeFullPath(CPPFile, 0, PrecompiledHeaderIncludeFile.AbsolutePath);
+			CPPEnvironment.IncludeDependencyCache[Target].CacheResolvedIncludeFullPath(CPPFile, 0, PrecompiledHeaderIncludeFile.AbsolutePath);
 			CPPFile.PrecompiledHeaderIncludeFilename = PrecompiledHeaderIncludeFile.AbsolutePath;
 
 			return PrecompiledHeaderIncludeFile;
@@ -1768,11 +1821,12 @@ namespace UnrealBuildTool
 
 		private UHTModuleInfo CachedModuleUHTInfo = null;
 
-		/// Total time spent figuring out which PCH to use for each module and source file
-		public static double TotalPCHTime = 0.0;
+		/// Total time spent generating PCHs for modules (not actually compiling, but generating the PCH's input data)
+		public static double TotalPCHGenTime = 0.0;
+
+		/// Time spent caching which PCH header is included by each module and source file
+		public static double TotalPCHCacheTime = 0.0;
 
-		/// Total time spent processing which shared PCH headers are needed for each module and source file
-		public static double TotalSharedPCHTime = 0.0;
 
 
 		/// <summary>
@@ -1792,31 +1846,38 @@ namespace UnrealBuildTool
 			var ClassesFolder = Path.Combine(this.ModuleDirectory, "Classes");
 			var PublicFolder  = Path.Combine(this.ModuleDirectory, "Public");
 
+			var BuildPlatform = UEBuildPlatform.GetBuildPlatform(Target.Platform);
+
 			var ClassesFiles = Directory.GetFiles( this.ModuleDirectory, "*.h", SearchOption.AllDirectories );
 			foreach (var ClassHeader in ClassesFiles)
 			{
+				// Check to see if we know anything about this file.  If we have up-to-date cached information about whether it has
+				// UObjects or not, we can skip doing a test here.
 				var UObjectHeaderFileItem = FileItem.GetExistingItemByPath( ClassHeader );
-				var FileContents = Utils.ReadAllText(UObjectHeaderFileItem.AbsolutePath);
-				if (!Regex.IsMatch(FileContents, "^\\s*U(CLASS|STRUCT|ENUM|INTERFACE|DELEGATE)\\b", RegexOptions.Multiline))
-				{ 
-					continue;
-				}
 
-				if (UObjectHeaderFileItem.AbsolutePath.StartsWith(ClassesFolder))
-				{
-					_AllClassesHeaders.Add(UObjectHeaderFileItem);
-				}
-				else if (UObjectHeaderFileItem.AbsolutePath.StartsWith(PublicFolder))
-				{
-					_PublicUObjectHeaders.Add(UObjectHeaderFileItem);
-				}
-				else
-				{
-					_PrivateUObjectHeaders.Add(UObjectHeaderFileItem);
+				bool HasUObjects;
+				var DirectIncludes = CPPEnvironment.GetDirectIncludeDependencies( Target, UObjectHeaderFileItem, BuildPlatform, bOnlyCachedDependencies:false, HasUObjects:out HasUObjects );
+				Debug.Assert( DirectIncludes != null );
+
+				if (HasUObjects)
+				{ 
+					if (UObjectHeaderFileItem.AbsolutePath.StartsWith(ClassesFolder))
+					{
+						_AllClassesHeaders.Add(UObjectHeaderFileItem);
+					}
+					else if (UObjectHeaderFileItem.AbsolutePath.StartsWith(PublicFolder))
+					{
+						_PublicUObjectHeaders.Add(UObjectHeaderFileItem);
+					}
+					else
+					{
+						_PrivateUObjectHeaders.Add(UObjectHeaderFileItem);
+					}
 				}
 			}
 
-			CachedModuleUHTInfo = new UHTModuleInfo {
+			CachedModuleUHTInfo = new UHTModuleInfo 
+			{
 				ModuleName                  = this.Name,
 				ModuleDirectory             = this.ModuleDirectory,
 				ModuleType					= this.Type.ToString(),
@@ -1827,8 +1888,8 @@ namespace UnrealBuildTool
 
 			return CachedModuleUHTInfo;
 		}
-
-		public override void GetAllDependencyModules( ref Dictionary<string, UEBuildModule> ReferencedModules, ref List<UEBuildModule> OrderedModules, bool bIncludeDynamicallyLoaded, bool bForceCircular )
+	
+		public override void GetAllDependencyModules( ref Dictionary<string, UEBuildModule> ReferencedModules, ref List<UEBuildModule> OrderedModules, bool bIncludeDynamicallyLoaded, bool bForceCircular, bool bOnlyDirectDependencies )
 		{
 			var AllModuleNames = new List<string>();
 			AllModuleNames.AddRange(PrivateDependencyModuleNames);
@@ -1850,8 +1911,11 @@ namespace UnrealBuildTool
 						var Module = Target.GetModuleByName( DependencyName );
 						ReferencedModules[ DependencyName ] = Module;
 
-						// Recurse into dependent modules first
-						Module.GetAllDependencyModules(ref ReferencedModules, ref OrderedModules, bIncludeDynamicallyLoaded, bForceCircular);
+						if( !bOnlyDirectDependencies )
+						{ 
+							// Recurse into dependent modules first
+							Module.GetAllDependencyModules(ref ReferencedModules, ref OrderedModules, bIncludeDynamicallyLoaded, bForceCircular, bOnlyDirectDependencies);
+						}
 
 						OrderedModules.Add( Module );
 					}
diff --git a/Engine/Source/Programs/UnrealBuildTool/Configuration/UEBuildTarget.cs b/Engine/Source/Programs/UnrealBuildTool/Configuration/UEBuildTarget.cs
--- a/Engine/Source/Programs/UnrealBuildTool/Configuration/UEBuildTarget.cs
+++ b/Engine/Source/Programs/UnrealBuildTool/Configuration/UEBuildTarget.cs
@@ -981,8 +981,14 @@ namespace UnrealBuildTool
 					}
 				}
 
-				// Delete the dependency cache
+				// Delete the dependency caches
 				{
+					var FlatCPPIncludeDependencyCacheFilename = FlatCPPIncludeDependencyCache.GetDependencyCachePathForTarget(this);
+					if (File.Exists(FlatCPPIncludeDependencyCacheFilename))
+					{
+						Log.TraceVerbose("\tDeleting " + FlatCPPIncludeDependencyCacheFilename);
+						CleanFile(FlatCPPIncludeDependencyCacheFilename);
+					}
 					var DependencyCacheFilename = DependencyCache.GetDependencyCachePathForTarget(this);
 					if (File.Exists(DependencyCacheFilename))
 					{
@@ -1784,7 +1790,7 @@ namespace UnrealBuildTool
 				NewModule.bIncludedInTarget = true;
 
 				// Process dependencies for this new module
-				NewModule.ProcessAllCppDependencies(NewModule.CreateModuleCompileEnvironment(GlobalCompileEnvironment));
+				NewModule.CachePCHUsageForModuleSourceFiles(NewModule.CreateModuleCompileEnvironment(GlobalCompileEnvironment));
 
 				// Add module to binary
 				ExecutableBinary.AddModule(NewModule.Name);
@@ -1827,7 +1833,7 @@ namespace UnrealBuildTool
 									var ModuleDependencies = new Dictionary<string, UEBuildModule>();
 									var ModuleList = new List<UEBuildModule>();
 									bool bIncludeDynamicallyLoaded = false;
-									CPPModule.GetAllDependencyModules(ref ModuleDependencies, ref ModuleList, bIncludeDynamicallyLoaded, bForceCircular: false);
+									CPPModule.GetAllDependencyModules(ref ModuleDependencies, ref ModuleList, bIncludeDynamicallyLoaded, bForceCircular: false, bOnlyDirectDependencies:false);
 
 									// Figure out where to insert the shared PCH into our list, based off the module dependency ordering
 									int InsertAtIndex = SharedPCHHeaderFiles.Count;
diff --git a/Engine/Source/Programs/UnrealBuildTool/HTML5/HTML5ToolChain.cs b/Engine/Source/Programs/UnrealBuildTool/HTML5/HTML5ToolChain.cs
--- a/Engine/Source/Programs/UnrealBuildTool/HTML5/HTML5ToolChain.cs
+++ b/Engine/Source/Programs/UnrealBuildTool/HTML5/HTML5ToolChain.cs
@@ -305,11 +305,11 @@ namespace UnrealBuildTool
 			Log.TraceInformation(Output);				// To preserve readable output log
 		}
 
-		public override CPPOutput CompileCPPFiles(CPPEnvironment CompileEnvironment, List<FileItem> SourceFiles, string ModuleName)
+		public override CPPOutput CompileCPPFiles(UEBuildTarget Target, CPPEnvironment CompileEnvironment, List<FileItem> SourceFiles, string ModuleName)
 		{
 			if (CompileEnvironment.Config.Target.Architecture == "-win32")
 			{
-				return base.CompileCPPFiles(CompileEnvironment, SourceFiles, ModuleName);
+				return base.CompileCPPFiles(Target, CompileEnvironment, SourceFiles, ModuleName);
 			}
 
 			string Arguments = GetCLArguments_Global(CompileEnvironment);
@@ -340,6 +340,8 @@ namespace UnrealBuildTool
             if (ModuleName == "Launch")
                 SourceFiles.Add(FileItem.GetItemByPath(BaseSDKPath + "/system/lib/libcxxabi/src/cxa_demangle.cpp")); 
         
+			var BuildPlatform = UEBuildPlatform.GetBuildPlatformForCPPTargetPlatform(CompileEnvironment.Config.Target.Platform);
+
 			foreach (FileItem SourceFile in SourceFiles)
 			{
 				Action CompileAction = new Action(ActionType.Compile);
@@ -347,9 +349,9 @@ namespace UnrealBuildTool
                 
 				// Add the C++ source file and its included files to the prerequisite item list.
 				CompileAction.PrerequisiteItems.Add(SourceFile);
-				foreach (FileItem IncludedFile in CompileEnvironment.GetIncludeDependencies(SourceFile))
 				{
-					CompileAction.PrerequisiteItems.Add(IncludedFile);
+					var IncludedFileList = CPPEnvironment.FindAndCacheAllIncludedFiles( Target, SourceFile, BuildPlatform, CompileEnvironment.GetIncludesPathsToSearch( SourceFile ), CompileEnvironment.IncludeFileSearchDictionary, bOnlyCachedDependencies:BuildConfiguration.bUseExperimentalFastDependencyScan );
+					CompileAction.PrerequisiteItems.AddRange( IncludedFileList );
 				}
 
 				// Add the source file path to the command-line.
@@ -403,13 +405,13 @@ namespace UnrealBuildTool
 			return Result;
 		}
 
-		public override CPPOutput CompileRCFiles(CPPEnvironment Environment, List<FileItem> RCFiles)
+		public override CPPOutput CompileRCFiles(UEBuildTarget Target, CPPEnvironment Environment, List<FileItem> RCFiles)
 		{
 			CPPOutput Result = new CPPOutput();
 
 			if (Environment.Config.Target.Architecture == "-win32")
 			{
-				return base.CompileRCFiles(Environment, RCFiles);
+				return base.CompileRCFiles(Target, Environment, RCFiles);
 			}
 
 			return Result;
diff --git a/Engine/Source/Programs/UnrealBuildTool/IOS/IOSToolChain.cs b/Engine/Source/Programs/UnrealBuildTool/IOS/IOSToolChain.cs
--- a/Engine/Source/Programs/UnrealBuildTool/IOS/IOSToolChain.cs
+++ b/Engine/Source/Programs/UnrealBuildTool/IOS/IOSToolChain.cs
@@ -448,7 +448,7 @@ namespace UnrealBuildTool
 			return Result;
 		}
 		
-		public override CPPOutput CompileCPPFiles(CPPEnvironment CompileEnvironment, List<FileItem> SourceFiles, string ModuleName)
+		public override CPPOutput CompileCPPFiles(UEBuildTarget Target, CPPEnvironment CompileEnvironment, List<FileItem> SourceFiles, string ModuleName)
 		{
 			string Arguments = GetCompileArguments_Global(CompileEnvironment);
 			string PCHArguments = "";
@@ -495,6 +495,8 @@ namespace UnrealBuildTool
 				Arguments += string.Format(" -D\"{0}\"", Definition);
 			}
 
+			var BuildPlatform = UEBuildPlatform.GetBuildPlatformForCPPTargetPlatform(CompileEnvironment.Config.Target.Platform);
+
 			CPPOutput Result = new CPPOutput();
 			// Create a compile action for each source file.
 			foreach (FileItem SourceFile in SourceFiles)
@@ -545,14 +547,18 @@ namespace UnrealBuildTool
 					QueueFileForBatchUpload(SourceFile);
 				}
 
-				foreach (FileItem IncludedFile in CompileEnvironment.GetIncludeDependencies(SourceFile))
 				{
-					if (ExternalExecution.GetRuntimePlatform() != UnrealTargetPlatform.Mac)
+					var IncludedFileList = CPPEnvironment.FindAndCacheAllIncludedFiles( Target, SourceFile, BuildPlatform, CompileEnvironment.GetIncludesPathsToSearch( SourceFile ), CompileEnvironment.IncludeFileSearchDictionary, bOnlyCachedDependencies:BuildConfiguration.bUseExperimentalFastDependencyScan );
+					foreach (FileItem IncludedFile in IncludedFileList)
 					{
-						QueueFileForBatchUpload(IncludedFile);
-					}
+						CompileAction.PrerequisiteItems.Add(IncludedFile);
 
-					CompileAction.PrerequisiteItems.Add(IncludedFile);
+						if (!BuildConfiguration.bUseExperimentalFastDependencyScan &&	// With fast dependency scanning, we will not have an exhaustive list of dependencies here.  We rely on PostCodeGeneration() to upload these files.
+							ExternalExecution.GetRuntimePlatform() != UnrealTargetPlatform.Mac)
+						{
+							QueueFileForBatchUpload(IncludedFile);
+						}
+					}
 				}
 
 				if (CompileEnvironment.Config.PrecompiledHeaderAction == PrecompiledHeaderAction.Create)
diff --git a/Engine/Source/Programs/UnrealBuildTool/Linux/LinuxToolChain.cs b/Engine/Source/Programs/UnrealBuildTool/Linux/LinuxToolChain.cs
--- a/Engine/Source/Programs/UnrealBuildTool/Linux/LinuxToolChain.cs
+++ b/Engine/Source/Programs/UnrealBuildTool/Linux/LinuxToolChain.cs
@@ -359,7 +359,7 @@ namespace UnrealBuildTool
 
 		private static List<FileItem> BundleDependencies = new List<FileItem>();
 
-        public override CPPOutput CompileCPPFiles(CPPEnvironment CompileEnvironment, List<FileItem> SourceFiles, string ModuleName)
+        public override CPPOutput CompileCPPFiles(UEBuildTarget Target, CPPEnvironment CompileEnvironment, List<FileItem> SourceFiles, string ModuleName)
         {
             string Arguments = GetCLArguments_Global(CompileEnvironment);
             string PCHArguments = "";
@@ -393,6 +393,8 @@ namespace UnrealBuildTool
                 Arguments += string.Format(" -D \"{0}\"", Definition);
             }
 
+			var BuildPlatform = UEBuildPlatform.GetBuildPlatformForCPPTargetPlatform(CompileEnvironment.Config.Target.Platform);
+
             // Create a compile action for each source file.
             CPPOutput Result = new CPPOutput();
             foreach (FileItem SourceFile in SourceFiles)
@@ -436,10 +438,10 @@ namespace UnrealBuildTool
 
                 // Add the C++ source file and its included files to the prerequisite item list.
                 CompileAction.PrerequisiteItems.Add(SourceFile);
-                foreach (FileItem IncludedFile in CompileEnvironment.GetIncludeDependencies(SourceFile))
-                {
-                    CompileAction.PrerequisiteItems.Add(IncludedFile);
-                }
+				{
+					var IncludedFileList = CPPEnvironment.FindAndCacheAllIncludedFiles( Target, SourceFile, BuildPlatform, CompileEnvironment.GetIncludesPathsToSearch( SourceFile ), CompileEnvironment.IncludeFileSearchDictionary, bOnlyCachedDependencies:BuildConfiguration.bUseExperimentalFastDependencyScan );
+					CompileAction.PrerequisiteItems.AddRange( IncludedFileList );
+				}
 
                 if (CompileEnvironment.Config.PrecompiledHeaderAction == PrecompiledHeaderAction.Create)
                 {
diff --git a/Engine/Source/Programs/UnrealBuildTool/Mac/MacToolChain.cs b/Engine/Source/Programs/UnrealBuildTool/Mac/MacToolChain.cs
--- a/Engine/Source/Programs/UnrealBuildTool/Mac/MacToolChain.cs
+++ b/Engine/Source/Programs/UnrealBuildTool/Mac/MacToolChain.cs
@@ -218,7 +218,7 @@ namespace UnrealBuildTool
 			return Result;
 		}
 
-		public override CPPOutput CompileCPPFiles(CPPEnvironment CompileEnvironment, List<FileItem> SourceFiles, string ModuleName)
+		public override CPPOutput CompileCPPFiles(UEBuildTarget Target, CPPEnvironment CompileEnvironment, List<FileItem> SourceFiles, string ModuleName)
 		{
 			var Arguments = new StringBuilder();
 			var PCHArguments = new StringBuilder();
@@ -284,6 +284,8 @@ namespace UnrealBuildTool
 				Arguments.Append("\"");
 			}
 
+			var BuildPlatform = UEBuildPlatform.GetBuildPlatformForCPPTargetPlatform(CompileEnvironment.Config.Target.Platform);
+
 			CPPOutput Result = new CPPOutput();
 			// Create a compile action for each source file.
 			foreach (FileItem SourceFile in SourceFiles)
@@ -334,14 +336,18 @@ namespace UnrealBuildTool
 					QueueFileForBatchUpload(SourceFile);
 				}
 
-				foreach (FileItem IncludedFile in CompileEnvironment.GetIncludeDependencies(SourceFile))
 				{
-					if (ExternalExecution.GetRuntimePlatform() != UnrealTargetPlatform.Mac)
+					var IncludedFileList = CPPEnvironment.FindAndCacheAllIncludedFiles( Target, SourceFile, BuildPlatform, CompileEnvironment.GetIncludesPathsToSearch( SourceFile ), CompileEnvironment.IncludeFileSearchDictionary, bOnlyCachedDependencies:BuildConfiguration.bUseExperimentalFastDependencyScan );
+					foreach (FileItem IncludedFile in IncludedFileList)
 					{
-						QueueFileForBatchUpload(IncludedFile);
-					}
+						CompileAction.PrerequisiteItems.Add(IncludedFile);
 
-					CompileAction.PrerequisiteItems.Add(IncludedFile);
+						if (!BuildConfiguration.bUseExperimentalFastDependencyScan &&	// With fast dependency scanning, we will not have an exhaustive list of dependencies here.  We rely on PostCodeGeneration() to upload these files.
+							 ExternalExecution.GetRuntimePlatform() != UnrealTargetPlatform.Mac)
+						{
+							QueueFileForBatchUpload(IncludedFile);
+						}
+					}
 				}
 
 				if (CompileEnvironment.Config.PrecompiledHeaderAction == PrecompiledHeaderAction.Create)
diff --git a/Engine/Source/Programs/UnrealBuildTool/System/ActionGraph.cs b/Engine/Source/Programs/UnrealBuildTool/System/ActionGraph.cs
--- a/Engine/Source/Programs/UnrealBuildTool/System/ActionGraph.cs
+++ b/Engine/Source/Programs/UnrealBuildTool/System/ActionGraph.cs
@@ -155,8 +155,19 @@ namespace UnrealBuildTool
 		}
 
 		/** Builds a list of actions that need to be executed to produce the specified output items. */
-		static List<Action> GetActionsToExecute(List<FileItem> OutputItems, List<UEBuildTarget> Targets)
+		static List<Action> GetActionsToExecute(List<FileItem> OutputItems, List<UEBuildTarget> Targets, out Dictionary<UEBuildTarget, List<FileItem>> TargetToOutdatedPrerequisitesMap)
 		{
+			// @todo fastubt: We really want to be able to generate actions in a separate phase from actually building.  For example, we could generate
+			// actions at GPF-time, then save them out to be reloaded quickly later when building.  It means we would separate C++ include dependencies
+			// out from the normal action graph and check outdatedness at build time.  The module relationships and link dependencies would all still
+			// be in the static part of the action graph that was loaded from disk.	 You would probably need to re-run GPF after adding or removing source files,
+			// but ideally not after changing which headers are included in an existing source file!  Maybe we could make adding/removing source files work
+			// too, but it affects the Unity/PCH/Linker input, so it would be easier to not deal with this.  Still need to run UHT during build phase
+			// as well, which means we need to know which files have UObjects for the manifest.  This could be cached up front, and kept up to date as
+			// we rescan modified source files looking for changed includes.  UHT does not require spidering includes.
+
+			// @todo fastubt: Can we use directory changed notifications or directory timestamps to accelerate C++ file outdatedness checking?
+			
 			// Link producing actions to the items they produce.
 			LinkActionsAndItems();
 
@@ -181,6 +192,7 @@ namespace UnrealBuildTool
 			var OutdatedActionDictionary = new Dictionary<Action, bool>();
 			var HistoryList = new List<ActionHistory>();
 			var OpenHistoryFiles = new HashSet<string>(StringComparer.InvariantCultureIgnoreCase);
+			TargetToOutdatedPrerequisitesMap = new Dictionary<UEBuildTarget,List<FileItem>>();
 			foreach (var BuildTarget in Targets)
 			{
 				var HistoryFilename = ActionHistory.GeneratePathForTarget(BuildTarget);
@@ -189,11 +201,7 @@ namespace UnrealBuildTool
 					var History = new ActionHistory(HistoryFilename);
 					HistoryList.Add(History);
 					OpenHistoryFiles.Add(HistoryFilename);
-					var OutdatedActionsForTarget = GatherAllOutdatedActions(History);
-					foreach (var OutdatedAction in OutdatedActionsForTarget)
-					{
-						OutdatedActionDictionary.Add(OutdatedAction.Key, OutdatedAction.Value);
-					}
+					GatherAllOutdatedActions(BuildTarget, History, ref OutdatedActionDictionary, TargetToOutdatedPrerequisitesMap);
 				}
 			}
 			
@@ -629,7 +637,7 @@ namespace UnrealBuildTool
 		 * @param OutdatedActionDictionary - 
 		 * @return true if outdated
 		 */
-		static public bool IsActionOutdated(Action RootAction, ref Dictionary<Action, bool> OutdatedActionDictionary,ActionHistory ActionHistory)
+		static public bool IsActionOutdated(UEBuildTarget Target, Action RootAction, ref Dictionary<Action, bool> OutdatedActionDictionary,ActionHistory ActionHistory, Dictionary<UEBuildTarget, List<FileItem>> TargetToOutdatedPrerequisitesMap )
 		{
 			// Only compute the outdated-ness for actions that don't aren't cached in the outdated action dictionary.
 			bool bIsOutdated = false;
@@ -722,7 +730,7 @@ namespace UnrealBuildTool
 							// If the prerequisite is produced by an outdated action, then this action is outdated too.
 							if( PrerequisiteItem.ProducingAction != null )
 							{
-								if(IsActionOutdated(PrerequisiteItem.ProducingAction,ref OutdatedActionDictionary,ActionHistory))
+								if(IsActionOutdated(Target, PrerequisiteItem.ProducingAction,ref OutdatedActionDictionary,ActionHistory, TargetToOutdatedPrerequisitesMap))
 								{
 									Log.TraceVerbose(
 										"{0}: Prerequisite {1} is produced by outdated action.",
@@ -761,6 +769,54 @@ namespace UnrealBuildTool
 					}
 				}
 
+				// For compile actions, we have C++ files that are actually dependent on header files that could have been changed.  We only need to
+				// know about the set of header files that are included for files that are already determined to be out of date (such as if the file
+				// is missing or was modified.)  In the case that the file is out of date, we'll perform a deep scan to update our cached set of
+				// includes for this file, so that we'll be able to determine whether it is out of date next time very quickly.
+				if( BuildConfiguration.bUseExperimentalFastDependencyScan )
+				{ 
+					var DeepIncludeScanStartTime = DateTime.UtcNow;
+
+					// @todo fastubt: we may be scanning more files than we need to here -- indirectly outdated files are bIsOutdated=true by this point (for example basemost includes when deeper includes are dirty)
+					if( bIsOutdated && RootAction.ActionType == ActionType.Compile )
+					{
+						Log.TraceVerbose( "Outdated action: {0}", RootAction.StatusDescription );
+						foreach (FileItem PrerequisiteItem in RootAction.PrerequisiteItems)
+						{
+							if( PrerequisiteItem.CachedCPPEnvironment != null )
+							{
+								if( !IsCPPFile( PrerequisiteItem ) )
+								{
+									throw new BuildException( "Was only expecting C++ files to have CachedCPPEnvironments!" );
+								}
+								Log.TraceVerbose( "  -> DEEP include scan: {0}", PrerequisiteItem.AbsolutePath );
+
+								List<FileItem> OutdatedPrerequisites;
+								if( !TargetToOutdatedPrerequisitesMap.TryGetValue( Target, out OutdatedPrerequisites ) )
+								{
+									OutdatedPrerequisites = new List<FileItem>();
+									TargetToOutdatedPrerequisitesMap.Add( Target, OutdatedPrerequisites );
+								}
+
+								OutdatedPrerequisites.Add( PrerequisiteItem );
+							}
+							else if( IsCPPImplementationFile( PrerequisiteItem ) || IsCPPResourceFile( PrerequisiteItem ) )
+							{
+								if( PrerequisiteItem.CachedCPPEnvironment == null )
+								{
+									Log.TraceVerbose( "  -> WARNING: No CachedCPPEnvironment: {0}", PrerequisiteItem.AbsolutePath );
+								}
+							}
+						}
+					}
+
+					if( BuildConfiguration.bPrintPerformanceInfo )
+					{
+						double DeepIncludeScanTime = ( DateTime.UtcNow - DeepIncludeScanStartTime ).TotalSeconds;
+						TotalDeepIncludeScanTime += DeepIncludeScanTime;
+					}
+				}
+
 				// Cache the outdated-ness of this action.
 				OutdatedActionDictionary.Add(RootAction, bIsOutdated);
 			}
@@ -768,19 +824,17 @@ namespace UnrealBuildTool
 			return bIsOutdated;
 		}
 
+
 		/**
 		 * Builds a dictionary containing the actions from AllActions that are outdated by calling
 		 * IsActionOutdated.
 		 */
-		static Dictionary<Action,bool> GatherAllOutdatedActions(ActionHistory ActionHistory)
+		static void GatherAllOutdatedActions(UEBuildTarget Target, ActionHistory ActionHistory, ref Dictionary<Action,bool> OutdatedActions, Dictionary<UEBuildTarget, List<FileItem>> TargetToOutdatedPrerequisitesMap )
 		{
-			var OutdatedActionDictionary = new Dictionary<Action, bool>();
 			foreach (var Action in AllActions)
 			{
-				IsActionOutdated(Action, ref OutdatedActionDictionary,ActionHistory);
+				IsActionOutdated(Target, Action, ref OutdatedActions, ActionHistory, TargetToOutdatedPrerequisitesMap);
 			}
-
-			return OutdatedActionDictionary;
 		}
 
 		/**
@@ -856,33 +910,49 @@ namespace UnrealBuildTool
 
 	
 		/// <summary>
-		/// Checks if the specified file is a C++ source file
+		/// Checks if the specified file is a C++ source implementation file (e.g., .cpp)
 		/// </summary>
 		/// <param name="FileItem">The file to check</param>
 		/// <returns>True if this is a C++ source file</returns>
-		private static bool IsCPPFile( FileItem FileItem )
+		private static bool IsCPPImplementationFile( FileItem FileItem )
 		{
-			if( FileItem.AbsolutePath.EndsWith( ".h", StringComparison.InvariantCultureIgnoreCase ) ||
-				FileItem.AbsolutePath.EndsWith( ".inl", StringComparison.InvariantCultureIgnoreCase ) )
-			{
-				// Header file
-				return true;
-			}
-			else if( FileItem.AbsolutePath.EndsWith( ".cpp", StringComparison.InvariantCultureIgnoreCase ) ||
-						FileItem.AbsolutePath.EndsWith( ".c", StringComparison.InvariantCultureIgnoreCase ) ||
-						FileItem.AbsolutePath.EndsWith( ".mm", StringComparison.InvariantCultureIgnoreCase ) )
-			{
-				// C++ file
-				return true;
-			}
+			return( FileItem.AbsolutePath.EndsWith( ".cpp", StringComparison.InvariantCultureIgnoreCase ) ||
+					FileItem.AbsolutePath.EndsWith( ".c", StringComparison.InvariantCultureIgnoreCase ) ||
+					FileItem.AbsolutePath.EndsWith( ".mm", StringComparison.InvariantCultureIgnoreCase ) );
+		}
+
+
+		/// <summary>
+		/// Checks if the specified file is a C++ source header file (e.g., .h or .inl)
+		/// </summary>
+		/// <param name="FileItem">The file to check</param>
+		/// <returns>True if this is a C++ source file</returns>
+		private static bool IsCPPIncludeFile( FileItem FileItem )
+		{
+			return( FileItem.AbsolutePath.EndsWith( ".h", StringComparison.InvariantCultureIgnoreCase ) ||
+					FileItem.AbsolutePath.EndsWith( ".inl", StringComparison.InvariantCultureIgnoreCase ) );
+		}
+
+	
+		/// <summary>
+		/// Checks if the specified file is a C++ resource file (e.g., .rc)
+		/// </summary>
+		/// <param name="FileItem">The file to check</param>
+		/// <returns>True if this is a C++ source file</returns>
+		private static bool IsCPPResourceFile( FileItem FileItem )
+		{
+			return( FileItem.AbsolutePath.EndsWith( ".rc", StringComparison.InvariantCultureIgnoreCase ) );
+		}
 
-			else if( FileItem.AbsolutePath.EndsWith( ".rc", StringComparison.InvariantCultureIgnoreCase ) )
-			{
-				// .rc file
-				return true;
-			}
 
-			return false;
+		/// <summary>
+		/// Checks if the specified file is a C++ source file
+		/// </summary>
+		/// <param name="FileItem">The file to check</param>
+		/// <returns>True if this is a C++ source file</returns>
+		private static bool IsCPPFile( FileItem FileItem )
+		{
+			return IsCPPImplementationFile( FileItem ) || IsCPPIncludeFile( FileItem ) || IsCPPResourceFile( FileItem );
 		}
 
 
@@ -909,7 +979,7 @@ namespace UnrealBuildTool
 		/// <param name="VisualizationType">Type of graph to create</param>
 		/// <param name="Actions">All actions</param>
 		/// <param name="IncludeCompileActions">True if we should include compile actions.  If disabled, only the static link actions will be shown, which is useful to see module relationships</param>
-		public static void SaveActionGraphVisualization( string Filename, string Description, ActionGraphVisualizationType VisualizationType, List<Action> Actions, bool IncludeCompileActions = true )
+		public static void SaveActionGraphVisualization( UEBuildTarget Target, string Filename, string Description, ActionGraphVisualizationType VisualizationType, List<Action> Actions, bool IncludeCompileActions = true )
 		{
 			// True if we should include individual files in the graph network, or false to include only the build actions
 			var IncludeFiles = VisualizationType != ActionGraphVisualizationType.OnlyActions;
@@ -1024,6 +1094,7 @@ namespace UnrealBuildTool
 				}
 			}
 
+			var BuildPlatform = UEBuildPlatform.GetBuildPlatform( UnrealTargetPlatform.Win64 );
 
 			foreach( var FileItem in FilesToCreateNodesFor )
 			{
@@ -1076,7 +1147,8 @@ namespace UnrealBuildTool
 
 				if( ExpandCPPHeaderDependencies && bIsCPPFile )
 				{
-					List<DependencyInclude> DirectlyIncludedFilenames = CPPEnvironment.GetDirectIncludeDependencies( FileItem, CPPTargetPlatform.Win64 );
+					bool HasUObjects;
+					List<DependencyInclude> DirectlyIncludedFilenames = CPPEnvironment.GetDirectIncludeDependencies( Target, FileItem, BuildPlatform, bOnlyCachedDependencies:false, HasUObjects:out HasUObjects );
 
 					// Resolve the included file name to an actual file.
 					var DirectlyIncludedFiles =
diff --git a/Engine/Source/Programs/UnrealBuildTool/System/CPPEnvironment.cs b/Engine/Source/Programs/UnrealBuildTool/System/CPPEnvironment.cs
--- a/Engine/Source/Programs/UnrealBuildTool/System/CPPEnvironment.cs
+++ b/Engine/Source/Programs/UnrealBuildTool/System/CPPEnvironment.cs
@@ -346,9 +346,9 @@ namespace UnrealBuildTool
 		 * @param ModuleName - Name of the module these files are associated with
 		 * @return The object files produced by the actions.
 		 */
-		public CPPOutput CompileFiles(List<FileItem> CPPFiles, string ModuleName)
+		public CPPOutput CompileFiles(UEBuildTarget Target, List<FileItem> CPPFiles, string ModuleName)
 		{
-			return UEToolChain.GetPlatformToolChain(Config.Target.Platform).CompileCPPFiles(this, CPPFiles, ModuleName);
+			return UEToolChain.GetPlatformToolChain(Config.Target.Platform).CompileCPPFiles(Target, this, CPPFiles, ModuleName);
 		}
 
 		/**
@@ -356,9 +356,9 @@ namespace UnrealBuildTool
 		 * @param RCFiles - The resource script files to compile.
 		 * @return The compiled resource (.res) files produced by the actions.
 		 */
-		public CPPOutput CompileRCFiles(List<FileItem> RCFiles)
+		public CPPOutput CompileRCFiles(UEBuildTarget Target, List<FileItem> RCFiles)
 		{
-			return UEToolChain.GetPlatformToolChain(Config.Target.Platform).CompileRCFiles(this, RCFiles);
+			return UEToolChain.GetPlatformToolChain(Config.Target.Platform).CompileRCFiles(Target, this, RCFiles);
 		}
 
 		/**
@@ -386,5 +386,27 @@ namespace UnrealBuildTool
 		{
 			return new CPPEnvironment(this);
 		}
+
+
+		/// <summary>
+		/// Give a C++ source file, returns a list of include paths we should search to resolve #includes for this path
+		/// </summary>
+		/// <param name="SourceFile">C++ source file we're going to check #includes for.</param>
+		/// <returns>Ordered list of paths to search</returns>
+		public List<string> GetIncludesPathsToSearch( FileItem SourceFile )
+		{
+			// Build a single list of include paths to search.
+			var IncludePathsToSearch = new List<string>();
+			string SourceFilesDirectory = Path.GetDirectoryName( SourceFile.AbsolutePath);
+			IncludePathsToSearch.Add( SourceFilesDirectory );
+			IncludePathsToSearch.AddRange( Config.IncludePaths );
+			if( BuildConfiguration.bCheckSystemHeadersForModification )
+			{
+				IncludePathsToSearch.AddRange( Config.SystemIncludePaths );
+			}
+			
+			return IncludePathsToSearch;
+		}
+
 	};
 }
diff --git a/Engine/Source/Programs/UnrealBuildTool/System/CPPHeaders.cs b/Engine/Source/Programs/UnrealBuildTool/System/CPPHeaders.cs
--- a/Engine/Source/Programs/UnrealBuildTool/System/CPPHeaders.cs
+++ b/Engine/Source/Programs/UnrealBuildTool/System/CPPHeaders.cs
@@ -23,14 +23,14 @@ namespace UnrealBuildTool
 
 	public partial class CPPEnvironment
 	{
-		/** Contains a cache of include dependencies (direct and indirect). */
-		public static DependencyCache IncludeDependencyCache = null;
+		/** Contains a cache of include dependencies (direct and indirect), one for each target we're building. */
+		public static readonly Dictionary<UEBuildTarget, DependencyCache> IncludeDependencyCache = new Dictionary<UEBuildTarget, DependencyCache>();
 
-		/** Contains a mapping from filename to the full path of the header in this environment. */
-		public Dictionary<string, FileItem> IncludeFileSearchDictionary = new Dictionary<string, FileItem>();
+		/** Contains a cache of include dependencies (direct and indirect), one for each target we're building. */
+		public static readonly Dictionary<UEBuildTarget, FlatCPPIncludeDependencyCache> FlatCPPIncludeDependencyCache = new Dictionary<UEBuildTarget, FlatCPPIncludeDependencyCache>();
 
-		/** Absolute path to external folder, upper invariant. */
-		private static string AbsoluteExternalPathUpperInvariant = Path.GetFullPath( "." + Path.DirectorySeparatorChar + "ThirdParty" + Path.DirectorySeparatorChar ).ToUpperInvariant();
+		/** Contains a mapping from filename to the full path of the header in this environment. */
+		public readonly Dictionary<string, FileItem> IncludeFileSearchDictionary = new Dictionary<string, FileItem>();
 
 		public static int TotalFindIncludedFileCalls = 0;
 		public static int IncludePathSearchAttempts = 0;
@@ -42,10 +42,15 @@ namespace UnrealBuildTool
 		 * @param bSkipExternalHeader true to skip processing of headers in external path
 		 * @param SourceFilesDirectory - The folder containing the source files we're generating a PCH for
 		 */
-		public FileItem FindIncludedFile( string RelativeIncludePath, bool bSkipExternalHeader, String SourceFilesDirectory )
+		public static FileItem FindIncludedFile( string RelativeIncludePath, bool bSkipExternalHeader, List<string> IncludePathsToSearch, Dictionary<string, FileItem> IncludeFileSearchDictionary )
 		{
 			FileItem Result = null;
 
+			if( IncludePathsToSearch == null )
+			{
+				throw new BuildException( "Was not expecting IncludePathsToSearch to be empty for file '{0}'!", RelativeIncludePath );
+			}
+
 			++TotalFindIncludedFileCalls;
 
 			// Only search for the include file if the result hasn't been cached.
@@ -63,18 +68,6 @@ namespace UnrealBuildTool
 				}
 				else
 				{
-					// Build a single list of include paths to search.
-					var IncludePathsToSearch = new List<string>();
-					if( !String.IsNullOrEmpty( SourceFilesDirectory ) )
-					{
-						IncludePathsToSearch.Add( SourceFilesDirectory );
-					}
-					IncludePathsToSearch.AddRange( Config.IncludePaths );
-					if( BuildConfiguration.bCheckSystemHeadersForModification )
-					{
-						IncludePathsToSearch.AddRange( Config.SystemIncludePaths );
-					}
-
 					// Find the first include path that the included file exists in.
 					foreach( string IncludePath in IncludePathsToSearch )
 					{
@@ -114,29 +107,12 @@ namespace UnrealBuildTool
 				IncludeFileSearchDictionary.Add( InvariantPath, Result );
 			}
 
-			// Check whether the header should be skipped. We need to do this after resolving as we need 
-			// the absolute path to compare against the External folder.
-			bool bWasHeaderSkipped = false;
-			if( Result != null && bSkipExternalHeader )
-			{
-				// Check whether header path is under External root.
-				if( Result.AbsolutePath.StartsWith( AbsoluteExternalPathUpperInvariant, StringComparison.InvariantCultureIgnoreCase ) )
-				{
-					// It is, skip and reset result.
-					Result = null;
-					bWasHeaderSkipped = true;
-				}
-			}
+			// @todo fastubt: The old UBT tried to skip 'external' headers here.  But it didn't work.  We might want to do this though!  Skip system headers and source/thirdparty headers!
 
-			// Log status for header.
 			if( Result != null )
 			{
 				Log.TraceVerbose("Resolved included file \"{0}\" to: {1}", RelativeIncludePath, Result.AbsolutePath);
 			}
-			else if( bWasHeaderSkipped )
-			{
-				Log.TraceVerbose("Skipped included file \"{0}\"", RelativeIncludePath);
-			}
 			else
 			{
 				Log.TraceVerbose("Couldn't resolve included file \"{0}\"", RelativeIncludePath);
@@ -146,34 +122,83 @@ namespace UnrealBuildTool
 		}
 
 		/** A cache of the list of other files that are directly or indirectly included by a C++ file. */
-		static Dictionary<FileItem, IncludedFilesSet> IncludedFilesMap = new Dictionary<FileItem, IncludedFilesSet>();
+		static Dictionary<FileItem, IncludedFilesSet> ExhaustiveIncludedFilesMap = new Dictionary<FileItem, IncludedFilesSet>();
+
+		/** A cache of all files included by a C++ file, but only has files that we knew about from a previous session, loaded from a cache at startup */
+		static Dictionary<FileItem, IncludedFilesSet> OnlyCachedIncludedFilesMap = new Dictionary<FileItem, IncludedFilesSet>();
 
-		public static void Reset()
+
+		public static List<FileItem> FindAndCacheAllIncludedFiles( UEBuildTarget Target, FileItem SourceFile, IUEBuildPlatform BuildPlatform, List<string> IncludePathsToSearch, Dictionary<string, FileItem> IncludeFileSearchDictionary, bool bOnlyCachedDependencies )
 		{
-			IncludedFilesMap = new Dictionary<FileItem, IncludedFilesSet>();
+			var Result = new List<FileItem>();
+
+			if( BuildConfiguration.bUseExperimentalFastDependencyScan && bOnlyCachedDependencies )
+			{ 
+				var Dependencies = FlatCPPIncludeDependencyCache[Target].GetDependenciesForFile( SourceFile.AbsolutePath );
+				if( Dependencies != null )
+				{ 
+					foreach( string Dependency in Dependencies )
+					{
+						Result.Add( FileItem.GetItemByFullPath( Dependency ) );	// @todo fastubt: Make sure this is as fast as possible (convert to FileItem)
+					}
+				}
+				else
+				{
+					// Nothing cached for this file!  It is new to us.  This is the expected flow when our CPPIncludeDepencencyCache is missing.
+				}
+			}
+			else
+			{
+				// @todo fastubt: HeaderParser.h is missing from the include set for Module.UnrealHeaderTool.cpp (failed to find include using:  FileItem DirectIncludeResolvedFile = CPPEnvironment.FindIncludedFile(DirectInclude.IncludeName, !BuildConfiguration.bCheckExternalHeadersForModification, IncludePathsToSearch, IncludeFileSearchDictionary );)
+
+				var IncludedFileList = new IncludedFilesSet();
+				CPPEnvironment.FindAndCacheAllIncludedFiles( Target, SourceFile, BuildPlatform, IncludePathsToSearch, IncludeFileSearchDictionary, ref IncludedFileList, bOnlyCachedDependencies:bOnlyCachedDependencies );
+				foreach( FileItem IncludedFile in IncludedFileList )
+				{
+					Result.Add( IncludedFile );
+				}
+			
+				// Update cache
+				if( BuildConfiguration.bUseExperimentalFastDependencyScan && !bOnlyCachedDependencies )
+				{ 
+					var Dependencies = new List<string>();
+					foreach( var IncludedFile in Result )	// @todo fastubt: Have to convert Hash to List here, try to avoid this.  Other locations too.
+					{
+						Dependencies.Add( IncludedFile.AbsolutePath );
+					}
+					string PCHName = SourceFile.PrecompiledHeaderIncludeFilename;
+					FlatCPPIncludeDependencyCache[Target].SetDependenciesForFile( SourceFile.AbsolutePath, PCHName, Dependencies );
+				}
+			}
+
+			return Result;
 		}
 
+
 		/// <summary>
 		/// Finds the files directly or indirectly included by the given C++ file.
 		/// </summary>
 		/// <param name="CPPFile">C++ file to get the dependencies for.</param>
 		/// <param name="Result">List of CPPFile dependencies.</param>
 		/// <returns>false if CPPFile is still being processed further down the callstack, true otherwise.</returns>
-		private bool GetIncludeDependencies(FileItem CPPFile, ref IncludedFilesSet Result)
+		public static bool FindAndCacheAllIncludedFiles(UEBuildTarget Target, FileItem CPPFile, IUEBuildPlatform BuildPlatform, List<string> IncludePathsToSearch, Dictionary<string, FileItem> IncludeFileSearchDictionary, ref IncludedFilesSet Result, bool bOnlyCachedDependencies)
 		{
 			IncludedFilesSet IncludedFileList;
+			var IncludedFilesMap = bOnlyCachedDependencies ? OnlyCachedIncludedFilesMap : ExhaustiveIncludedFilesMap;
 			if( !IncludedFilesMap.TryGetValue( CPPFile, out IncludedFileList ) )
 			{
+				var TimerStartTime = DateTime.UtcNow;
+
 				IncludedFileList = new IncludedFilesSet();
 
 				// Add an uninitialized entry for the include file to avoid infinitely recursing on include file loops.
 				IncludedFilesMap.Add(CPPFile, IncludedFileList);
 
 				// Gather a list of names of files directly included by this C++ file.
-				List<DependencyInclude> DirectIncludes = GetDirectIncludeDependencies(CPPFile, Config.Target.Platform);
+				bool HasUObjects;
+				List<DependencyInclude> DirectIncludes = GetDirectIncludeDependencies( Target, CPPFile, BuildPlatform, bOnlyCachedDependencies:bOnlyCachedDependencies, HasUObjects:out HasUObjects );
 
 				// Build a list of the unique set of files that are included by this file.
-				string SourceFilesDirectory = Path.GetDirectoryName( CPPFile.AbsolutePath );
 				var DirectlyIncludedFiles = new HashSet<FileItem>();
 				// require a for loop here because we need to keep track of the index in the list.
 				for (int DirectlyIncludedFileNameIndex = 0; DirectlyIncludedFileNameIndex < DirectIncludes.Count; ++DirectlyIncludedFileNameIndex)
@@ -190,12 +215,12 @@ namespace UnrealBuildTool
 						++TotalDirectIncludeResolveCacheMisses;
 
 						// search the include paths to resolve the file
-						FileItem DirectIncludeResolvedFile = FindIncludedFile(DirectInclude.IncludeName, !BuildConfiguration.bCheckExternalHeadersForModification, SourceFilesDirectory);
+						FileItem DirectIncludeResolvedFile = CPPEnvironment.FindIncludedFile(DirectInclude.IncludeName, !BuildConfiguration.bCheckExternalHeadersForModification, IncludePathsToSearch, IncludeFileSearchDictionary );
 						if (DirectIncludeResolvedFile != null)
 						{
 							DirectlyIncludedFiles.Add(DirectIncludeResolvedFile);
 						}
-						IncludeDependencyCache.CacheResolvedIncludeFullPath(CPPFile, DirectlyIncludedFileNameIndex, DirectIncludeResolvedFile != null ? DirectIncludeResolvedFile.AbsolutePath : "");
+						IncludeDependencyCache[Target].CacheResolvedIncludeFullPath(CPPFile, DirectlyIncludedFileNameIndex, DirectIncludeResolvedFile != null ? DirectIncludeResolvedFile.AbsolutePath : "");
 					}
 					else
 					{
@@ -215,7 +240,7 @@ namespace UnrealBuildTool
 					IncludedFileList.Add( DirectlyIncludedFile );
 
 					// Also add all of the indirectly included files!
-					if (GetIncludeDependencies(DirectlyIncludedFile, ref IncludedFileList) == false)
+					if (FindAndCacheAllIncludedFiles(Target, DirectlyIncludedFile, BuildPlatform, IncludePathsToSearch, IncludeFileSearchDictionary, ref IncludedFileList, bOnlyCachedDependencies:bOnlyCachedDependencies) == false)
 					{
 						// DirectlyIncludedFile is a circular dependency which is still being processed
 						// further down the callstack. Add this file to its circular dependencies list 
@@ -243,6 +268,9 @@ namespace UnrealBuildTool
 
 				// Done collecting files.
 				IncludedFileList.bIsInitialized = true;
+
+				var TimerDuration = DateTime.UtcNow - TimerStartTime;
+				TotalTimeSpentGettingIncludes += TimerDuration.TotalSeconds;
 			}
 
 			if (IncludedFileList.bIsInitialized)
@@ -254,6 +282,7 @@ namespace UnrealBuildTool
 					// NOTE: For some reason in .NET 4, Add() is over twice as fast as calling UnionWith() on the set
 					Result.Add( IncludedFile );
 				}
+
 				return true;
 			}
 			else
@@ -271,58 +300,35 @@ namespace UnrealBuildTool
 		public static int TotalDirectIncludeResolveCacheMisses = 0;
 		public static int TotalDirectIncludeResolves = 0;
 
-		/** @return The list of files which are directly or indirectly included by a C++ file. */
-		public List<FileItem> GetIncludeDependencies( FileItem CPPFile )
-		{
-			++TotalIncludesRequested;
-
-			List<FileItem> Result = new List<FileItem>();
-
-			// Don't bother gathering includes when only generating IntelliSense data.  We don't need to actually
-			// know about source file dependencies to create the data we need.
-			if( !ProjectFileGenerator.bGenerateProjectFiles )
-			{
-				var TimerStartTime = DateTime.UtcNow;
-
-				// Find the dependencies of the file.
-				var IncludedFileDictionary = new IncludedFilesSet();
-				GetIncludeDependencies( CPPFile, ref IncludedFileDictionary );
-
-				// Convert the dependency dictionary into a list.
-				foreach( var IncludedFile in IncludedFileDictionary )
-				{
-					Result.Add( IncludedFile );
-				}
-
-				var TimerDuration = DateTime.UtcNow - TimerStartTime;
-				TotalTimeSpentGettingIncludes += TimerDuration.TotalSeconds;
-			}
-
-			return Result;
-		}
 
 		/** Regex that matches #include statements. */
-		static Regex CPPHeaderRegex = new Regex("(([ \t]*#[ \t]*include[ \t]*[<\"](?<HeaderFile>[^\">]*)[\">][^\n]*\n*)|([^\n]*\n*))*",
+		static readonly Regex CPPHeaderRegex = new Regex("(([ \t]*#[ \t]*include[ \t]*[<\"](?<HeaderFile>[^\">]*)[\">][^\n]*\n*)|([^\n]*\n*))*",
 													RegexOptions.Compiled | RegexOptions.Singleline | RegexOptions.ExplicitCapture );
 
-		static Regex MMHeaderRegex = new Regex("(([ \t]*#[ \t]*import[ \t]*[<\"](?<HeaderFile>[^\">]*)[\">][^\n]*\n*)|([^\n]*\n*))*",
+		static readonly Regex MMHeaderRegex = new Regex("(([ \t]*#[ \t]*import[ \t]*[<\"](?<HeaderFile>[^\">]*)[\">][^\n]*\n*)|([^\n]*\n*))*",
 													RegexOptions.Compiled | RegexOptions.Singleline | RegexOptions.ExplicitCapture );
 
+		/** Regex that matches C++ code with UObject declarations which we will need to generated code for. */
+		static readonly Regex UObjectRegex = new Regex("^\\s*U(CLASS|STRUCT|ENUM|INTERFACE|DELEGATE)\\b", RegexOptions.Compiled | RegexOptions.Multiline );
+
 		/** Finds the names of files directly included by the given C++ file. */
-		public static List<DependencyInclude> GetDirectIncludeDependencies( FileItem CPPFile, CPPTargetPlatform TargetPlatform )
+		public static List<DependencyInclude> GetDirectIncludeDependencies( UEBuildTarget Target, FileItem CPPFile, IUEBuildPlatform BuildPlatform, bool bOnlyCachedDependencies, out bool HasUObjects )
 		{
 			// Try to fulfill request from cache first.
-			List<DependencyInclude> Result = null;
-			if( IncludeDependencyCache.GetDirectDependencies( CPPFile, out Result ) )
+			List<DependencyInclude> Result;
+			if( IncludeDependencyCache[Target].GetCachedDirectDependencies( CPPFile, out Result, out HasUObjects ) )
 			{
 				return Result;
 			}
+			else if( bOnlyCachedDependencies )
+			{
+				return new List<DependencyInclude>();
+			}
 
 			var TimerStartTime = DateTime.UtcNow;
 			++CPPEnvironment.TotalDirectIncludeCacheMisses;
 
 			// Get the adjusted filename
-			var BuildPlatform = UEBuildPlatform.GetBuildPlatformForCPPTargetPlatform( TargetPlatform );
 			string FileToRead = CPPFile.AbsolutePath;
 
 			if (BuildPlatform.RequiresExtraUnityCPPWriter() == true &&
@@ -338,6 +344,8 @@ namespace UnrealBuildTool
                 return new List<DependencyInclude>();
             }
 
+			HasUObjects = CPPEnvironment.UObjectRegex.IsMatch(FileContents);
+
 			// Note: This depends on UBT executing w/ a working directory of the Engine/Source folder!
 			string EngineSourceFolder = Directory.GetCurrentDirectory();
 			string InstalledFolder = EngineSourceFolder;
@@ -372,7 +380,7 @@ namespace UnrealBuildTool
 			}
 
 			// Populate cache with results.
-			IncludeDependencyCache.SetDependencyInfo( CPPFile, Result );
+			IncludeDependencyCache[Target].SetDependencyInfo( CPPFile, Result, HasUObjects );
 
 			CPPEnvironment.DirectIncludeCacheMissesTotalTime += ( DateTime.UtcNow - TimerStartTime ).TotalSeconds;
 
diff --git a/Engine/Source/Programs/UnrealBuildTool/System/DependencyCache.cs b/Engine/Source/Programs/UnrealBuildTool/System/DependencyCache.cs
--- a/Engine/Source/Programs/UnrealBuildTool/System/DependencyCache.cs
+++ b/Engine/Source/Programs/UnrealBuildTool/System/DependencyCache.cs
@@ -42,6 +42,10 @@ namespace UnrealBuildTool
 		/** List of files this file includes.  These are direct include paths and cannot be resolved to an actual file on 
 		    disk without using the proper list of include directories for this file's module */
 		public List<DependencyInclude> Includes;
+
+		/** True if this file was found to contain UObject class or type definitions, and the file needs to be preprocessed
+		    by UnrealHeaderTool to generate reflection code */
+		public bool HasUObjects;
 	}
 
 
@@ -201,6 +205,35 @@ namespace UnrealBuildTool
 			}
 		}
 
+
+		/// <summary>
+		/// Gets information about this file from our dependency cache.  Only returns information if the file has already been 
+		/// cached, and the file has not been changed since the last time we updated our cache
+		/// </summary>
+		/// <param name="File">The file to check</param>
+		/// <returns>Information about this dependency, or null if we have nothing cached</returns>
+		DependencyInfo? GetCachedDependencyInfo( FileItem File )
+		{
+			// Check whether File is in cache.
+			DependencyInfo DependencyInfo;
+			if (DependencyMap.TryGetValue(File.AbsolutePath, out DependencyInfo))
+			{
+				// File is in cache, now check whether last write time is prior to cache creation time.
+				if (File.LastWriteTime < CacheCreateDate)
+				{
+					return DependencyInfo;					
+				}
+				else
+				{
+					// Remove entry from cache as it's stale.
+					RemoveFileFromCache(File.AbsolutePath);
+					return null;
+				}
+			}
+
+			return null;
+		}
+
 		/**
 		 * Returns the direct dependencies of the specified FileItem if it exists in the cache and if the
 		 * file has a last write time before the creation time of the cache. 
@@ -210,50 +243,42 @@ namespace UnrealBuildTool
 		 * 
 		 * @param	File				File to try to find dependencies in cache
 		 * @param	Result	[out]		List of dependencies if successful, null otherwise
+		 * @param	HasUObjects			True if the file was found to have UObject classes, otherwise false
 		 */
-		public bool GetDirectDependencies(FileItem File, out List<DependencyInclude> Result)
+		public bool GetCachedDirectDependencies(FileItem File, out List<DependencyInclude> Result, out bool HasUObjects)
 		{
 			Result = null;
+			HasUObjects = false;
 
 			// Check whether File is in cache.
-			DependencyInfo DependencyInfo;
-			if (DependencyMap.TryGetValue(File.AbsolutePath, out DependencyInfo))
+			DependencyInfo? DependencyInfo = GetCachedDependencyInfo(File);
+			if( DependencyInfo != null )
 			{
-				// File is in cache, now check whether last write time is prior to cache creation time.
-				if (File.LastWriteTime < CacheCreateDate)
+				// Check if any of the resolved includes is missing
+				foreach (var Include in DependencyInfo.Value.Includes)
 				{
-					// Check if any of the resolved includes is missing
-					foreach (var Include in DependencyInfo.Includes)
+					if (!String.IsNullOrEmpty(Include.IncludeResolvedName))
 					{
-						if (!String.IsNullOrEmpty(Include.IncludeResolvedName))
+						bool bIncludeExists = false;
+						string FileExistsKey = Include.IncludeResolvedName.ToLowerInvariant();
+						if (FileExistsInfo.TryGetValue(FileExistsKey, out bIncludeExists) == false)
 						{
-							bool bIncludeExists = false;
-							string FileExistsKey = Include.IncludeResolvedName.ToLowerInvariant();
-							if (FileExistsInfo.TryGetValue(FileExistsKey, out bIncludeExists) == false)
-							{
-								bIncludeExists = System.IO.File.Exists(Include.IncludeResolvedName);
-								FileExistsInfo.Add(FileExistsKey, bIncludeExists);
-							}
-							if (!bIncludeExists)
-							{
-								// Remove entry from cache as it's stale, as well as the include which no longer exists
-								RemoveFileFromCache(Include.IncludeResolvedName);
-								RemoveFileFromCache(File.AbsolutePath);							
-								return false;
-							}
-						}						
-					}
-					// Cached version is up to date, return it.
-					Result = DependencyInfo.Includes;
-					return true;
-				}
-				// File has been modified since the last time it was cached.
-				else
-				{
-					// Remove entry from cache as it's stale.
-					RemoveFileFromCache(File.AbsolutePath);
-					return false;
+							bIncludeExists = System.IO.File.Exists(Include.IncludeResolvedName);
+							FileExistsInfo.Add(FileExistsKey, bIncludeExists);
+						}
+						if (!bIncludeExists)
+						{
+							// Remove entry from cache as it's stale, as well as the include which no longer exists
+							RemoveFileFromCache(Include.IncludeResolvedName);
+							RemoveFileFromCache(File.AbsolutePath);							
+							return false;
+						}
+					}						
 				}
+				// Cached version is up to date, return it.
+				Result = DependencyInfo.Value.Includes;
+				HasUObjects = DependencyInfo.Value.HasUObjects;
+				return true;
 			}
 			// Not in cache.
 			else
@@ -277,10 +302,11 @@ namespace UnrealBuildTool
 		 * 
 		 * @param	File			File to update dependencies for
 		 * @param	Dependencies	List of dependencies to cache for passed in file
+		 * @param	HasUObjects		True if this file was found to contain UObject classes or types
 		 */
-		public void SetDependencyInfo(FileItem File, List<DependencyInclude> DirectDependencies)
+		public void SetDependencyInfo(FileItem File, List<DependencyInclude> DirectDependencies, bool HasUObjects)
 		{
-			DependencyMap[File.AbsolutePath] = new DependencyInfo() { Includes = DirectDependencies };
+			DependencyMap[File.AbsolutePath] = new DependencyInfo() { Includes = DirectDependencies, HasUObjects = HasUObjects };
 			bIsDirty = true;
 		}
 
@@ -352,7 +378,7 @@ namespace UnrealBuildTool
 			string PlatformIntermediatePath = BuildConfiguration.PlatformIntermediatePath;
 			if (UnrealBuildTool.HasUProjectFile())
 			{
-				PlatformIntermediatePath = Path.Combine(UnrealBuildTool.GetUProjectPath(), BuildConfiguration.BaseIntermediateFolder);
+				PlatformIntermediatePath = Path.Combine(UnrealBuildTool.GetUProjectPath(), BuildConfiguration.PlatformIntermediateFolder);
 			}
 			string CachePath = Path.Combine(PlatformIntermediatePath, Target.GetTargetName(), "DependencyCache.bin");
 			return CachePath;
diff --git a/Engine/Source/Programs/UnrealBuildTool/System/ExternalExecution.cs b/Engine/Source/Programs/UnrealBuildTool/System/ExternalExecution.cs
--- a/Engine/Source/Programs/UnrealBuildTool/System/ExternalExecution.cs
+++ b/Engine/Source/Programs/UnrealBuildTool/System/ExternalExecution.cs
@@ -476,7 +476,7 @@ namespace UnrealBuildTool
 						// We need to figure out which PCH header this module is including, so that UHT can inject an include statement for it into any .cpp files it is synthesizing
 						var DependencyModuleCPP = (UEBuildModuleCPP)Target.GetModuleByName( UHTModuleInfo.ModuleName );
 						var ModuleCompileEnvironment = DependencyModuleCPP.CreateModuleCompileEnvironment(GlobalCompileEnvironment);
-						DependencyModuleCPP.ProcessAllCppDependencies(ModuleCompileEnvironment);
+						DependencyModuleCPP.CachePCHUsageForModuleSourceFiles(ModuleCompileEnvironment);
 						if (DependencyModuleCPP.ProcessedDependencies.UniquePCHHeaderFile != null)
 						{
 							UHTModuleInfo.PCH = DependencyModuleCPP.ProcessedDependencies.UniquePCHHeaderFile.AbsolutePath;
diff --git a/Engine/Source/Programs/UnrealBuildTool/System/FileItem.cs b/Engine/Source/Programs/UnrealBuildTool/System/FileItem.cs
--- a/Engine/Source/Programs/UnrealBuildTool/System/FileItem.cs
+++ b/Engine/Source/Programs/UnrealBuildTool/System/FileItem.cs
@@ -56,6 +56,60 @@ namespace UnrealBuildTool
 		/** A list of remote file items that have been created but haven't needed the remote info yet, so we can gang up many into one request */
 		static List<FileItem> DelayedRemoteLookupFiles = new List<FileItem>();
 
+		/** The compile environment for this file.  This is only set on C++ source file file items, or C++-related files that have include paths.  It's part of UBT's dependency caching optimizations. */
+		public CPPEnvironment CachedCPPEnvironment
+		{
+			get
+			{
+				return _CachedCPPEnvironment;
+			}
+			set
+			{
+				if( value != null && _CachedCPPEnvironment != null && _CachedCPPEnvironment != value )
+				{
+					// Uh oh.  We're clobbering our cached CompileEnvironment for this file with a different CompileEnvironment.  This means
+					// that the same source file is being compiled into more than one module. (e.g. PCLaunch.rc)
+
+					// @todo fastubt: The only expected offender here is PCLaunch.rc and friends, which are injected by UBT into every module when not compiling monolithic.
+					// PCLaunch.rc and ModuleVersionResource.rc.inl are "safe" because they do not include any headers that would be affected by include path order.
+					// ==> Ideally we would use a different "shared" CompileEnvironment for these injected .rc files, so their include paths would not change
+					// ==> OR, we can make an Intermediate copy of the .rc file for each module (easier)
+					if( !AbsolutePath.EndsWith( "PCLaunch.rc", StringComparison.InvariantCultureIgnoreCase ) &&
+						!AbsolutePath.EndsWith( "ModuleVersionResource.rc.inl", StringComparison.InvariantCultureIgnoreCase ) )
+					{ 					
+						// Let's make sure the include paths are the same
+						// @todo fastubt: We have not seen examples of this actually firing off, so we could probably remove the check for matching includes and simply always make this an error case
+						var CachedIncludePathsToSearch = _CachedCPPEnvironment.GetIncludesPathsToSearch( this );
+						var NewIncludePathsToSearch = value.GetIncludesPathsToSearch( this );
+
+						bool bIncludesAreDifferent = false;
+						if( CachedIncludePathsToSearch.Count != NewIncludePathsToSearch.Count )
+						{
+							bIncludesAreDifferent = true;
+						}
+						else
+						{
+							for( var IncludeIndex = 0; IncludeIndex < CachedIncludePathsToSearch.Count; ++IncludeIndex )
+							{
+								if( !CachedIncludePathsToSearch[ IncludeIndex ].Equals( NewIncludePathsToSearch[ IncludeIndex ], StringComparison.InvariantCultureIgnoreCase ) )
+								{
+									bIncludesAreDifferent = true;
+									break;
+								}
+							}
+						}
+
+						if( bIncludesAreDifferent )
+						{ 
+							throw new BuildException( "File '{0}' was included by multiple modules, but with different include paths", this.Info.FullName );
+						}
+					}
+				}
+				_CachedCPPEnvironment = value;
+			}
+		}
+		public CPPEnvironment _CachedCPPEnvironment;
+
 		/** The last write time of the file. */
 		public DateTimeOffset _LastWriteTime;
 		public DateTimeOffset LastWriteTime
diff --git a/Engine/Source/Programs/UnrealBuildTool/System/FlatCPPIncludeDepencencyCache.cs b/Engine/Source/Programs/UnrealBuildTool/System/FlatCPPIncludeDepencencyCache.cs
new file mode 100644
--- /dev/null
+++ b/Engine/Source/Programs/UnrealBuildTool/System/FlatCPPIncludeDepencencyCache.cs
@@ -0,0 +1,217 @@
+// Copyright 1998-2014 Epic Games, Inc. All Rights Reserved.
+
+using System;
+using System.Collections.Generic;
+using System.Diagnostics;
+using System.IO;
+using System.Linq;
+using System.Runtime.Serialization.Formatters.Binary;
+
+namespace UnrealBuildTool
+{
+	/// <summary>
+	/// Stores an ordered list of header files
+	/// </summary>
+	[Serializable]
+	struct FlatCPPIncludeDependencyInfo
+	{
+		/// The PCH header this file is dependent on.
+		public string PCHName;
+
+		/// List of files this file includes, excluding headers that were included from a PCH.
+		public List<string> Includes;
+	}
+
+
+	/// <summary>
+	/// For a given target, caches all of the C++ source files and the files they are including
+	/// </summary>
+	[Serializable]
+	public class FlatCPPIncludeDependencyCache
+	{
+		/// <summary>
+		/// Creates the cache object
+		/// </summary>
+		/// <param name="Target">The target to create the cache for</param>
+		/// <returns>The new instance</returns>
+		public static FlatCPPIncludeDependencyCache Create( UEBuildTarget Target )
+		{
+			string CachePath = FlatCPPIncludeDependencyCache.GetDependencyCachePathForTarget( Target );;
+
+			// See whether the cache file exists.
+			FileItem Cache = FileItem.GetItemByPath(CachePath);
+			if (Cache.bExists)
+			{
+				if (BuildConfiguration.bPrintPerformanceInfo)
+				{
+					Log.TraceInformation("Loading existing FlatCPPIncludeDependencyCache: " + Cache.AbsolutePath);
+				}
+
+				var TimerStartTime = DateTime.UtcNow;
+
+				// Deserialize cache from disk if there is one.
+				FlatCPPIncludeDependencyCache Result = Load(Cache);
+				if (Result != null)
+				{
+					var TimerDuration = DateTime.UtcNow - TimerStartTime;
+					if (BuildConfiguration.bPrintPerformanceInfo)
+					{
+						Log.TraceInformation("Loading FlatCPPIncludeDependencyCache took " + TimerDuration.TotalSeconds + "s");
+					}
+					return Result;
+				}
+			}
+
+			// Fall back to a clean cache on error or non-existence.
+			return new FlatCPPIncludeDependencyCache( Cache );
+		}
+
+		/// <summary>
+		/// Loads the cache from disk
+		/// </summary>
+		/// <param name="Cache">The file to load</param>
+		/// <returns>The loaded instance</returns>
+		public static FlatCPPIncludeDependencyCache Load(FileItem Cache)
+		{
+			FlatCPPIncludeDependencyCache Result = null;
+			try
+			{
+				using (FileStream Stream = new FileStream(Cache.AbsolutePath, FileMode.Open, FileAccess.Read))
+				{	
+					// @todo fastubt: We can store the cache in a cheaper/smaller way using hash file names and indicies into included headers, but it might actually slow down load times
+					// @todo fastubt: If we can index PCHs here, we can avoid storing all of the PCH's included headers (PCH's action should have been invalidated, so we shouldn't even have to report the PCH's includes as our indirect includes)
+					BinaryFormatter Formatter = new BinaryFormatter();
+					Result = Formatter.Deserialize(Stream) as FlatCPPIncludeDependencyCache;
+					Result.CacheFileItem = Cache;
+					Result.bIsDirty = false;
+				}
+			}
+			catch (Exception Ex)
+			{
+				// Don't bother failing if the file format has changed, simply abort the cache load
+				if (Ex.Message.Contains( "cannot be converted to type" ))	// To catch serialization differences added when we added the DependencyInfo struct
+				{
+					Console.Error.WriteLine("Failed to read FlatCPPIncludeDependencyCache: {0}", Ex.Message);
+				}
+			}
+			return Result;
+		}
+
+
+		/// <summary>
+		/// Constructs a fresh cache, storing the file name that it should be saved as later on
+		/// </summary>
+		/// <param name="Cache">File name for this cache, usually unique per context (e.g. target)</param>
+		protected FlatCPPIncludeDependencyCache(FileItem Cache)
+		{
+			CacheFileItem = Cache;
+			DependencyMap = new Dictionary<string, FlatCPPIncludeDependencyInfo>(StringComparer.InvariantCultureIgnoreCase);
+		}
+
+
+		/// <summary>
+		/// Saves out the cache
+		/// </summary>
+		public void Save()
+		{
+			// Only save if we've made changes to it since load.
+			if (bIsDirty)
+			{
+				var TimerStartTime = DateTime.UtcNow;
+
+				// Serialize the cache to disk.
+				try
+				{
+					Directory.CreateDirectory(Path.GetDirectoryName(CacheFileItem.AbsolutePath));
+					using (FileStream Stream = new FileStream(CacheFileItem.AbsolutePath, FileMode.Create, FileAccess.Write))
+					{
+						BinaryFormatter Formatter = new BinaryFormatter();
+						Formatter.Serialize(Stream, this);
+					}
+				}
+				catch (Exception Ex)
+				{
+					Console.Error.WriteLine("Failed to write FlatCPPIncludeDependencyCache: {0}", Ex.Message);
+				}
+
+				if (BuildConfiguration.bPrintPerformanceInfo)
+				{
+					var TimerDuration = DateTime.UtcNow - TimerStartTime;
+					Log.TraceInformation("Saving FlatCPPIncludeDependencyCache took " + TimerDuration.TotalSeconds + "s");
+				}
+			}
+			else
+			{
+				if (BuildConfiguration.bPrintPerformanceInfo)
+				{
+					Log.TraceInformation("FlatCPPIncludeDependencyCache did not need to be saved (bIsDirty=false)");
+				}
+			}
+		}
+
+
+		/// <summary>
+		/// Sets the new dependencies for the specified file
+		/// </summary>
+		/// <param name="AbsoluteFilePath">File to set</param>
+		/// <param name="PCHName">The PCH for this file</param>
+		/// <param name="Dependencies">List of source dependencies</param>
+		public void SetDependenciesForFile( string AbsoluteFilePath, string PCHName, List<string> Dependencies )
+		{
+			FlatCPPIncludeDependencyInfo DependencyInfo;
+			DependencyInfo.PCHName = PCHName;	// @todo fastubt: Not actually used yet.  The idea is to use this to reduce the number of indirect includes we need to store in the cache.
+			DependencyInfo.Includes = Dependencies;
+
+			// @todo fastubt: We could shrink this file by not storing absolute paths (project and engine relative paths only, except for system headers.)  May affect load times.
+			DependencyMap[ AbsoluteFilePath ] = DependencyInfo;
+			bIsDirty = true;
+		}
+
+
+		/// <summary>
+		/// Gets everything that this file includes from our cache (direct and indirect!)
+		/// </summary>
+		/// <param name="AbsoluteFilePath">Path to the file</param>
+		/// <returns>The list of includes</returns>
+		public List<string> GetDependenciesForFile( string AbsoluteFilePath )
+		{
+			FlatCPPIncludeDependencyInfo DependencyInfo;
+			if( DependencyMap.TryGetValue( AbsoluteFilePath, out DependencyInfo ) )
+			{
+				return DependencyInfo.Includes;
+			}
+
+			return null;
+		}
+
+		/// <summary>
+		/// Gets the dependency cache path and filename for the specified target.
+		/// </summary>
+		/// <param name="Target">Current build target</param>
+		/// <returns>Cache Path</returns>
+		public static string GetDependencyCachePathForTarget(UEBuildTarget Target)
+		{
+			string PlatformIntermediatePath = BuildConfiguration.PlatformIntermediatePath;
+			if (UnrealBuildTool.HasUProjectFile())
+			{
+				PlatformIntermediatePath = Path.Combine(UnrealBuildTool.GetUProjectPath(), BuildConfiguration.PlatformIntermediateFolder);
+			}
+			string CachePath = Path.Combine(PlatformIntermediatePath, Target.GetTargetName(), "FlatCPPIncludes.bin" );
+			return CachePath;
+		}
+
+
+		/// File name of this cache, should be unique for every includes context (e.g. target)
+		[NonSerialized]
+		private FileItem CacheFileItem;
+
+		/// True if the cache needs to be saved
+		[NonSerialized]
+		private bool bIsDirty = false;
+
+		/// Dependency lists, keyed (case-insensitively) on file's absolute path.
+		private Dictionary<string, FlatCPPIncludeDependencyInfo> DependencyMap;
+
+
+	}
+}
diff --git a/Engine/Source/Programs/UnrealBuildTool/System/ResponseFile.cs b/Engine/Source/Programs/UnrealBuildTool/System/ResponseFile.cs
--- a/Engine/Source/Programs/UnrealBuildTool/System/ResponseFile.cs
+++ b/Engine/Source/Programs/UnrealBuildTool/System/ResponseFile.cs
@@ -16,6 +16,7 @@ namespace UnrealBuildTool
 		/// <param name="Lines">List of lines to write to the response file</param>
 		public static string Create(string TempFileName, List<string> Lines)
 		{
+			// @todo fastubt: Make sure we aren't spitting out response files more often than we have to.  Also look at Unity.cpp files, injected .cpp files, etc.
 			FileInfo TempFileInfo = new FileInfo( TempFileName );
 			// Delete the existing file if it exists
 			if( TempFileInfo.Exists )
diff --git a/Engine/Source/Programs/UnrealBuildTool/System/Unity.cs b/Engine/Source/Programs/UnrealBuildTool/System/Unity.cs
--- a/Engine/Source/Programs/UnrealBuildTool/System/Unity.cs
+++ b/Engine/Source/Programs/UnrealBuildTool/System/Unity.cs
@@ -102,11 +102,13 @@ namespace UnrealBuildTool
 		/// files, the goal being to compile the same code in fewer translation units.
 		/// The "unity" files are written to the CompileEnvironment's OutputDirectory.
 		/// </summary>
+		/// <param name="Target">The target we're building</param>
 		/// <param name="CPPFiles">The C++ files to #include.</param>
 		/// <param name="CompileEnvironment">The environment that is used to compile the C++ files.</param>
 		/// <param name="BaseName">Base name to use for the Unity files</param>
 		/// <returns>The "unity" C++ files.</returns>
 		public static List<FileItem> GenerateUnityCPPs(
+			UEBuildTarget Target,
 			List<FileItem> CPPFiles, 
 			CPPEnvironment CompileEnvironment,
 			string BaseName
@@ -142,6 +144,12 @@ namespace UnrealBuildTool
 					{
 						CPPUnityFileBuilder.AddFile(CPPFile);
 					}
+
+					// Now that the CPPFile is part of this unity file, we will no longer need to treat it like a root level prerequisite for our
+					// dependency cache, as it is now an "indirect include" from the unity file.  We'll clear out the compile environment
+					// attached to this file.  This prevents us from having to cache all of the indirect includes from these files inside our
+					// dependency cache, which speeds up iterative builds a lot!
+					CPPFile.CachedCPPEnvironment = null;
 				}
 				AllUnityFiles = CPPUnityFileBuilder.GetUnityFiles();
 			}
@@ -220,6 +228,9 @@ namespace UnrealBuildTool
 				UnityCPPFile.Description         = string.Join(" + ", UnityFile.Files.Select(F => Path.GetFileName(F.AbsolutePath)));
 				UnityCPPFile.PCHHeaderNameInCode = PCHHeaderNameInCode;
                 UnityCPPFiles.Add(UnityCPPFile);
+
+				// Cache information about the unity .cpp dependencies
+				UEBuildModuleCPP.CachePCHUsageForModuleSourceFile( Target, CompileEnvironment, UnityCPPFile );
 			}
 
 			return UnityCPPFiles;
diff --git a/Engine/Source/Programs/UnrealBuildTool/System/UnrealBuildTool.cs b/Engine/Source/Programs/UnrealBuildTool/System/UnrealBuildTool.cs
--- a/Engine/Source/Programs/UnrealBuildTool/System/UnrealBuildTool.cs
+++ b/Engine/Source/Programs/UnrealBuildTool/System/UnrealBuildTool.cs
@@ -39,6 +39,9 @@ namespace UnrealBuildTool
         /** Total time spent on other actions. */
         static public double TotalOtherActionsTime = 0;
 
+		/// How much time was spent scanning for include dependencies for outdated C++ files
+		public static double TotalDeepIncludeScanTime = 0.0;
+
         /** The command line */
         static public List<string> CmdLine = new List<string>();
 
@@ -975,7 +978,9 @@ namespace UnrealBuildTool
                         Log.TraceInformation("GetIncludes time: " + CPPEnvironment.TotalTimeSpentGettingIncludes + "s (" + CPPEnvironment.TotalIncludesRequested + " includes)" );
                         Log.TraceInformation("DirectIncludes cache miss time: " + CPPEnvironment.DirectIncludeCacheMissesTotalTime + "s (" + CPPEnvironment.TotalDirectIncludeCacheMisses + " misses)" );
                         Log.TraceInformation("FindIncludePaths calls: " + CPPEnvironment.TotalFindIncludedFileCalls + " (" + CPPEnvironment.IncludePathSearchAttempts + " searches)" );
-                        Log.TraceInformation("PCH time: " + UEBuildModuleCPP.TotalPCHTime + "s (Shared PCH time: " + UEBuildModuleCPP.TotalSharedPCHTime + "s)" );
+                        Log.TraceInformation("PCH gen time: " + UEBuildModuleCPP.TotalPCHGenTime + "s" );
+                        Log.TraceInformation("PCH cache time: " + UEBuildModuleCPP.TotalPCHCacheTime + "s" );
+                        Log.TraceInformation("Deep C++ include scan time: " + UnrealBuildTool.TotalDeepIncludeScanTime + "s" );
                         Log.TraceInformation("Include Resolves: {0} ({1} misses, {2:0.00}%)", CPPEnvironment.TotalDirectIncludeResolves, CPPEnvironment.TotalDirectIncludeResolveCacheMisses, (float)CPPEnvironment.TotalDirectIncludeResolveCacheMisses / (float)CPPEnvironment.TotalDirectIncludeResolves * 100);
                         Log.TraceInformation("Total FileItems: {0} ({1} missing)", FileItem.TotalFileItemCount, FileItem.MissingFileItemCount);
 
@@ -1189,13 +1194,21 @@ namespace UnrealBuildTool
                     }
                     Targets.Add(Target);
 
-                    // Load the include dependency cache.
-                    if (CPPEnvironment.IncludeDependencyCache == null)
-                    {
-                        CPPEnvironment.IncludeDependencyCache = DependencyCache.Create(
-                            DependencyCache.GetDependencyCachePathForTarget(Target)
-                            );
-                    }
+					// Load the include dependency cache.
+					// @todo fastubt: Ideally load this asynchronously at startup and only block when it is first needed and not finished loading
+					{	
+						if (!CPPEnvironment.IncludeDependencyCache.ContainsKey(Target))
+						{
+							CPPEnvironment.IncludeDependencyCache.Add( Target, DependencyCache.Create( DependencyCache.GetDependencyCachePathForTarget(Target) ) );
+						}
+					}
+					if( BuildConfiguration.bUseExperimentalFastDependencyScan )
+					{ 
+						if (!CPPEnvironment.FlatCPPIncludeDependencyCache.ContainsKey(Target))
+						{
+							CPPEnvironment.FlatCPPIncludeDependencyCache.Add( Target, FlatCPPIncludeDependencyCache.Create( Target ) );
+						}
+					}
 
                     var TargetOutputItems = new List<FileItem>();
                     BuildResult = Target.Build(ToolChain, TargetOutputItems, out EULAViolationWarning);
@@ -1230,33 +1243,26 @@ namespace UnrealBuildTool
                             LinkActionsAndItems();
                             var ActionsToExecute = AllActions;
 
-                            var VisualizationType = ActionGraphVisualizationType.OnlyCPlusPlusFilesAndHeaders;
-                            SaveActionGraphVisualization( Path.Combine( BuildConfiguration.BaseIntermediatePath, Target.GetTargetName() + ".gexf" ), Target.GetTargetName(), VisualizationType, ActionsToExecute );
-                        }
-
-                        // Save the include dependency cache.
-                        if (CPPEnvironment.IncludeDependencyCache != null)
-                        {
-                            CPPEnvironment.IncludeDependencyCache.Save();
-                            CPPEnvironment.IncludeDependencyCache = null;
-                        }
-                    }
+							var VisualizationType = ActionGraphVisualizationType.OnlyCPlusPlusFilesAndHeaders;
+							SaveActionGraphVisualization( Target, Path.Combine( BuildConfiguration.BaseIntermediatePath, Target.GetTargetName() + ".gexf" ), Target.GetTargetName(), VisualizationType, ActionsToExecute );
+						}
+					}
 
                     var TargetBuildTime = (DateTime.UtcNow - TargetStartTime).TotalSeconds;
 
-                    // Send out telemetry for this target
-                    Telemetry.SendEvent("TargetBuildStats.2",
-                        "AppName", Target.AppName,
-                        "GameName", Target.GameName,
-                        "Platform", Target.Platform.ToString(),
-                        "Configuration", Target.Configuration.ToString(),
-                        "CleanTarget", UEBuildConfiguration.bCleanProject.ToString(),
-                        "Monolithic", Target.ShouldCompileMonolithic().ToString(),
-                        "CreateDebugInfo", Target.IsCreatingDebugInfo().ToString(),
-                        "TargetType", Target.Rules.Type.ToString(),
-                        "TargetCreateTimeSec", TargetBuildTime.ToString("0.00")
-                        );
-                }
+					// Send out telemetry for this target
+					Telemetry.SendEvent("TargetBuildStats.2",
+						"AppName", Target.AppName,
+						"GameName", Target.GameName,
+						"Platform", Target.Platform.ToString(),
+						"Configuration", Target.Configuration.ToString(),
+						"CleanTarget", UEBuildConfiguration.bCleanProject.ToString(),
+						"Monolithic", Target.ShouldCompileMonolithic().ToString(),
+						"CreateDebugInfo", Target.IsCreatingDebugInfo().ToString(),
+						"TargetType", Target.Rules.Type.ToString(),
+						"TargetCreateTimeSec", TargetBuildTime.ToString("0.00")
+						);
+				}
 
                 if (BuildResult == ECompilationResult.Succeeded &&
                     (
@@ -1265,7 +1271,8 @@ namespace UnrealBuildTool
                     ))
                 {
                     // Plan the actions to execute for the build.
-                    List<Action> ActionsToExecute = GetActionsToExecute(OutputItemsForAllTargets, Targets);
+					Dictionary<UEBuildTarget,List<FileItem>> TargetToOutdatedPrerequisitesMap;
+                    List<Action> ActionsToExecute = GetActionsToExecute(OutputItemsForAllTargets, Targets, out TargetToOutdatedPrerequisitesMap);
 
                     // Display some stats to the user.
                     Log.TraceVerbose(
@@ -1276,34 +1283,70 @@ namespace UnrealBuildTool
 
                     ToolChain.PreBuildSync();
 
+					
+					// Cache indirect includes for all outdated C++ files.  We kick this off as a background thread so that it can
+					// perform the scan while we're compilng.  It usually only takes up to a few seconds, but we don't want to hurt
+					// our best case UBT iteration times for this task which can easily be performed asynchronously
+					Thread CPPIncludesThread = null;
+					if( BuildConfiguration.bUseExperimentalFastDependencyScan )
+					{
+						CPPIncludesThread = CreateThreadForCachingCPPIncludes( TargetToOutdatedPrerequisitesMap );
+						CPPIncludesThread.Start();
+					}
+
                     // Execute the actions.
                     bSuccess = ExecuteActions(ActionsToExecute, out ExecutorName);
 
-                    // if the build succeeded, do any needed syncing
-                    if (bSuccess)
-                    {
-                        foreach (UEBuildTarget Target in Targets)
-                        {
-                            ToolChain.PostBuildSync(Target);
-                        }
-                    }
-                    else
-                    {
-                        BuildResult = ECompilationResult.OtherCompilationError;
-                    }
-                }
-            }
-            catch (BuildException Exception)
-            {
-                // Output the message only, without the call stack
-                Log.TraceInformation(Exception.Message);
-                BuildResult = ECompilationResult.OtherCompilationError;
-            }
-            catch (Exception Exception)
-            {
-                Log.TraceInformation("ERROR: {0}", Exception);
-                BuildResult = ECompilationResult.OtherCompilationError;
-            }
+					// if the build succeeded, do any needed syncing
+					if (bSuccess)
+					{
+						foreach (UEBuildTarget Target in Targets)
+						{
+							ToolChain.PostBuildSync(Target);
+						}
+					}
+					else
+					{
+						BuildResult = ECompilationResult.OtherCompilationError;
+					}
+
+					if( CPPIncludesThread != null )
+					{ 
+						// Wait until our CPPIncludes dependency scanner thread has finished
+						CPPIncludesThread.Join();
+					}
+				}
+
+				// Save the include dependency cache.
+				{ 
+					// @todo fastubt urgent: If we update the cache and actually compile files, but DON'T SAVE the stream, we will be in big trouble on the next run -- the files could be out of date due to an indirect include, but we won't know how to check because we didn't cache the includes off!
+					//			-> Link errors shouldn't prevent the cache from saving, but UBT exceptions could.  
+					//			-> In any case, we need to:  A) Wipe cache file on exception	B) force SLOW intermediate checking on restart if not fully clean and cache file is missing (or force a clean first)
+					//			
+					foreach( var DependencyCache in CPPEnvironment.IncludeDependencyCache.Values )
+					{ 
+						DependencyCache.Save();
+					}
+					CPPEnvironment.IncludeDependencyCache.Clear();
+
+					foreach( var FlatCPPIncludeDependencyCache in CPPEnvironment.FlatCPPIncludeDependencyCache.Values )
+					{ 
+						FlatCPPIncludeDependencyCache.Save();
+					}
+					CPPEnvironment.FlatCPPIncludeDependencyCache.Clear();
+				}
+			}
+			catch (BuildException Exception)
+			{
+				// Output the message only, without the call stack
+				Log.TraceInformation(Exception.Message);
+				BuildResult = ECompilationResult.OtherCompilationError;
+			}
+			catch (Exception Exception)
+			{
+				Log.TraceInformation("ERROR: {0}", Exception);
+				BuildResult = ECompilationResult.OtherCompilationError;
+			}
 
             if(EULAViolationWarning != null)
             {
@@ -1512,5 +1555,34 @@ namespace UnrealBuildTool
 
             return TargetSettings;
         }
+
+
+		/// <summary>
+		/// Returns a Thread object that can be kicked off to update C++ include dependency cache
+		/// </summary>
+		/// <param name="TargetToOutdatedPrerequisitesMap">Maps each target to a list of outdated C++ files that need indirect dependencies cached</param>
+		/// <returns>The thread object</returns>
+		private static Thread CreateThreadForCachingCPPIncludes( Dictionary<UEBuildTarget, List<FileItem>> TargetToOutdatedPrerequisitesMap)
+		{
+			return new Thread( new ThreadStart( () =>
+				{
+					// @todo fastubt: This thread will access data structures that are also used on the main UBT thread, but during this time UBT
+					// is only invoking the build executor, so should not be touching this stuff.  However, we need to at some guards to make sure.
+
+					foreach( var TargetAndOutdatedPrerequisites in TargetToOutdatedPrerequisitesMap )
+					{
+						var Target = TargetAndOutdatedPrerequisites.Key;
+						var OutdatedPrerequisites = TargetAndOutdatedPrerequisites.Value;
+
+						foreach( var PrerequisiteItem in OutdatedPrerequisites )
+						{
+							// Invoke our deep include scanner to figure out whether any of the files included by this source file have
+							// changed since the build product was built
+							var FileBuildPlatform = UEBuildPlatform.GetBuildPlatformForCPPTargetPlatform( PrerequisiteItem.CachedCPPEnvironment.Config.Target.Platform );
+							CPPEnvironment.FindAndCacheAllIncludedFiles( Target, PrerequisiteItem, FileBuildPlatform, PrerequisiteItem.CachedCPPEnvironment.GetIncludesPathsToSearch( PrerequisiteItem ), PrerequisiteItem.CachedCPPEnvironment.IncludeFileSearchDictionary, bOnlyCachedDependencies:false );
+						}
+					}
+				} ) );
+		}
     }
 }
diff --git a/Engine/Source/Programs/UnrealBuildTool/System/Utils.cs b/Engine/Source/Programs/UnrealBuildTool/System/Utils.cs
--- a/Engine/Source/Programs/UnrealBuildTool/System/Utils.cs
+++ b/Engine/Source/Programs/UnrealBuildTool/System/Utils.cs
@@ -57,7 +57,7 @@ namespace UnrealBuildTool
 		/// <param name="ModuleName">Name of the module this PCH is being generated for</param>
 		/// <param name="bAllowDLLExports">True if we should allow DLLEXPORT definitions for this PCH</param>
 		/// <returns>the compilation output result of the created pch.</returns>
-		public static CPPOutput GeneratePCHCreationAction(string PCHHeaderNameInCode, FileItem PrecompiledHeaderIncludeFilename, CPPEnvironment ProjectCPPEnvironment, string OutputDirectory, string ModuleName, bool bAllowDLLExports )
+		public static CPPOutput GeneratePCHCreationAction(UEBuildTarget Target, string PCHHeaderNameInCode, FileItem PrecompiledHeaderIncludeFilename, CPPEnvironment ProjectCPPEnvironment, string OutputDirectory, string ModuleName, bool bAllowDLLExports )
 		{
 			// Find the header file to be precompiled. Don't skip external headers
 			if (PrecompiledHeaderIncludeFilename.bExists)
@@ -92,10 +92,13 @@ namespace UnrealBuildTool
 					}
 				}
 
+				// Cache our CPP environment so that we can check for outdatedness quickly.  Only files that have includes need this.
+				DummyPCH.CachedCPPEnvironment = ProjectPCHEnvironment;
+
 				Log.TraceVerbose( "Found PCH file \"{0}\".", PrecompiledHeaderIncludeFilename );
 
 				// Create the action to compile the PCH file.
-				return ProjectPCHEnvironment.CompileFiles(new List<FileItem>() { DummyPCH }, ModuleName);
+				return ProjectPCHEnvironment.CompileFiles(Target, new List<FileItem>() { DummyPCH }, ModuleName);
 			}
 			throw new BuildException( "Couldn't find PCH file \"{0}\".", PrecompiledHeaderIncludeFilename );
 		}
@@ -257,6 +260,7 @@ namespace UnrealBuildTool
 		 */
 		public static void SetEnvironmentVariablesFromBatchFile(string BatchFileName)
 		{
+			// @todo fastubt: Experiment with changing this to run asynchronously at startup, and only blocking if accessed before the .bat file finishes
 			if( File.Exists( BatchFileName ) )
 			{
 				// Create a wrapper batch file that echoes environment variables to a text file
@@ -475,7 +479,7 @@ namespace UnrealBuildTool
 		/// <returns>File path with consistent separators</returns>
 		public static string CleanDirectorySeparators( string FilePath, char UseDirectorySeparatorChar = '\0' )
 		{
-			StringBuilder CleanPath = new StringBuilder(FilePath.Length);
+			StringBuilder CleanPath = null;
 			if (UseDirectorySeparatorChar == '\0')
 			{
 				UseDirectorySeparatorChar = Environment.OSVersion.Platform == PlatformID.Unix ? '/' : '\\';
@@ -488,9 +492,21 @@ namespace UnrealBuildTool
 				char C = FilePath[Index];				
 				if (C == '/' || C == '\\')
 				{
-					C = UseDirectorySeparatorChar;
+					if( C != UseDirectorySeparatorChar )
+					{ 
+						C = UseDirectorySeparatorChar;
+						if( CleanPath == null )
+						{ 
+							CleanPath = new StringBuilder( FilePath.Substring( 0, Index ), FilePath.Length );
+						}
+					}
+
 					if (bCanCheckDoubleSeparators && C == PrevC)
 					{
+						if( CleanPath == null )
+						{ 
+							CleanPath = new StringBuilder( FilePath.Substring( 0, Index ), FilePath.Length );
+						}
 						continue;
 					}
 				}
@@ -499,10 +515,14 @@ namespace UnrealBuildTool
 					// First non-separator character, safe to check double separators
 					bCanCheckDoubleSeparators = true;
 				}
-				CleanPath.Append(C);
+
+				if( CleanPath != null )
+				{ 
+					CleanPath.Append(C);
+				}
 				PrevC = C;
 			}
-			return CleanPath.ToString();
+			return CleanPath != null ? CleanPath.ToString() : FilePath;
 		}
 
 	
diff --git a/Engine/Source/Programs/UnrealBuildTool/ToolChain/UEToolChain.cs b/Engine/Source/Programs/UnrealBuildTool/ToolChain/UEToolChain.cs
--- a/Engine/Source/Programs/UnrealBuildTool/ToolChain/UEToolChain.cs
+++ b/Engine/Source/Programs/UnrealBuildTool/ToolChain/UEToolChain.cs
@@ -13,9 +13,9 @@ namespace UnrealBuildTool
 	{
 		void RegisterToolChain();
 		
-		CPPOutput CompileCPPFiles(CPPEnvironment CompileEnvironment, List<FileItem> SourceFiles, string ModuleName);
+		CPPOutput CompileCPPFiles(UEBuildTarget Target, CPPEnvironment CompileEnvironment, List<FileItem> SourceFiles, string ModuleName);
 		
-		CPPOutput CompileRCFiles(CPPEnvironment Environment, List<FileItem> RCFiles);
+		CPPOutput CompileRCFiles(UEBuildTarget Target, CPPEnvironment Environment, List<FileItem> RCFiles);
 		
 		FileItem LinkFiles(LinkEnvironment LinkEnvironment, bool bBuildImportLibraryOnly);
 		
@@ -75,9 +75,9 @@ namespace UnrealBuildTool
 
 		public abstract void RegisterToolChain();
 
-		public abstract CPPOutput CompileCPPFiles(CPPEnvironment CompileEnvironment, List<FileItem> SourceFiles, string ModuleName);
+		public abstract CPPOutput CompileCPPFiles(UEBuildTarget Target, CPPEnvironment CompileEnvironment, List<FileItem> SourceFiles, string ModuleName);
 
-		public virtual CPPOutput CompileRCFiles(CPPEnvironment Environment, List<FileItem> RCFiles)
+		public virtual CPPOutput CompileRCFiles(UEBuildTarget Target, CPPEnvironment Environment, List<FileItem> RCFiles)
 		{
 			CPPOutput Result = new CPPOutput();
 			return Result;
diff --git a/Engine/Source/Programs/UnrealBuildTool/UnrealBuildTool.csproj b/Engine/Source/Programs/UnrealBuildTool/UnrealBuildTool.csproj
--- a/Engine/Source/Programs/UnrealBuildTool/UnrealBuildTool.csproj
+++ b/Engine/Source/Programs/UnrealBuildTool/UnrealBuildTool.csproj
@@ -215,6 +215,7 @@
     <Compile Include="System\ActionHistory.cs" />
     <Compile Include="System\BuildException.cs" />
     <Compile Include="Configuration\BuildConfiguration.cs" />
+    <Compile Include="System\FlatCPPIncludeDepencencyCache.cs" />
     <Compile Include="System\Distcc.cs" />
     <Compile Include="System\ModuleProcessingException.cs" />
     <Compile Include="System\MakefileGenerator.cs" />
diff --git a/Engine/Source/Programs/UnrealBuildTool/UnrealBuildTool_Mono.csproj b/Engine/Source/Programs/UnrealBuildTool/UnrealBuildTool_Mono.csproj
--- a/Engine/Source/Programs/UnrealBuildTool/UnrealBuildTool_Mono.csproj
+++ b/Engine/Source/Programs/UnrealBuildTool/UnrealBuildTool_Mono.csproj
@@ -124,6 +124,7 @@
     <Compile Include="System\DirectoryLookupCache.cs" />
     <Compile Include="System\CPPEnvironment.cs" />
     <Compile Include="System\CPPHeaders.cs" />
+    <Compile Include="System\CPPIncludeDepencencyCache.cs" />
     <Compile Include="System\DependencyCache.cs" />
     <Compile Include="System\DynamicCompilation.cs" />
     <Compile Include="System\ExternalExecution.cs" />
diff --git a/Engine/Source/Programs/UnrealBuildTool/WinRT/WinRTToolChain.cs b/Engine/Source/Programs/UnrealBuildTool/WinRT/WinRTToolChain.cs
--- a/Engine/Source/Programs/UnrealBuildTool/WinRT/WinRTToolChain.cs
+++ b/Engine/Source/Programs/UnrealBuildTool/WinRT/WinRTToolChain.cs
@@ -421,7 +421,7 @@ namespace UnrealBuildTool
 			return Result;
 		}
 
-		public override CPPOutput CompileCPPFiles(CPPEnvironment CompileEnvironment, List<FileItem> SourceFiles, string ModuleName)
+		public override CPPOutput CompileCPPFiles(UEBuildTarget Target, CPPEnvironment CompileEnvironment, List<FileItem> SourceFiles, string ModuleName)
 		{
 			string Arguments = GetCLArguments_Global(CompileEnvironment);
 
@@ -480,6 +480,8 @@ namespace UnrealBuildTool
 
 // Log.TraceInformation("Compile Arguments for {0}:", ModuleName);
 // Log.TraceInformation(Arguments);
+
+			var BuildPlatform = UEBuildPlatform.GetBuildPlatformForCPPTargetPlatform(CompileEnvironment.Config.Target.Platform);
 			
 			// Create a compile action for each source file.
 			CPPOutput Result = new CPPOutput();
@@ -491,9 +493,9 @@ namespace UnrealBuildTool
 
 				// Add the C++ source file and its included files to the prerequisite item list.
 				CompileAction.PrerequisiteItems.Add(SourceFile);
-				foreach (FileItem IncludedFile in CompileEnvironment.GetIncludeDependencies(SourceFile))
 				{
-					CompileAction.PrerequisiteItems.Add(IncludedFile);
+					var IncludedFileList = CPPEnvironment.FindAndCacheAllIncludedFiles( Target, SourceFile, BuildPlatform, CompileEnvironment.GetIncludesPathsToSearch( SourceFile ), CompileEnvironment.IncludeFileSearchDictionary, bOnlyCachedDependencies:BuildConfiguration.bUseExperimentalFastDependencyScan );
+					CompileAction.PrerequisiteItems.AddRange( IncludedFileList );
 				}
 
 				// If this is a CLR file then make sure our dependent assemblies are added as prerequisites
diff --git a/Engine/Source/Programs/UnrealBuildTool/Windows/VCToolChain.cs b/Engine/Source/Programs/UnrealBuildTool/Windows/VCToolChain.cs
--- a/Engine/Source/Programs/UnrealBuildTool/Windows/VCToolChain.cs
+++ b/Engine/Source/Programs/UnrealBuildTool/Windows/VCToolChain.cs
@@ -24,7 +24,12 @@ namespace UnrealBuildTool
 		{
 			string Result = "";
 
-			//Result += " /showIncludes";
+			// @todo fastubt: Use should be using StringBuilder instead of concatenation in all of these toolchain files.  On Mono, string builder is dramatically faster.
+
+
+			// @todo fastubt: We can use '/showIncludes' to accelerate outdatedness checking, as the compiler will discover all indirect includes itself.  But, the spew is pretty noisy!
+			//		-> If no files in source file chain have changed, even if the build product is outdated, we can skip using /showIncludes (relies on cache surviving)
+			// Result += " /showIncludes";
 
 			if( WindowsPlatform.bCompileWithClang )
 			{ 
@@ -639,7 +644,7 @@ namespace UnrealBuildTool
 			return Result;
 		}
 
-		public override CPPOutput CompileCPPFiles(CPPEnvironment CompileEnvironment, List<FileItem> SourceFiles, string ModuleName)
+		public override CPPOutput CompileCPPFiles(UEBuildTarget Target, CPPEnvironment CompileEnvironment, List<FileItem> SourceFiles, string ModuleName)
 		{
 			string Arguments = GetCLArguments_Global(CompileEnvironment);
 
@@ -699,6 +704,8 @@ namespace UnrealBuildTool
 				Arguments += string.Format(" /D \"{0}\"", DefinitionArgument);
 			}
 
+			var BuildPlatform = UEBuildPlatform.GetBuildPlatformForCPPTargetPlatform(CompileEnvironment.Config.Target.Platform);
+
 			// Create a compile action for each source file.
 			CPPOutput Result = new CPPOutput();
 			foreach (FileItem SourceFile in SourceFiles)
@@ -709,10 +716,13 @@ namespace UnrealBuildTool
 
 				// Add the C++ source file and its included files to the prerequisite item list.
 				CompileAction.PrerequisiteItems.Add(SourceFile);
-				var IncludeDependencies = CompileEnvironment.GetIncludeDependencies( SourceFile );
-				foreach (FileItem IncludedFile in IncludeDependencies)
 				{
-					CompileAction.PrerequisiteItems.Add(IncludedFile);
+					// @todo fastubt: What if one of the prerequisite files has become missing since it was updated in our cache? (usually, because a coder eliminated the source file)
+					//		-> Two CASES:
+					//				1) NOT WORKING: Non-unity file went away (SourceFile in this context).  That seems like an existing old use case.  Compile params or Response file should have changed?
+					//				2) WORKING: Indirect file went away (unity'd original source file or include).  This would return a file that no longer exists and adds to the prerequiteitems list
+					var IncludedFileList = CPPEnvironment.FindAndCacheAllIncludedFiles( Target, SourceFile, BuildPlatform, CompileEnvironment.GetIncludesPathsToSearch( SourceFile ), CompileEnvironment.IncludeFileSearchDictionary, bOnlyCachedDependencies:BuildConfiguration.bUseExperimentalFastDependencyScan );
+					CompileAction.PrerequisiteItems.AddRange( IncludedFileList );
 				}
 
 				// If this is a CLR file then make sure our dependent assemblies are added as prerequisites
@@ -951,10 +961,12 @@ namespace UnrealBuildTool
 			return Result;
 		}
 
-		public override CPPOutput CompileRCFiles(CPPEnvironment Environment, List<FileItem> RCFiles)
+		public override CPPOutput CompileRCFiles(UEBuildTarget Target, CPPEnvironment Environment, List<FileItem> RCFiles)
 		{
 			CPPOutput Result = new CPPOutput();
 
+			var BuildPlatform = UEBuildPlatform.GetBuildPlatformForCPPTargetPlatform(Environment.Config.Target.Platform);
+
 			foreach (FileItem RCFile in RCFiles)
 			{
 				Action CompileAction = new Action(ActionType.Compile);
@@ -1023,9 +1035,9 @@ namespace UnrealBuildTool
 				CompileAction.CommandArguments += string.Format(" \"{0}\"", RCFile.AbsolutePath);
 
 				// Add the files included by the RC file as prerequisites of the action.
-				foreach (FileItem IncludedFile in Environment.GetIncludeDependencies(RCFile))
 				{
-					CompileAction.PrerequisiteItems.Add(IncludedFile);
+					var IncludedFileList = CPPEnvironment.FindAndCacheAllIncludedFiles( Target, RCFile, BuildPlatform, Environment.GetIncludesPathsToSearch( RCFile ), Environment.IncludeFileSearchDictionary, bOnlyCachedDependencies:BuildConfiguration.bUseExperimentalFastDependencyScan );
+					CompileAction.PrerequisiteItems.AddRange( IncludedFileList );
 				}
 			}
 
