Proper name for the file in UBT mono project.

diff --git a/Engine/Source/Programs/UnrealBuildTool/UnrealBuildTool_Mono.csproj b/Engine/Source/Programs/UnrealBuildTool/UnrealBuildTool_Mono.csproj
--- a/Engine/Source/Programs/UnrealBuildTool/UnrealBuildTool_Mono.csproj
+++ b/Engine/Source/Programs/UnrealBuildTool/UnrealBuildTool_Mono.csproj
@@ -124,7 +124,7 @@
     <Compile Include="System\DirectoryLookupCache.cs" />
     <Compile Include="System\CPPEnvironment.cs" />
     <Compile Include="System\CPPHeaders.cs" />
-    <Compile Include="System\CPPIncludeDepencencyCache.cs" />
+    <Compile Include="System\FlatCPPIncludeDepencencyCache.cs" />
     <Compile Include="System\DependencyCache.cs" />
     <Compile Include="System\DynamicCompilation.cs" />
     <Compile Include="System\ExternalExecution.cs" />
