Fixed a problem with default target arguments in UE4 Xcode project generated with intellisense enabled (issue introduced by CL #2237527)

diff --git a/Engine/Source/Programs/UnrealBuildTool/System/XcodeProjectFileGenerator.cs b/Engine/Source/Programs/UnrealBuildTool/System/XcodeProjectFileGenerator.cs
--- a/Engine/Source/Programs/UnrealBuildTool/System/XcodeProjectFileGenerator.cs
+++ b/Engine/Source/Programs/UnrealBuildTool/System/XcodeProjectFileGenerator.cs
@@ -56,12 +56,13 @@ namespace UnrealBuildTool
 		/// <param name="InTargetPlatform">Name of the target that may be built.</param>
 		/// <param name="InDependencies">Name of the target that may be built.</param>
 		/// <param name="bHasPlist">Name of the target that may be built.</param>
-		public XcodeProjectTarget(string InDisplayName, string InTargetName, XcodeTargetType InType, string InProductName = "", UnrealTargetPlatform InTargetPlatform = UnrealTargetPlatform.Mac, bool bInIsMacOnly = false, List<XcodeTargetDependency> InDependencies = null, bool bHasPlist = false, List<XcodeFrameworkRef> InFrameworks = null)
+		public XcodeProjectTarget(string InDisplayName, string InTargetName, XcodeTargetType InType, string InTargetFilePath, string InProductName = "", UnrealTargetPlatform InTargetPlatform = UnrealTargetPlatform.Mac, bool bInIsMacOnly = false, List<XcodeTargetDependency> InDependencies = null, bool bHasPlist = false, List<XcodeFrameworkRef> InFrameworks = null)
 		{
 			DisplayName = InDisplayName;
 			TargetName = InTargetName;
 			Type = InType;
 			ProductName = InProductName;
+			TargetFilePath = InTargetFilePath;
 			TargetPlatform = InTargetPlatform;
 			bIsMacOnly = bInIsMacOnly;
 			Guid = XcodeProjectFileGenerator.MakeXcodeGuid();
@@ -91,6 +92,7 @@ namespace UnrealBuildTool
 		public UnrealTargetPlatform TargetPlatform;	//Mac, IOS
 		public bool bIsMacOnly;
 		public string ProductName;
+		public string TargetFilePath;
 
 		public string Guid;
 		public string BuildConfigGuild;
@@ -348,7 +350,7 @@ namespace UnrealBuildTool
 			if (Target.Type == XcodeTargetType.Legacy)
 			{
 				string UProjectPath = "";
-				if (UnrealBuildTool.HasUProjectFile())
+				if (UnrealBuildTool.HasUProjectFile() && !string.IsNullOrEmpty(Target.TargetFilePath) && Utils.IsFileUnderDirectory(Target.TargetFilePath, UnrealBuildTool.GetUProjectPath()))
 				{
 					if (MasterProjectRelativePath == UnrealBuildTool.GetUProjectPath())
 					{
@@ -1071,7 +1073,7 @@ namespace UnrealBuildTool
 					// @todo: Remove target platform param and merge Mac and iOS targets. For now BuildTarget knows how to build iOS, but cannot run iOS apps, so we need separate DeployTarget.
 					bool bIsMacOnly = !SupportedPlatforms.Contains(UnrealTargetPlatform.IOS);
 
-					XcodeProjectTarget BuildTarget = new XcodeProjectTarget(TargetName + " - Mac", TargetName, XcodeTargetType.Legacy, "", UnrealTargetPlatform.Mac, bIsMacOnly);
+					XcodeProjectTarget BuildTarget = new XcodeProjectTarget(TargetName + " - Mac", TargetName, XcodeTargetType.Legacy, TargetFilePath, "", UnrealTargetPlatform.Mac, bIsMacOnly);
 					if (!bGeneratingRunIOSProject)
 					{
 						ProjectTargets.Add(BuildTarget);
@@ -1088,14 +1090,14 @@ namespace UnrealBuildTool
 								FrameworkRefs.Add(new XcodeFrameworkRef(Framework));
 							}
 
-							XcodeProjectTarget IOSDeployTarget = new XcodeProjectTarget(TargetName + " - iOS", TargetName, XcodeTargetType.Native, TargetName + ".app", UnrealTargetPlatform.IOS, false, null, true, FrameworkRefs);
+							XcodeProjectTarget IOSDeployTarget = new XcodeProjectTarget(TargetName + " - iOS", TargetName, XcodeTargetType.Native, TargetFilePath, TargetName + ".app", UnrealTargetPlatform.IOS, false, null, true, FrameworkRefs);
 							ProjectTargets.Add(IOSDeployTarget);
 						}
 						else
 						{
 							XcodeContainerItemProxy ContainerProxy = new XcodeContainerItemProxy(ProjectTarget.Guid, BuildTarget.Guid, BuildTarget.DisplayName);
 							XcodeTargetDependency TargetDependency = new XcodeTargetDependency(BuildTarget.DisplayName, BuildTarget.Guid, ContainerProxy.Guid);
-							XcodeProjectTarget IOSDeployTarget = new XcodeProjectTarget(TargetName + " - iOS", TargetName, XcodeTargetType.Native, TargetName + ".app", UnrealTargetPlatform.IOS, false, new List<XcodeTargetDependency>() { TargetDependency }, true);
+							XcodeProjectTarget IOSDeployTarget = new XcodeProjectTarget(TargetName + " - iOS", TargetName, XcodeTargetType.Native, TargetFilePath, TargetName + ".app", UnrealTargetPlatform.IOS, false, new List<XcodeTargetDependency>() { TargetDependency }, true);
 							ProjectTargets.Add(IOSDeployTarget);
 							ContainerItemProxies.Add(ContainerProxy);
 							TargetDependencies.Add(TargetDependency);
@@ -1307,7 +1309,7 @@ namespace UnrealBuildTool
 					continue;
 				}
 
-				if (UnrealBuildTool.HasUProjectFile())
+				if (UnrealBuildTool.HasUProjectFile() && !string.IsNullOrEmpty(Target.TargetFilePath) && Utils.IsFileUnderDirectory(Target.TargetFilePath, UnrealBuildTool.GetUProjectPath()))
 				{
 					if (Target.TargetName.EndsWith("Editor"))
 					{
@@ -1428,13 +1430,13 @@ namespace UnrealBuildTool
 			// attempt to determine targets for the project
 			List<XcodeProjectTarget> ProjectTargets = new List<XcodeProjectTarget> ();
 			// add mandatory ones
-			XcodeProjectTarget UE4ProjectTarget = new XcodeProjectTarget ("UE4", "UE4", XcodeTargetType.Project);
-			XcodeProjectTarget UE4XcodeHelperTarget = new XcodeProjectTarget ("UE4XcodeHelper", "UE4XcodeHelper", XcodeTargetType.XcodeHelper, "libUE4XcodeHelper.a");
+			XcodeProjectTarget UE4ProjectTarget = new XcodeProjectTarget ("UE4", "UE4", XcodeTargetType.Project, "");
+			XcodeProjectTarget UE4XcodeHelperTarget = new XcodeProjectTarget ("UE4XcodeHelper", "UE4XcodeHelper", XcodeTargetType.XcodeHelper, "", "libUE4XcodeHelper.a");
 			ProjectTargets.AddRange(new XcodeProjectTarget[] { UE4ProjectTarget, UE4XcodeHelperTarget });
 
 			if (ProjectFilePlatform.HasFlag(XcodeProjectFilePlatform.iOS))
 			{
-				XcodeProjectTarget UE4CmdLineRunTarget = new XcodeProjectTarget("UE4CmdLineRun", "UE4CmdLineRun", XcodeTargetType.XCTest, "UE4CmdLineRun.xctest", UnrealTargetPlatform.IOS);
+				XcodeProjectTarget UE4CmdLineRunTarget = new XcodeProjectTarget("UE4CmdLineRun", "UE4CmdLineRun", XcodeTargetType.XCTest, "", "UE4CmdLineRun.xctest", UnrealTargetPlatform.IOS);
 				ProjectTargets.Add(UE4CmdLineRunTarget);
 				// This GUID will be referenced by each app's test action.
 				UE4CmdLineGuid = UE4CmdLineRunTarget.Guid;
