fix for ubt attempting to set the current directory Engine/Source when it doesn't exist
TTP342006
#ue4

diff --git a/Engine/Source/Programs/UnrealBuildTool/System/UnrealBuildTool.cs b/Engine/Source/Programs/UnrealBuildTool/System/UnrealBuildTool.cs
--- a/Engine/Source/Programs/UnrealBuildTool/System/UnrealBuildTool.cs
+++ b/Engine/Source/Programs/UnrealBuildTool/System/UnrealBuildTool.cs
@@ -585,7 +585,8 @@ namespace UnrealBuildTool
                     EngineSourceDirectory += "Source";
                 }
             }
-            Directory.SetCurrentDirectory(EngineSourceDirectory);
+//			if (Directory.Exists(EngineSourceDirectory))
+				Directory.SetCurrentDirectory(EngineSourceDirectory);
 
             // Build the list of game projects that we know about. When building from the editor (for hot-reload) or for Rocket projects, we require the 
             // project file to be passed in. Otherwise we scan for projects in directories named in UE4Games.uprojectdirs.
