Add an optional output manifest to UBT actions, which will allow detecting whether build products have been clobbered.

diff --git a/Engine/Source/Programs/UnrealBuildTool/System/ActionGraph.cs b/Engine/Source/Programs/UnrealBuildTool/System/ActionGraph.cs
--- a/Engine/Source/Programs/UnrealBuildTool/System/ActionGraph.cs
+++ b/Engine/Source/Programs/UnrealBuildTool/System/ActionGraph.cs
@@ -45,6 +45,7 @@ namespace UnrealBuildTool
 		public string CommandDescription = null;
 		public string StatusDescription = "...";
 		public string StatusDetailedDescription = "";
+		public string OutputManifestPath = null;
 		public bool bCanExecuteRemotely = false;
 		public bool bIsVCCompiler = false;
 		public bool bIsGCCCompiler = false;
@@ -349,6 +350,9 @@ namespace UnrealBuildTool
 						}
 					}
 				}
+
+				// Write the manifest for this action, if there is one
+				WriteOutputManifest(Action);
 			}
 
 			return Result;
@@ -707,6 +711,12 @@ namespace UnrealBuildTool
 					}
 				}
 
+				// Check if the manifest is out of date
+				if(!bIsOutdated && IsOutputManifestOutdated(RootAction))
+				{
+					bIsOutdated = true;
+				}
+
 				Log.WriteLineIf(BuildConfiguration.bLogDetailedActionStats && !String.IsNullOrEmpty( LatestUpdatedProducedItemName ),
 					TraceEventType.Verbose, "{0}: Oldest produced item is {1}", RootAction.StatusDescription, LatestUpdatedProducedItemName);
 
@@ -825,6 +835,65 @@ namespace UnrealBuildTool
 		}
 
 
+		/**
+		 * Determines if the output manifest for the given action is out of date. Actions which generate build products
+		 * in source-controlled directories can use a manifest to detect when they are externally modified.
+		 * @param RootAction The action being considered.
+		 * @return true if the manifest is outdated
+		 */
+		static public bool IsOutputManifestOutdated(Action RootAction)
+		{
+			if(RootAction.OutputManifestPath != null)
+			{
+				if(!File.Exists(RootAction.OutputManifestPath))
+				{
+					return true;
+				}
+
+				string[] Lines = File.ReadAllLines(RootAction.OutputManifestPath);
+				for(int Idx = 0; Idx + 2 < Lines.Length; Idx += 3)
+				{
+					FileInfo Info = new FileInfo(Lines[Idx]);
+					if(!Info.Exists)
+					{
+						return true;
+					}
+
+					long LastWriteTime;
+					if(!long.TryParse(Lines[Idx + 1], out LastWriteTime) || LastWriteTime != Info.LastWriteTimeUtc.Ticks)
+					{
+						return true;
+					}
+
+					long Length;
+					if(!long.TryParse(Lines[Idx + 2], out Length) || Length != Info.Length)
+					{
+						return true;
+					}
+				}
+			}
+			return false;
+		}
+
+		/**
+		 * Writes the output manifest for an action, if the OutputManifestPath is set.
+		 * @param RootAction The action being considered.
+		 */
+		static public void WriteOutputManifest(Action RootAction)
+		{
+			if(RootAction.OutputManifestPath != null && RootAction.ProducedItems.All(x => x.Info.Exists))
+			{
+				List<string> Lines = new List<string>();
+				foreach(FileInfo Info in RootAction.ProducedItems.Select(x => new FileInfo(x.AbsolutePath)))
+				{
+					Lines.Add(Info.FullName);
+					Lines.Add(Info.LastWriteTimeUtc.Ticks.ToString());
+					Lines.Add(Info.Length.ToString());
+				}
+				File.WriteAllLines(RootAction.OutputManifestPath, Lines);
+			}
+		}
+
 		/**
 		 * Builds a dictionary containing the actions from AllActions that are outdated by calling
 		 * IsActionOutdated.
diff --git a/Engine/Source/Programs/UnrealBuildTool/Windows/VCToolChain.cs b/Engine/Source/Programs/UnrealBuildTool/Windows/VCToolChain.cs
--- a/Engine/Source/Programs/UnrealBuildTool/Windows/VCToolChain.cs
+++ b/Engine/Source/Programs/UnrealBuildTool/Windows/VCToolChain.cs
@@ -1237,6 +1237,9 @@ namespace UnrealBuildTool
 			// Only execute linking on the local PC.
 			LinkAction.bCanExecuteRemotely = false;
 
+			// Create a manifest for all linked build products, since they might be clobbered from P4.
+			LinkAction.OutputManifestPath = Path.Combine(LinkEnvironment.Config.IntermediateDirectory, Path.GetFileName(LinkEnvironment.Config.OutputFilePath) + ".output");
+
 			return OutputFile;
 		}
 
