Generated code information separated into if the module *has* any and if so, *if* we're generating it now... then using the 'has' information to generate the linker fixup file.

#ttp 342622: CRITICAL: Regression: BINARY: PLATFORMERGAME: white square on the screen in cooked version

diff --git a/Engine/Source/Programs/UnrealBuildTool/Configuration/UEBuildModule.cs b/Engine/Source/Programs/UnrealBuildTool/Configuration/UEBuildModule.cs
--- a/Engine/Source/Programs/UnrealBuildTool/Configuration/UEBuildModule.cs
+++ b/Engine/Source/Programs/UnrealBuildTool/Configuration/UEBuildModule.cs
@@ -859,18 +859,29 @@ namespace UnrealBuildTool
 	{
 		public class AutoGenerateCppInfoClass
 		{
-			/** The filename of the *.generated.cpp file which was generated for the module */
-			public readonly string Filename;
-
-			public AutoGenerateCppInfoClass(string InFilename)
+			public class BuildInfoClass
 			{
-				Debug.Assert(InFilename != null);
+				/** The filename of the *.generated.cpp file which was generated for the module */
+				public readonly string Filename;
+
+				public BuildInfoClass(string InFilename)
+				{
+					Debug.Assert(InFilename != null);
+
+					Filename = InFilename;
+				}
+			}
 
-				Filename = InFilename;
+			/** Information about how to build the .generated.cpp files. If this is null, then we're not building .generated.cpp files for this module. */
+			public BuildInfoClass BuildInfo;
+
+			public AutoGenerateCppInfoClass(BuildInfoClass InBuildInfo)
+			{
+				BuildInfo = InBuildInfo;
 			}
 		}
 
-		/** Information about the .inl file . */
+		/** Information about the .generated.cpp file.  If this is null then this module doesn't have any UHT-produced code. */
 		public AutoGenerateCppInfoClass AutoGenerateCppInfo = null;
 
 		public class SourceFilesClass
@@ -1436,9 +1447,9 @@ namespace UnrealBuildTool
 				LinkInputFiles.AddRange( CPPCompileEnvironment.CompileFiles( Target, CPPFilesToCompile, Name ).ObjectFiles );
 			}
 
-			if (AutoGenerateCppInfo != null && !CPPCompileEnvironment.bHackHeaderGenerator)
+			if (AutoGenerateCppInfo != null && AutoGenerateCppInfo.BuildInfo != null && !CPPCompileEnvironment.bHackHeaderGenerator)
 			{
-				var GeneratedCppFileItem = FileItem.GetItemByPath( AutoGenerateCppInfo.Filename );
+				var GeneratedCppFileItem = FileItem.GetItemByPath( AutoGenerateCppInfo.BuildInfo.Filename );
 
 				CachePCHUsageForModuleSourceFile( this.Target, CPPCompileEnvironment, GeneratedCppFileItem );
 
diff --git a/Engine/Source/Programs/UnrealBuildTool/Configuration/UEBuildTarget.cs b/Engine/Source/Programs/UnrealBuildTool/Configuration/UEBuildTarget.cs
--- a/Engine/Source/Programs/UnrealBuildTool/Configuration/UEBuildTarget.cs
+++ b/Engine/Source/Programs/UnrealBuildTool/Configuration/UEBuildTarget.cs
@@ -1335,12 +1335,15 @@ namespace UnrealBuildTool
 						{
 							// If we've got this far and there are no source files then it's likely we're running Rocket and ignoring
 							// engine files, so we don't need a .generated.cpp either
+							UEBuildModuleCPP.AutoGenerateCppInfoClass.BuildInfoClass BuildInfo = null;
 							UHTModuleInfo.GeneratedCPPFilenameBase = Path.Combine( UEBuildModuleCPP.GetGeneratedCodeDirectoryForModule(this, UHTModuleInfo.ModuleDirectory, UHTModuleInfo.ModuleName), UHTModuleInfo.ModuleName ) + ".generated";
 							if (DependencyModuleCPP.SourceFilesToBuild.Count != 0)
 							{
-								DependencyModuleCPP.AutoGenerateCppInfo = new UEBuildModuleCPP.AutoGenerateCppInfoClass( UHTModuleInfo.GeneratedCPPFilenameBase + ".cpp" );
+								BuildInfo = new UEBuildModuleCPP.AutoGenerateCppInfoClass.BuildInfoClass( UHTModuleInfo.GeneratedCPPFilenameBase + ".cpp" );
 							}
 
+							DependencyModuleCPP.AutoGenerateCppInfo = new UEBuildModuleCPP.AutoGenerateCppInfoClass(BuildInfo);
+
 							UObjectModules.Add( UHTModuleInfo );
 							Log.TraceVerbose( "Detected UObject module: " + UHTModuleInfo.ModuleName );
 						}
