Revert change to the way that ndk_build is called.

diff --git a/Engine/Source/Programs/UnrealBuildTool/Android/UEDeployAndroid.cs b/Engine/Source/Programs/UnrealBuildTool/Android/UEDeployAndroid.cs
--- a/Engine/Source/Programs/UnrealBuildTool/Android/UEDeployAndroid.cs
+++ b/Engine/Source/Programs/UnrealBuildTool/Android/UEDeployAndroid.cs
@@ -466,10 +466,13 @@ namespace UnrealBuildTool.Android
 				NDKBuildInfo.UseShellExecute = true;
 				NDKBuildInfo.WindowStyle = ProcessWindowStyle.Minimized;
 				Console.WriteLine("\nRunning: " + NDKBuildInfo.FileName + " " + NDKBuildInfo.Arguments);
-				int NDKBuildExitCode = Utils.RunLocalProcessAndLogOutput(NDKBuildInfo);
+				Process NDKBuild = new Process();
+				NDKBuild.StartInfo = NDKBuildInfo;
+				NDKBuild.Start();
+				NDKBuild.WaitForExit();
 
 				// ndk build failure
-				if (NDKBuildExitCode != 0)
+				if (NDKBuild.ExitCode != 0)
 				{
 					throw new BuildException("ndk-build failed [{0}]", NDKBuildInfo.Arguments);
 				}
@@ -494,10 +497,13 @@ namespace UnrealBuildTool.Android
 			CallAntStartInfo.Arguments = "/c \"" + AntBuildPath + "\"  " + (bForDistribution ? "release" : "debug");
 			CallAntStartInfo.UseShellExecute = false;
 			Console.WriteLine("\nRunning: " + CallAntStartInfo.Arguments);
-			int CallAntExitCode = Utils.RunLocalProcessAndLogOutput(CallAntStartInfo);
+			Process CallAnt = new Process();
+			CallAnt.StartInfo = CallAntStartInfo;
+			CallAnt.Start();
+			CallAnt.WaitForExit();
 
 			// ant failure
-			if (CallAntExitCode != 0)
+			if (CallAnt.ExitCode != 0)
 			{
 				throw new BuildException("ant.bat failed [{0}]", CallAntStartInfo.Arguments);
 			}
