Optimization - use StringBuilder to concatenate compile arguments for VCToolChain, rather than just appending strings.

diff --git a/Engine/Source/Programs/UnrealBuildTool/Windows/VCToolChain.cs b/Engine/Source/Programs/UnrealBuildTool/Windows/VCToolChain.cs
--- a/Engine/Source/Programs/UnrealBuildTool/Windows/VCToolChain.cs
+++ b/Engine/Source/Programs/UnrealBuildTool/Windows/VCToolChain.cs
@@ -6,6 +6,7 @@ using System.Text.RegularExpressions;
 using System.Diagnostics;
 using System.IO;
 using Microsoft.Win32;
+using System.Text;
 
 namespace UnrealBuildTool
 {
@@ -20,10 +21,8 @@ namespace UnrealBuildTool
 			UEToolChain.RegisterPlatformToolChain(CPPTargetPlatform.Win32, this);
 		}
 
-		static string GetCLArguments_Global(CPPEnvironment CompileEnvironment)
+		static void AppendCLArguments_Global(CPPEnvironment CompileEnvironment, StringBuilder Arguments)
 		{
-			string Result = "";
-
 			// @todo fastubt: Use should be using StringBuilder instead of concatenation in all of these toolchain files.  On Mono, string builder is dramatically faster.
 
 
@@ -50,86 +49,86 @@ namespace UnrealBuildTool
 					default:
 						throw new BuildException( "Unexpected value for WindowsPlatform.Compiler: " + WindowsPlatform.Compiler.ToString() );
 				}
-				Result += " -fmsc-version=" + VersionString;
-				Result += " /D_MSC_FULL_VER=" + VersionString + "00000";
+				Arguments.Append(" -fmsc-version=" + VersionString);
+				Arguments.Append(" /D_MSC_FULL_VER=" + VersionString + "00000");
 			}
 
 
 			if( BuildConfiguration.bEnableCodeAnalysis )
 			{
-				Result += " /analyze";
+				Arguments.Append(" /analyze");
 
 				// Don't cause analyze warnings to be errors
-				Result += " /analyze:WX-";
+				Arguments.Append(" /analyze:WX-");
 
 				// Report functions that use a LOT of stack space.  You can lower this value if you
 				// want more aggressive checking for functions that use a lot of stack memory.
-				Result += " /analyze:stacksize81940";
+				Arguments.Append(" /analyze:stacksize81940");
 
 				// Don't bother generating code, only analyze code (may report fewer warnings though.)
-				//Result += " /analyze:only";
+				//Arguments.Append(" /analyze:only");
 			}
 
 			// Prevents the compiler from displaying its logo for each invocation.
-			Result += " /nologo";
+			Arguments.Append(" /nologo");
 
 			// Enable intrinsic functions.
-			Result += " /Oi";
+			Arguments.Append(" /Oi");
 
 			// Enable for static code analysis (where supported). Not treating analysis warnings as errors.
-			// Result += " /analyze:WX-";
+			// Arguments.Append(" /analyze:WX-");
 			
 			if( !WindowsPlatform.bCompileWithClang )	// @todo clang: These options are not supported with clang-cl yet
 			{
 				if (CompileEnvironment.Config.Target.Platform == CPPTargetPlatform.Win64)
 				{
 					// Pack struct members on 8-byte boundaries.
-					Result += " /Zp8";
+					Arguments.Append(" /Zp8");
 				}
 				else
 				{
 					// Pack struct members on 4-byte boundaries.
-					Result += " /Zp4";
+					Arguments.Append(" /Zp4");
 				}
 
 				// Separate functions for linker.
-				Result += " /Gy";
+				Arguments.Append(" /Gy");
 
 				// Relaxes floating point precision semantics to allow more optimization.
-				Result += " /fp:fast";
+				Arguments.Append(" /fp:fast");
 			}
 
 			// Compile into an .obj file, and skip linking.
-			Result += " /c";
+			Arguments.Append(" /c");
 
 			// Allow 800% of the default memory allocation limit.
-			Result += " /Zm800";
+			Arguments.Append(" /Zm800");
 
 			if( !WindowsPlatform.bCompileWithClang )	// @todo clang: Not supported in clang-cl yet
 			{
 				// Allow large object files to avoid hitting the 2^16 section limit when running with -StressTestUnity.
-				Result += " /bigobj";
+				Arguments.Append(" /bigobj");
 			}
 
 			// Disable "The file contains a character that cannot be represented in the current code page" warning for non-US windows.
-			Result += " /wd4819";
+			Arguments.Append(" /wd4819");
 
 			if( BuildConfiguration.bUseSharedPCHs )
 			{
 				// @todo SharedPCH: Disable warning about PCH defines not matching .cpp defines.  We "cheat" these defines a little
 				// bit to make shared PCHs work.  But it's totally safe.  Trust us.
-				Result += " /wd4651";
+				Arguments.Append(" /wd4651");
 
 				// @todo SharedPCH: Disable warning about redefining *API macros.  The PCH header is compiled with various DLLIMPORTs, but
 				// when a module that uses that PCH header *IS* one of those imports, that module is compiled with EXPORTS, so the macro
 				// is redefined on the command-line.  We need to clobber those defines to make shared PCHs work properly!
-				Result += " /wd4005";
+				Arguments.Append(" /wd4005");
 			}
 
 			// If compiling as a DLL, set the relevant defines
 			if (CompileEnvironment.Config.bIsBuildingDLL)
 			{
-				Result += " /D _WINDLL";
+				Arguments.Append(" /D _WINDLL");
 			}
 
 			// When targeting Windows XP with Visual Studio 2012+, we need to tell the compiler to use the older Windows SDK that works
@@ -138,11 +137,11 @@ namespace UnrealBuildTool
 			{
 				if( WindowsPlatform.Compiler == WindowsCompiler.VisualStudio2013 )
 				{
-					Result += " /D_USING_V120_SDK71_";
+					Arguments.Append(" /D_USING_V120_SDK71_");
 				}
 				else if (WindowsPlatform.Compiler == WindowsCompiler.VisualStudio2012)
 				{
-					Result += " /D_USING_V110_SDK71_";
+					Arguments.Append(" /D_USING_V110_SDK71_");
 				}
 			}
 
@@ -150,10 +149,10 @@ namespace UnrealBuildTool
 			// Handle Common Language Runtime support (C++/CLI)
 			if (CompileEnvironment.Config.CLRMode == CPPCLRMode.CLREnabled)
 			{
-				Result += " /clr";
+				Arguments.Append(" /clr");
 
 				// Don't use default lib path, we override it explicitly to use the 4.0 reference assemblies.
-				Result += " /clr:nostdlib";
+				Arguments.Append(" /clr:nostdlib");
 			}
 
 			//
@@ -162,15 +161,15 @@ namespace UnrealBuildTool
 			if (CompileEnvironment.Config.Target.Configuration == CPPTargetConfiguration.Debug)
 			{
 				// Disable compiler optimization.
-				Result += " /Od";
+				Arguments.Append(" /Od");
 
 				// Favor code size (especially useful for embedded platforms).
-				Result += " /Os";
+				Arguments.Append(" /Os");
 
 				// Allow inline method expansion unless E&C support is requested
 				if( !BuildConfiguration.bSupportEditAndContinue )
 				{
-					Result += " /Ob2";
+					Arguments.Append(" /Ob2");
 				}
 
 				if ((CompileEnvironment.Config.Target.Platform == CPPTargetPlatform.Win32) || 
@@ -179,7 +178,7 @@ namespace UnrealBuildTool
 					// Runtime stack checks are not allowed when compiling for CLR
 					if (CompileEnvironment.Config.CLRMode == CPPCLRMode.CLRDisabled)
 					{
-						Result += " /RTCs";
+						Arguments.Append(" /RTCs");
 					}
 				}
 			}
@@ -191,26 +190,26 @@ namespace UnrealBuildTool
 				// Maximum optimizations if desired.
 				if( CompileEnvironment.Config.OptimizeCode >= ModuleRules.CodeOptimization.InNonDebugBuilds )
 				{
-					Result += " /Ox";
+					Arguments.Append(" /Ox");
 
 					// Allow optimized code to be debugged more easily.  This makes PDBs a bit larger, but doesn't noticeably affect
 					// compile times.  The executable code is not affected at all by this switch, only the debugging information.
-					Result += " /d2Zi+";
+					Arguments.Append(" /d2Zi+");
 				}
 				
 				// Favor code speed.
-				Result += " /Ot";
+				Arguments.Append(" /Ot");
 
 				// Only omit frame pointers on the PC (which is implied by /Ox) if wanted.
 				if ( BuildConfiguration.bOmitFramePointers == false
 				&& ((CompileEnvironment.Config.Target.Platform == CPPTargetPlatform.Win32) ||
 					(CompileEnvironment.Config.Target.Platform == CPPTargetPlatform.Win64)))
 				{
-					Result += " /Oy-";
+					Arguments.Append(" /Oy-");
 				}
 
 				// Allow inline method expansion
-				Result += " /Ob2";
+				Arguments.Append(" /Ob2");
 
 				//
 				// LTCG
@@ -220,7 +219,7 @@ namespace UnrealBuildTool
 					if( BuildConfiguration.bAllowLTCG )
 					{
 						// Enable link-time code generation.
-						Result += " /GL";
+						Arguments.Append(" /GL");
 					}
 				}
 			}
@@ -240,12 +239,12 @@ namespace UnrealBuildTool
 					if( !WindowsPlatform.bCompileWithClang )
 					{
 						// Allow the compiler to generate SSE2 instructions.
-						Result += " /arch:SSE2";
+						Arguments.Append(" /arch:SSE2");
 					}
 				}
 
 				// Prompt the user before reporting internal errors to Microsoft.
-				Result += " /errorReport:prompt";
+				Arguments.Append(" /errorReport:prompt");
 
 				if (CompileEnvironment.Config.CLRMode == CPPCLRMode.CLRDisabled)
 				{
@@ -254,7 +253,7 @@ namespace UnrealBuildTool
 						&& !WindowsPlatform.bCompileWithClang)	// @todo clang: Clang compiler frontend doesn't understand /EHsc yet
 					{
 						// Enable C++ exception handling, but not C exceptions.
-						Result += " /EHsc";
+						Arguments.Append(" /EHsc");
 					}
 					else
 					{
@@ -265,14 +264,14 @@ namespace UnrealBuildTool
 				else
 				{
 					// For C++/CLI all exceptions must be left enabled
-					Result += " /EHa";
+					Arguments.Append(" /EHa");
 				}
             }
 
             // @todo - find a better place for this. 
             if (CompileEnvironment.Config.Target.Platform == CPPTargetPlatform.HTML5)
             {
-                Result += " /EHsc"; 
+                Arguments.Append(" /EHsc"); 
             }
             // If enabled, create debug information.
 			if (CompileEnvironment.Config.bCreateDebugInfo)
@@ -281,7 +280,7 @@ namespace UnrealBuildTool
 				{
 					// Manually enable full debug info.  This is basically DWARF3 debug data.  MSVC won't know what to do with this and
 					// will just strip it out of executables linked using the MSVC linker.
-					Result += " -Xclang -g";
+					Arguments.Append(" -Xclang -g");
 				}
 				else
 				{
@@ -293,24 +292,24 @@ namespace UnrealBuildTool
 							// We only need to do this in debug as that's the only configuration that supports E&C.
 							CompileEnvironment.Config.Target.Configuration == CPPTargetConfiguration.Debug)
 						{
-							Result += " /ZI";
+							Arguments.Append(" /ZI");
 						}
 						// Regular PDB debug information.
 						else
 						{
-							Result += " /Zi";
+							Arguments.Append(" /Zi");
 						}
                         // We need to add this so VS won't lock the PDB file and prevent synchronous updates. This forces serialization through MSPDBSRV.exe.
                         // See http://msdn.microsoft.com/en-us/library/dn502518.aspx for deeper discussion of /FS switch.
                         if (BuildConfiguration.bUseIncrementalLinking)
                         {
-                            Result += " /FS";
+                            Arguments.Append(" /FS");
                         }
 					}
 					// Store C7-format debug info in the .obj files, which is faster.
 					else
 					{
-						Result += " /Z7";
+						Arguments.Append(" /Z7");
 					}
 				}
 			}
@@ -318,29 +317,25 @@ namespace UnrealBuildTool
 			// Specify the appropriate runtime library based on the platform and config.
 			if( CompileEnvironment.Config.Target.Configuration == CPPTargetConfiguration.Debug && BuildConfiguration.bDebugBuildsActuallyUseDebugCRT )
 			{
-				Result += " /MDd";
+				Arguments.Append(" /MDd");
 			}
 			else
 			{
-				Result += " /MD";
+				Arguments.Append(" /MD");
 			}
-
-			return Result;
 		}
 
-		static string GetCLArguments_CPP( CPPEnvironment CompileEnvironment )
+		static void AppendCLArguments_CPP( CPPEnvironment CompileEnvironment, StringBuilder Arguments )
 		{
-			string Result = "";
-
 			// Explicitly compile the file as C++.
-			Result += " /TP";
+			Arguments.Append(" /TP");
 
 			if (!CompileEnvironment.Config.bEnableBufferSecurityChecks)
 			{
 				// This will disable buffer security checks (which are enabled by default) that the MS compiler adds around arrays on the stack,
 				// Which can add some performance overhead, especially in performance intensive code
 				// Only disable this if you know what you are doing, because it will be disabled for the entire module!
-				Result += " /GS-";
+				Arguments.Append(" /GS-");
 			}
 
 			// C++/CLI requires that RTTI is left enabled
@@ -349,17 +344,17 @@ namespace UnrealBuildTool
 				if (CompileEnvironment.Config.bUseRTTI)
 				{
 					// Enable C++ RTTI.
-					Result += " /GR";
+					Arguments.Append(" /GR");
 				}
 				else
 				{
 					// Disable C++ RTTI.
-					Result += " /GR-";
+					Arguments.Append(" /GR-");
 				}
 			}
 
 			// Level 4 warnings.
-			Result += " /W4";
+			Arguments.Append(" /W4");
 
 			if( WindowsPlatform.bCompileWithClang )
 			{
@@ -367,66 +362,58 @@ namespace UnrealBuildTool
 				// NOTE: These must appear after we set the MSVC warning level
 
 				// Allow Microsoft-specific syntax to slide, even though it may be non-standard.  Needed for Windows headers.
-				Result += " -Wno-microsoft";								
+				Arguments.Append(" -Wno-microsoft");								
 
 				// @todo clang: Kind of a shame to turn these off.  We'd like to catch unused variables, but it is tricky with how our assertion macros work.
-				Result += " -Wno-unused-variable";
-				Result += " -Wno-unused-function";
-				Result += " -Wno-unused-private-field";
-				Result += " -Wno-unused-value";
+				Arguments.Append(" -Wno-unused-variable");
+				Arguments.Append(" -Wno-unused-function");
+				Arguments.Append(" -Wno-unused-private-field");
+				Arguments.Append(" -Wno-unused-value");
 
-				Result += " -Wno-inline-new-delete";	// @todo clang: We declare operator new as inline.  Clang doesn't seem to like that.
+				Arguments.Append(" -Wno-inline-new-delete");	// @todo clang: We declare operator new as inline.  Clang doesn't seem to like that.
 
 				// @todo clang: Disabled warnings were copied from MacToolChain for the most part
-				Result += " -Wno-deprecated-declarations";
-				Result += " -Wno-deprecated-writable-strings";
-				Result += " -Wno-deprecated-register";
-				Result += " -Wno-switch-enum";
-				Result += " -Wno-logical-op-parentheses";	// needed for external headers we shan't change
-				Result += " -Wno-null-arithmetic";			// needed for external headers we shan't change
-				Result += " -Wno-deprecated-declarations";	// needed for wxWidgets
-				Result += " -Wno-return-type-c-linkage";	// needed for PhysX
-				Result += " -Wno-ignored-attributes";		// needed for nvtesslib
-				Result += " -Wno-uninitialized";
-				Result += " -Wno-tautological-compare";
-				Result += " -Wno-switch";
-				Result += " -Wno-invalid-offsetof"; // needed to suppress warnings about using offsetof on non-POD types.
+				Arguments.Append(" -Wno-deprecated-declarations");
+				Arguments.Append(" -Wno-deprecated-writable-strings");
+				Arguments.Append(" -Wno-deprecated-register");
+				Arguments.Append(" -Wno-switch-enum");
+				Arguments.Append(" -Wno-logical-op-parentheses");	// needed for external headers we shan't change
+				Arguments.Append(" -Wno-null-arithmetic");			// needed for external headers we shan't change
+				Arguments.Append(" -Wno-deprecated-declarations");	// needed for wxWidgets
+				Arguments.Append(" -Wno-return-type-c-linkage");	// needed for PhysX
+				Arguments.Append(" -Wno-ignored-attributes");		// needed for nvtesslib
+				Arguments.Append(" -Wno-uninitialized");
+				Arguments.Append(" -Wno-tautological-compare");
+				Arguments.Append(" -Wno-switch");
+				Arguments.Append(" -Wno-invalid-offsetof"); // needed to suppress warnings about using offsetof on non-POD types.
 			}
-
-			return Result;
 		}
 		
-		static string GetCLArguments_C()
+		static void AppendCLArguments_C(StringBuilder Arguments)
 		{
-			string Result = "";
-			
 			// Explicitly compile the file as C.
-			Result += " /TC";
+			Arguments.Append(" /TC");
 		
 			// Level 0 warnings.  Needed for external C projects that produce warnings at higher warning levels.
-			Result += " /W0";
-
-			return Result;
+			Arguments.Append(" /W0");
 		}
 
-		static string GetLinkArguments(LinkEnvironment LinkEnvironment)
+		static void AppendLinkArguments(LinkEnvironment LinkEnvironment, StringBuilder Arguments)
 		{
-			string Result = "";
-
 			// Don't create a side-by-side manifest file for the executable.
-			Result += " /MANIFEST:NO";
+			Arguments.Append(" /MANIFEST:NO");
 
 			// Prevents the linker from displaying its logo for each invocation.
-			Result += " /NOLOGO";
+			Arguments.Append(" /NOLOGO");
 
 			if (LinkEnvironment.Config.bCreateDebugInfo)
 			{
 				// Output debug info for the linked executable.
-				Result += " /DEBUG";
+				Arguments.Append(" /DEBUG");
 			}
 
 			// Prompt the user before reporting internal errors to Microsoft.
-			Result += " /errorReport:prompt";
+			Arguments.Append(" /errorReport:prompt");
 
 			//
 			//	PC
@@ -437,53 +424,53 @@ namespace UnrealBuildTool
 				// Set machine type/ architecture to be 64 bit.
 				if (LinkEnvironment.Config.Target.Platform == CPPTargetPlatform.Win64)
 				{
-					Result += " /MACHINE:x64";
+					Arguments.Append(" /MACHINE:x64");
 				}
 				// 32 bit executable/ target.
 				else
 				{
-					Result += " /MACHINE:x86";
+					Arguments.Append(" /MACHINE:x86");
 				}
 
 				{
 					if (LinkEnvironment.Config.bIsBuildingConsoleApplication)
 					{
-						Result += " /SUBSYSTEM:CONSOLE";
+						Arguments.Append(" /SUBSYSTEM:CONSOLE");
 					}
 					else
 					{
-						Result += " /SUBSYSTEM:WINDOWS";
+						Arguments.Append(" /SUBSYSTEM:WINDOWS");
 					}
 
 					// When targeting Windows XP in Visual Studio 2012+, we need to tell the linker we are going to support execution
 					// on that older platform.  The compiler defaults to version 6.0+.  We'll modify the SUBSYSTEM parameter here.
 					if( WindowsPlatform.SupportWindowsXP )
 					{
-						Result += LinkEnvironment.Config.Target.Platform == CPPTargetPlatform.Win64 ? ",5.02" : ",5.01";
+						Arguments.Append(LinkEnvironment.Config.Target.Platform == CPPTargetPlatform.Win64 ? ",5.02" : ",5.01");
 					}
 				}
 
 				if (LinkEnvironment.Config.bIsBuildingConsoleApplication && !LinkEnvironment.Config.bIsBuildingDLL && !String.IsNullOrEmpty( LinkEnvironment.Config.WindowsEntryPointOverride ) )
 				{
 					// Use overridden entry point
-					Result += " /ENTRY:" + LinkEnvironment.Config.WindowsEntryPointOverride;
+					Arguments.Append(" /ENTRY:" + LinkEnvironment.Config.WindowsEntryPointOverride);
 				}
 
 				// Allow the OS to load the EXE at different base addresses than its preferred base address.
-				Result += " /FIXED:No";
+				Arguments.Append(" /FIXED:No");
 
 				// Option is only relevant with 32 bit toolchain.
 				if (LinkEnvironment.Config.Target.Platform == CPPTargetPlatform.Win32)
 				{
 					// Disables the 2GB address space limit on 64-bit Windows and 32-bit Windows with /3GB specified in boot.ini
-					Result += " /LARGEADDRESSAWARE";
+					Arguments.Append(" /LARGEADDRESSAWARE");
 				}
 
 				// Explicitly declare that the executable is compatible with Data Execution Prevention.
-				Result += " /NXCOMPAT";
+				Arguments.Append(" /NXCOMPAT");
 
 				// Set the default stack size.
-				Result += " /STACK:5000000";
+				Arguments.Append(" /STACK:5000000");
 
 				// E&C can't use /SAFESEH.  Also, /SAFESEH isn't compatible with 64-bit linking
 				if( !BuildConfiguration.bSupportEditAndContinue &&
@@ -491,15 +478,15 @@ namespace UnrealBuildTool
 				{
 					// Generates a table of Safe Exception Handlers.  Documentation isn't clear whether they actually mean
 					// Structured Exception Handlers.
-					Result += " /SAFESEH";
+					Arguments.Append(" /SAFESEH");
 				}
 
 				// Allow delay-loaded DLLs to be explicitly unloaded.
-				Result += " /DELAY:UNLOAD";
+				Arguments.Append(" /DELAY:UNLOAD");
 
 				if (LinkEnvironment.Config.bIsBuildingDLL)
 				{
-					Result += " /DLL";
+					Arguments.Append(" /DLL");
 				}
 
 				if (LinkEnvironment.Config.CLRMode == CPPCLRMode.CLREnabled)
@@ -518,7 +505,7 @@ namespace UnrealBuildTool
 			// Don't embed the full PDB path in the binary when building Rocket executables; the location on disk won't match the user's install directory.
 			if(UnrealBuildTool.BuildingRocket())
 			{
-				Result += " /PDBALTPATH:%_PDB%";
+				Arguments.Append(" /PDBALTPATH:%_PDB%");
 			}
 
 			//
@@ -528,7 +515,7 @@ namespace UnrealBuildTool
 				LinkEnvironment.Config.Target.Configuration == CPPTargetConfiguration.Shipping)
 			{
 				// Use link-time code generation.
-				Result += " /LTCG";
+				Arguments.Append(" /LTCG");
 
 				// This is where we add in the PGO-Lite linkorder.txt if we are using PGO-Lite
 				//Result += " /ORDER:@linkorder.txt";
@@ -541,13 +528,13 @@ namespace UnrealBuildTool
 			if (LinkEnvironment.Config.Target.Configuration == CPPTargetConfiguration.Shipping)
 			{
 				// Generate an EXE checksum.
-				Result += " /RELEASE";
+				Arguments.Append(" /RELEASE");
 
 				// Eliminate unreferenced symbols.
-				Result += " /OPT:REF";
+				Arguments.Append(" /OPT:REF");
 
 				// Remove redundant COMDATs.
-				Result += " /OPT:ICF";
+				Arguments.Append(" /OPT:ICF");
 			}
 			//
 			//	Regular development binary. 
@@ -555,10 +542,10 @@ namespace UnrealBuildTool
 			else
 			{
 				// Keep symbols that are unreferenced.
-				Result += " /OPT:NOREF";
+				Arguments.Append(" /OPT:NOREF");
 
 				// Disable identical COMDAT folding.
-				Result += " /OPT:NOICF";
+				Arguments.Append(" /OPT:NOICF");
 			}
 
 			// Enable incremental linking if wanted.
@@ -566,11 +553,11 @@ namespace UnrealBuildTool
 			//		 will silently be ignored anyway
 			if (BuildConfiguration.bUseIncrementalLinking && LinkEnvironment.Config.CLRMode != CPPCLRMode.CLREnabled)
 			{
-				Result += " /INCREMENTAL";
+				Arguments.Append(" /INCREMENTAL");
 			}
 			else
 			{
-				Result += " /INCREMENTAL:NO";
+				Arguments.Append(" /INCREMENTAL:NO");
 			}
 
 			// Disable 
@@ -578,24 +565,20 @@ namespace UnrealBuildTool
 			// type warning as we leverage the DelayLoad option to put third-party DLLs into a 
 			// non-standard location. This requires the module(s) that use said DLL to ensure that it 
 			// is loaded prior to using it.
-			Result += " /ignore:4199";
+			Arguments.Append(" /ignore:4199");
 
 			// Suppress warnings about missing PDB files for statically linked libraries.  We often don't want to distribute
 			// PDB files for these libraries.
-			Result += " /ignore:4099";		// warning LNK4099: PDB '<file>' was not found with '<file>'
-
-			return Result;
+			Arguments.Append(" /ignore:4099");		// warning LNK4099: PDB '<file>' was not found with '<file>'
 		}
 
-		static string GetLibArguments(LinkEnvironment LinkEnvironment)
+		static void AppendLibArguments(LinkEnvironment LinkEnvironment, StringBuilder Arguments)
 		{
-			string Result = "";
-
 			// Prevents the linker from displaying its logo for each invocation.
-			Result += " /NOLOGO";
+			Arguments.Append(" /NOLOGO");
 
 			// Prompt the user before reporting internal errors to Microsoft.
-			Result += " /errorReport:prompt";
+			Arguments.Append(" /errorReport:prompt");
 
 			//
 			//	PC
@@ -605,29 +588,29 @@ namespace UnrealBuildTool
 				// Set machine type/ architecture to be 64 bit.
 				if (LinkEnvironment.Config.Target.Platform == CPPTargetPlatform.Win64)
 				{
-					Result += " /MACHINE:x64";
+					Arguments.Append(" /MACHINE:x64");
 				}
 				// 32 bit executable/ target.
 				else
 				{
-					Result += " /MACHINE:x86";
+					Arguments.Append(" /MACHINE:x86");
 				}
 
 				{
 					if (LinkEnvironment.Config.bIsBuildingConsoleApplication)
 					{
-						Result += " /SUBSYSTEM:CONSOLE";
+						Arguments.Append(" /SUBSYSTEM:CONSOLE");
 					}
 					else
 					{
-						Result += " /SUBSYSTEM:WINDOWS";
+						Arguments.Append(" /SUBSYSTEM:WINDOWS");
 					}
 
 					// When targeting Windows XP in Visual Studio 2012+, we need to tell the linker we are going to support execution
 					// on that older platform.  The compiler defaults to version 6.0+.  We'll modify the SUBSYSTEM parameter here.
 					if( WindowsPlatform.SupportWindowsXP )
 					{
-						Result += LinkEnvironment.Config.Target.Platform == CPPTargetPlatform.Win64 ? ",5.02" : ",5.01";
+						Arguments.Append(LinkEnvironment.Config.Target.Platform == CPPTargetPlatform.Win64 ? ",5.02" : ",5.01");
 					}
 				}
 			}
@@ -638,20 +621,19 @@ namespace UnrealBuildTool
 			if (LinkEnvironment.Config.Target.Configuration == CPPTargetConfiguration.Shipping)
 			{
 				// Use link-time code generation.
-				Result += " /LTCG";
+				Arguments.Append(" /LTCG");
 			}
-
-			return Result;
 		}
 
 		public override CPPOutput CompileCPPFiles(UEBuildTarget Target, CPPEnvironment CompileEnvironment, List<FileItem> SourceFiles, string ModuleName)
 		{
-			string Arguments = GetCLArguments_Global(CompileEnvironment);
+			StringBuilder Arguments = new StringBuilder();
+			AppendCLArguments_Global(CompileEnvironment, Arguments);
 
 			// Add include paths to the argument list.
 			foreach (string IncludePath in CompileEnvironment.Config.IncludePaths)
 			{
-				Arguments += string.Format(" /I \"{0}\"", IncludePath);
+				Arguments.AppendFormat(" /I \"{0}\"", IncludePath);
 			}
 			foreach (string IncludePath in CompileEnvironment.Config.SystemIncludePaths)
 			{
@@ -663,11 +645,11 @@ namespace UnrealBuildTool
 					// this case Clang will ignore any earlier occurrence of the include path, preventing a system header include path from overriding a different system 
 					// include path set later on by a module.  NOTE: When passing "-Xclang", these options will always appear at the end of the command-line string, meaning
 					// they will be forced to appear *after* all environment-variable-extracted includes.  This is technically okay though.
-					Arguments += string.Format(" -Xclang -internal-isystem -Xclang \"{0}\"", IncludePath);
+					Arguments.AppendFormat(" -Xclang -internal-isystem -Xclang \"{0}\"", IncludePath);
 				}
 				else
 				{
-					Arguments += string.Format(" /I \"{0}\"", IncludePath);
+					Arguments.AppendFormat(" /I \"{0}\"", IncludePath);
 				}
 			}
 
@@ -679,19 +661,19 @@ namespace UnrealBuildTool
 				// .cpp file to the assembly's location.				
 				foreach (string AssemblyPath in CompileEnvironment.Config.SystemDotNetAssemblyPaths)
 				{
-					Arguments += string.Format(" /AI \"{0}\"", AssemblyPath);
+					Arguments.AppendFormat(" /AI \"{0}\"", AssemblyPath);
 				}
 
 				// Add explicit .NET framework assembly references				
 				foreach (string AssemblyName in CompileEnvironment.Config.FrameworkAssemblyDependencies)
 				{
-					Arguments += string.Format( " /FU \"{0}\"", AssemblyName );
+					Arguments.AppendFormat( " /FU \"{0}\"", AssemblyName );
 				}
 
 				// Add private assembly references				
 				foreach( PrivateAssemblyInfo CurAssemblyInfo in CompileEnvironment.PrivateAssemblyDependencies )
 				{
-					Arguments += string.Format( " /FU \"{0}\"", CurAssemblyInfo.FileItem.AbsolutePath );
+					Arguments.AppendFormat( " /FU \"{0}\"", CurAssemblyInfo.FileItem.AbsolutePath );
 				}
 			}
 
@@ -701,7 +683,7 @@ namespace UnrealBuildTool
 			{
 				// Escape all quotation marks so that they get properly passed with the command line.
 				var DefinitionArgument = Definition.Contains("\"") ? Definition.Replace("\"", "\\\"") : Definition;
-				Arguments += string.Format(" /D \"{0}\"", DefinitionArgument);
+				Arguments.AppendFormat(" /D \"{0}\"", DefinitionArgument);
 			}
 
 			var BuildPlatform = UEBuildPlatform.GetBuildPlatformForCPPTargetPlatform(CompileEnvironment.Config.Target.Platform);
@@ -711,7 +693,7 @@ namespace UnrealBuildTool
 			foreach (FileItem SourceFile in SourceFiles)
 			{
 				Action CompileAction = new Action(ActionType.Compile);
-				string FileArguments = "";
+				StringBuilder FileArguments = new StringBuilder();
 				bool bIsPlainCFile = Path.GetExtension(SourceFile.AbsolutePath).ToUpperInvariant() == ".C";
 
 				// Add the C++ source file and its included files to the prerequisite item list.
@@ -748,7 +730,7 @@ namespace UnrealBuildTool
 					// Make sure the original source directory the PCH header file existed in is added as an include
 					// path -- it might be a private PCH header and we need to make sure that its found!
 					string OriginalPCHHeaderDirectory = Path.GetDirectoryName( SourceFile.AbsolutePath );
-					FileArguments += string.Format(" /I \"{0}\"", OriginalPCHHeaderDirectory);
+					FileArguments.AppendFormat(" /I \"{0}\"", OriginalPCHHeaderDirectory);
 
 					var PrecompiledFileExtension = UEBuildPlatform.BuildPlatformDictionary[UnrealTargetPlatform.Win64].GetBinaryExtension(UEBuildBinaryType.PrecompiledHeader);
 					// Add the precompiled header file to the produced items list.
@@ -765,8 +747,8 @@ namespace UnrealBuildTool
 					if( WindowsPlatform.bCompileWithClang )
 					{
 						// Tell Clang to generate a PCH header
-						FileArguments += " -Xclang -x -Xclang c++-header";	// @todo clang: Doesn't work do to Clang-cl overriding us at the end of the command-line (see -### option output)
-						FileArguments += string.Format(" /Fo\"{0}\"", PrecompiledHeaderFile.AbsolutePath);
+						FileArguments.Append(" -Xclang -x -Xclang c++-header");	// @todo clang: Doesn't work do to Clang-cl overriding us at the end of the command-line (see -### option output)
+						FileArguments.AppendFormat(" /Fo\"{0}\"", PrecompiledHeaderFile.AbsolutePath);
 
 						// Clang PCH generation doesn't create an .obj file to link in, unlike MSVC
 						bEmitsObjectFile = false;
@@ -774,8 +756,8 @@ namespace UnrealBuildTool
 					else
 					{
 						// Add the parameters needed to compile the precompiled header file to the command-line.
-						FileArguments += string.Format(" /Yc\"{0}\"", CompileEnvironment.Config.PrecompiledHeaderIncludeFilename);
-						FileArguments += string.Format(" /Fp\"{0}\"", PrecompiledHeaderFile.AbsolutePath);
+						FileArguments.AppendFormat(" /Yc\"{0}\"", CompileEnvironment.Config.PrecompiledHeaderIncludeFilename);
+						FileArguments.AppendFormat(" /Fp\"{0}\"", PrecompiledHeaderFile.AbsolutePath);
 
 						// If we're creating a PCH that will be used to compile source files for a library, we need
 						// the compiled modules to retain a reference to PCH's module, so that debugging information
@@ -786,11 +768,11 @@ namespace UnrealBuildTool
 							// NOTE: The symbol name we use here is arbitrary, and all that matters is that it is
 							// unique per PCH module used in our library
 							string FakeUniquePCHSymbolName = Path.GetFileNameWithoutExtension(CompileEnvironment.Config.PrecompiledHeaderIncludeFilename);
-							FileArguments += string.Format(" /Yl{0}", FakeUniquePCHSymbolName);
+							FileArguments.AppendFormat(" /Yl{0}", FakeUniquePCHSymbolName);
 						}
 					}
 
-					FileArguments += string.Format(" \"{0}\"", PCHCPPFile.AbsolutePath);
+					FileArguments.AppendFormat(" \"{0}\"", PCHCPPFile.AbsolutePath);
 
 					CompileAction.StatusDescription = PCHCPPFilename;
 				}
@@ -805,12 +787,12 @@ namespace UnrealBuildTool
 						{
 							// NOTE: With Clang, PCH headers are ALWAYS forcibly included!
 							// NOTE: This needs to be before the other include paths to ensure Clang uses it instead of the source header file.
-							FileArguments += string.Format(" /FI\"{0}\"", Path.ChangeExtension(CompileEnvironment.PrecompiledHeaderFile.AbsolutePath, null));
+							FileArguments.AppendFormat(" /FI\"{0}\"", Path.ChangeExtension(CompileEnvironment.PrecompiledHeaderFile.AbsolutePath, null));
 						}
 						else
 						{
-							FileArguments += string.Format(" /Yu\"{0}\"", CompileEnvironment.Config.PCHHeaderNameInCode);
-							FileArguments += string.Format(" /Fp\"{0}\"", CompileEnvironment.PrecompiledHeaderFile.AbsolutePath);
+							FileArguments.AppendFormat(" /Yu\"{0}\"", CompileEnvironment.Config.PCHHeaderNameInCode);
+							FileArguments.AppendFormat(" /Fp\"{0}\"", CompileEnvironment.PrecompiledHeaderFile.AbsolutePath);
 
 							// Is it unsafe to always force inclusion?  Clang is doing it, and .generated.cpp files
 							// won't work otherwise, because they're not located in the context of the module,
@@ -820,13 +802,13 @@ namespace UnrealBuildTool
 								// Force include the precompiled header file.  This is needed because we may have selected a
 								// precompiled header that is different than the first direct include in the C++ source file, but
 								// we still need to make sure that our precompiled header is the first thing included!
-								FileArguments += string.Format( " /FI\"{0}\"", CompileEnvironment.Config.PCHHeaderNameInCode);
+								FileArguments.AppendFormat( " /FI\"{0}\"", CompileEnvironment.Config.PCHHeaderNameInCode);
 							}
 						}
 					}
 					
 					// Add the source file path to the command-line.
-					FileArguments += string.Format(" \"{0}\"", SourceFile.AbsolutePath);
+					FileArguments.AppendFormat(" \"{0}\"", SourceFile.AbsolutePath);
 
 					CompileAction.StatusDescription = Path.GetFileName( SourceFile.AbsolutePath );
 				}
@@ -843,7 +825,7 @@ namespace UnrealBuildTool
 						);
 					CompileAction.ProducedItems.Add(ObjectFile);
 					Result.ObjectFiles.Add(ObjectFile);
-					FileArguments += string.Format(" /Fo\"{0}\"", ObjectFile.AbsolutePath);
+					FileArguments.AppendFormat(" /Fo\"{0}\"", ObjectFile.AbsolutePath);
 				}
 
 				// Create PDB files if we were configured to do that.
@@ -890,7 +872,7 @@ namespace UnrealBuildTool
 									PDBFileName + ".pdb"
 									)
 								);
-						FileArguments += string.Format(" /Fd\"{0}\"", PDBFile.AbsolutePath);
+						FileArguments.AppendFormat(" /Fd\"{0}\"", PDBFile.AbsolutePath);
 
 						// Only use the PDB as an output file if we want PDBs and this particular action is
 						// the one that produces the PDB (as opposed to no debug info, where the above code
@@ -907,11 +889,11 @@ namespace UnrealBuildTool
 				// Add C or C++ specific compiler arguments.
 				if (bIsPlainCFile)
 				{
-					FileArguments += GetCLArguments_C();
+					AppendCLArguments_C(FileArguments);
 				}
 				else
 				{
-					FileArguments += GetCLArguments_CPP( CompileEnvironment );
+					AppendCLArguments_CPP( CompileEnvironment, FileArguments );
 				}
 
 				CompileAction.WorkingDirectory = Path.GetFullPath(".");
@@ -921,7 +903,7 @@ namespace UnrealBuildTool
 				{
 					CompileAction.bIsVCCompiler = true;
 				}
-				CompileAction.CommandArguments = Arguments + FileArguments + CompileEnvironment.Config.AdditionalArguments;
+				CompileAction.CommandArguments = Arguments.ToString() + FileArguments.ToString() + CompileEnvironment.Config.AdditionalArguments;
 				CompileAction.StatusDetailedDescription = SourceFile.Description;
 
 				if( CompileEnvironment.Config.PrecompiledHeaderAction == PrecompiledHeaderAction.Create )
@@ -1054,33 +1036,27 @@ namespace UnrealBuildTool
 			bool bIsBuildingLibrary = LinkEnvironment.Config.bIsBuildingLibrary || bBuildImportLibraryOnly;
 			bool bIncludeDependentLibrariesInLibrary = bIsBuildingLibrary && LinkEnvironment.Config.bIncludeDependentLibrariesInLibrary;
 
-			// Create an action that invokes the linker.
-			Action LinkAction = new Action(ActionType.Link);
-			LinkAction.WorkingDirectory = Path.GetFullPath(".");
-			LinkAction.CommandPath = GetVCToolPath(
-				LinkEnvironment.Config.Target.Platform,
-				LinkEnvironment.Config.Target.Configuration,
-				bIsBuildingLibrary ? "lib" : "link");
-
 			// Get link arguments.
-			LinkAction.CommandArguments = bIsBuildingLibrary ?
-				GetLibArguments(LinkEnvironment) :
-				GetLinkArguments(LinkEnvironment);
-
+			StringBuilder Arguments = new StringBuilder();
+			if(bIsBuildingLibrary)
+			{
+				AppendLibArguments(LinkEnvironment, Arguments);
+			}
+			else
+			{
+				AppendLinkArguments(LinkEnvironment, Arguments);
+			}
 
-			// Tell the action that we're building an import library here and it should conditionally be
-			// ignored as a prerequisite for other actions
-			LinkAction.bProducesImportLibrary = bBuildImportLibraryOnly || LinkEnvironment.Config.bIsBuildingDLL;
 
 
 			// If we're only building an import library, add the '/DEF' option that tells the LIB utility
 			// to simply create a .LIB file and .EXP file, and don't bother validating imports
 			if (bBuildImportLibraryOnly)
 			{
-				LinkAction.CommandArguments += " /DEF";
+				Arguments.Append(" /DEF");
 
 				// Ensure that the import library references the correct filename for the linked binary.
-				LinkAction.CommandArguments += string.Format(" /NAME:\"{0}\"", Path.GetFileName(LinkEnvironment.Config.OutputFilePath));
+				Arguments.AppendFormat(" /NAME:\"{0}\"", Path.GetFileName(LinkEnvironment.Config.OutputFilePath));
 			}
 
 
@@ -1090,7 +1066,7 @@ namespace UnrealBuildTool
 				// Delay-load these DLLs.
 				foreach (string DelayLoadDLL in LinkEnvironment.Config.DelayLoadDLLs)
 				{
-					LinkAction.CommandArguments += string.Format(" /DELAYLOAD:\"{0}\"", DelayLoadDLL);
+					Arguments.AppendFormat(" /DELAYLOAD:\"{0}\"", DelayLoadDLL);
 				}
 			}
 
@@ -1100,13 +1076,13 @@ namespace UnrealBuildTool
 				// Add the library paths to the argument list.
 				foreach (string LibraryPath in LinkEnvironment.Config.LibraryPaths)
 				{
-					LinkAction.CommandArguments += string.Format(" /LIBPATH:\"{0}\"", LibraryPath);
+					Arguments.AppendFormat(" /LIBPATH:\"{0}\"", LibraryPath);
 				}
 
 				// Add the excluded default libraries to the argument list.
 				foreach (string ExcludedLibrary in LinkEnvironment.Config.ExcludedLibraries)
 				{
-					LinkAction.CommandArguments += string.Format(" /NODEFAULTLIB:\"{0}\"", ExcludedLibrary);
+					Arguments.AppendFormat(" /NODEFAULTLIB:\"{0}\"", ExcludedLibrary);
 				}
 			}
 
@@ -1133,15 +1109,18 @@ namespace UnrealBuildTool
 				OutputFile = FileItem.GetItemByPath(LinkEnvironment.Config.OutputFilePath);
 				OutputFile.bNeedsHotReloadNumbersDLLCleanUp = LinkEnvironment.Config.bIsBuildingDLL;
 			}
-			LinkAction.ProducedItems.Add(OutputFile);
-			LinkAction.StatusDescription = Path.GetFileName(OutputFile.AbsolutePath);
+
+			List<FileItem> ProducedItems = new List<FileItem>();
+			ProducedItems.Add(OutputFile);
+
+			List<FileItem> PrerequisiteItems = new List<FileItem>();
 
 			// Add the input files to a response file, and pass the response file on the command-line.
 			List<string> InputFileNames = new List<string>();
 			foreach (FileItem InputFile in LinkEnvironment.InputFiles)
 			{
 				InputFileNames.Add(string.Format("\"{0}\"", InputFile.AbsolutePath));
-				LinkAction.PrerequisiteItems.Add(InputFile);
+				PrerequisiteItems.Add(InputFile);
 			}
 
 			if (!bBuildImportLibraryOnly)
@@ -1150,7 +1129,7 @@ namespace UnrealBuildTool
 				foreach (FileItem InputLibrary in LinkEnvironment.InputLibraries)
 				{
 					InputFileNames.Add(string.Format("\"{0}\"", InputLibrary.AbsolutePath));
-					LinkAction.PrerequisiteItems.Add(InputLibrary);
+					PrerequisiteItems.Add(InputLibrary);
 				}
 			}
 
@@ -1165,7 +1144,7 @@ namespace UnrealBuildTool
 					// when the binary needs to be relinked because a dependent external library has changed.
 					//if( !String.IsNullOrEmpty( Path.GetDirectoryName( AdditionalLibrary ) ) )
 					{
-						LinkAction.PrerequisiteItems.Add(FileItem.GetItemByPath(AdditionalLibrary));
+						PrerequisiteItems.Add(FileItem.GetItemByPath(AdditionalLibrary));
 					}
 				}
 			}
@@ -1178,17 +1157,17 @@ namespace UnrealBuildTool
 			{
 				ResponseFile.Create( ResponseFileName, InputFileNames );
 			}
-			LinkAction.CommandArguments += string.Format(" @\"{0}\"", ResponseFileName);
+			Arguments.AppendFormat(" @\"{0}\"", ResponseFileName);
 
 			// Add the output file to the command-line.
-			LinkAction.CommandArguments += string.Format(" /OUT:\"{0}\"", OutputFile.AbsolutePath);
+			Arguments.AppendFormat(" /OUT:\"{0}\"", OutputFile.AbsolutePath);
 
 			if (bBuildImportLibraryOnly || (LinkEnvironment.Config.bHasExports && !bIsBuildingLibrary))
 			{
 				// An export file is written to the output directory implicitly; add it to the produced items list.
 				string ExportFilePath = Path.ChangeExtension(ImportLibraryFilePath, ".exp");
 				FileItem ExportFile = FileItem.GetItemByPath(ExportFilePath);
-				LinkAction.ProducedItems.Add(ExportFile);
+				ProducedItems.Add(ExportFile);
 			}
 
 			if (!bIsBuildingLibrary)
@@ -1199,8 +1178,8 @@ namespace UnrealBuildTool
 				{
 					// Write the import library to the output directory for nFringe support.
 					FileItem ImportLibraryFile = FileItem.GetItemByPath(ImportLibraryFilePath);
-					LinkAction.CommandArguments += string.Format(" /IMPLIB:\"{0}\"", ImportLibraryFilePath);
-					LinkAction.ProducedItems.Add(ImportLibraryFile);
+					Arguments.AppendFormat(" /IMPLIB:\"{0}\"", ImportLibraryFilePath);
+					ProducedItems.Add(ImportLibraryFile);
 				}
 
 				if (LinkEnvironment.Config.bCreateDebugInfo)
@@ -1211,8 +1190,8 @@ namespace UnrealBuildTool
 					{
 						string PDBFilePath = Path.Combine(LinkEnvironment.Config.OutputDirectory, Path.GetFileNameWithoutExtension(OutputFile.AbsolutePath) + ".pdb");
 						FileItem PDBFile = FileItem.GetItemByPath(PDBFilePath);
-						LinkAction.CommandArguments += string.Format(" /PDB:\"{0}\"", PDBFilePath);
-						LinkAction.ProducedItems.Add(PDBFile);
+						Arguments.AppendFormat(" /PDB:\"{0}\"", PDBFilePath);
+						ProducedItems.Add(PDBFile);
 					}
 
 					// Write the MAP file to the output directory.			
@@ -1228,15 +1207,32 @@ namespace UnrealBuildTool
 				}
 
 				// Add the additional arguments specified by the environment.
-				LinkAction.CommandArguments += LinkEnvironment.Config.AdditionalArguments;
+				Arguments.Append(LinkEnvironment.Config.AdditionalArguments);
 			}
 
-			Log.TraceVerbose( "     Linking: " + LinkAction.StatusDescription );
-			Log.TraceVerbose( "     Command: " + LinkAction.CommandArguments );
+
+			// Create an action that invokes the linker.
+			Action LinkAction = new Action(ActionType.Link);
+			LinkAction.WorkingDirectory = Path.GetFullPath(".");
+			LinkAction.CommandPath = GetVCToolPath(
+				LinkEnvironment.Config.Target.Platform,
+				LinkEnvironment.Config.Target.Configuration,
+				bIsBuildingLibrary ? "lib" : "link");
+			LinkAction.CommandArguments = Arguments.ToString();
+			LinkAction.ProducedItems.AddRange(ProducedItems);
+			LinkAction.PrerequisiteItems.AddRange(PrerequisiteItems);
+			LinkAction.StatusDescription = Path.GetFileName(OutputFile.AbsolutePath);
+
+			// Tell the action that we're building an import library here and it should conditionally be
+			// ignored as a prerequisite for other actions
+			LinkAction.bProducesImportLibrary = bBuildImportLibraryOnly || LinkEnvironment.Config.bIsBuildingDLL;
 
 			// Only execute linking on the local PC.
 			LinkAction.bCanExecuteRemotely = false;
 
+			Log.TraceVerbose( "     Linking: " + LinkAction.StatusDescription );
+			Log.TraceVerbose( "     Command: " + LinkAction.CommandArguments );
+
 			return OutputFile;
 		}
 
