remove HTML5 target platform dll from being built as part of the mac and linux editor builds for rocket
TTP342631
#ue4
#html5

diff --git a/Engine/Source/Programs/UnrealBuildTool/HTML5/UEBuildHTML5.cs b/Engine/Source/Programs/UnrealBuildTool/HTML5/UEBuildHTML5.cs
--- a/Engine/Source/Programs/UnrealBuildTool/HTML5/UEBuildHTML5.cs
+++ b/Engine/Source/Programs/UnrealBuildTool/HTML5/UEBuildHTML5.cs
@@ -367,9 +367,14 @@ namespace UnrealBuildTool
                     }
                 }
             }
-            else if (Target.Platform == UnrealTargetPlatform.Win64 || Target.Platform == UnrealTargetPlatform.Mac )
+            else if (Target.Platform == UnrealTargetPlatform.Win32 || Target.Platform == UnrealTargetPlatform.Win64 || Target.Platform == UnrealTargetPlatform.Mac )
             {
-                if ( (!UEBuildConfiguration.bBuildRequiresCookedData
+				// allow standalone tools to use targetplatform modules, without needing Engine
+				if (UEBuildConfiguration.bForceBuildTargetPlatforms)
+				{
+					InModule.AddPlatformSpecificDynamicallyLoadedModule("IOSTargetPlatform");
+				}
+				if ((!UEBuildConfiguration.bBuildRequiresCookedData
                     && InModule.ToString() == "Engine"
                     && UEBuildConfiguration.bBuildDeveloperTools)
                     || UEBuildConfiguration.bForceBuildTargetPlatforms)
diff --git a/Engine/Source/Programs/UnrealBuildTool/Linux/UEBuildLinux.cs b/Engine/Source/Programs/UnrealBuildTool/Linux/UEBuildLinux.cs
--- a/Engine/Source/Programs/UnrealBuildTool/Linux/UEBuildLinux.cs
+++ b/Engine/Source/Programs/UnrealBuildTool/Linux/UEBuildLinux.cs
@@ -319,7 +319,6 @@ namespace UnrealBuildTool
                 // allow standalone tools to use target platform modules, without needing Engine
                 if (UEBuildConfiguration.bForceBuildTargetPlatforms)
                 {
-                    InModule.AddPlatformSpecificDynamicallyLoadedModule("HTML5TargetPlatform");
                     InModule.AddDynamicallyLoadedModule("LinuxTargetPlatform");
                     InModule.AddDynamicallyLoadedModule("LinuxNoEditorTargetPlatform");
                     InModule.AddDynamicallyLoadedModule("LinuxServerTargetPlatform");
diff --git a/Engine/Source/Programs/UnrealBuildTool/Mac/UEBuildMac.cs b/Engine/Source/Programs/UnrealBuildTool/Mac/UEBuildMac.cs
--- a/Engine/Source/Programs/UnrealBuildTool/Mac/UEBuildMac.cs
+++ b/Engine/Source/Programs/UnrealBuildTool/Mac/UEBuildMac.cs
@@ -109,7 +109,6 @@ namespace UnrealBuildTool
 				// allow standalone tools to use target platform modules, without needing Engine
 				if (UEBuildConfiguration.bForceBuildTargetPlatforms)
 				{
-					InModule.AddPlatformSpecificDynamicallyLoadedModule("HTML5TargetPlatform");
 					InModule.AddDynamicallyLoadedModule("MacTargetPlatform");
 					InModule.AddDynamicallyLoadedModule("MacNoEditorTargetPlatform");
 					InModule.AddDynamicallyLoadedModule("MacClientTargetPlatform");
