accidentally left some code in I meant to remove
#ue4
#html5

diff --git a/Engine/Source/Programs/UnrealBuildTool/HTML5/UEBuildHTML5.cs b/Engine/Source/Programs/UnrealBuildTool/HTML5/UEBuildHTML5.cs
--- a/Engine/Source/Programs/UnrealBuildTool/HTML5/UEBuildHTML5.cs
+++ b/Engine/Source/Programs/UnrealBuildTool/HTML5/UEBuildHTML5.cs
@@ -370,10 +370,6 @@ namespace UnrealBuildTool
             else if (Target.Platform == UnrealTargetPlatform.Win32 || Target.Platform == UnrealTargetPlatform.Win64 || Target.Platform == UnrealTargetPlatform.Mac )
             {
 				// allow standalone tools to use targetplatform modules, without needing Engine
-				if (UEBuildConfiguration.bForceBuildTargetPlatforms)
-				{
-					InModule.AddPlatformSpecificDynamicallyLoadedModule("IOSTargetPlatform");
-				}
 				if ((!UEBuildConfiguration.bBuildRequiresCookedData
                     && InModule.ToString() == "Engine"
                     && UEBuildConfiguration.bBuildDeveloperTools)
