fix for Win64 editor not building for Rocket because HTML5 targetplatform is being added when it shouldn't
#ue4
#html5

diff --git a/Engine/Source/Programs/UnrealBuildTool/HTML5/UEBuildHTML5.cs b/Engine/Source/Programs/UnrealBuildTool/HTML5/UEBuildHTML5.cs
--- a/Engine/Source/Programs/UnrealBuildTool/HTML5/UEBuildHTML5.cs
+++ b/Engine/Source/Programs/UnrealBuildTool/HTML5/UEBuildHTML5.cs
@@ -336,7 +336,13 @@ namespace UnrealBuildTool
          */
         public override void ModifyNewlyLoadedModule(UEBuildModule InModule, TargetInfo Target)
         {
-            if (Target.Platform == UnrealTargetPlatform.HTML5)
+			[REDACTED_EMAIL]: Add platform support
+			if (UnrealBuildTool.RunningRocket() || UnrealBuildTool.BuildingRocket())
+			{
+				return;
+			}
+
+			if (Target.Platform == UnrealTargetPlatform.HTML5)
             {
                 if (InModule.ToString() == "Core")
                 {
