Make sure that last write time for a resource copied to the Mac app bundle is updated after copy so UBT doesn't think it's out of date.

diff --git a/Engine/Source/Programs/UnrealBuildTool/Mac/MacToolChain.cs b/Engine/Source/Programs/UnrealBuildTool/Mac/MacToolChain.cs
--- a/Engine/Source/Programs/UnrealBuildTool/Mac/MacToolChain.cs
+++ b/Engine/Source/Programs/UnrealBuildTool/Mac/MacToolChain.cs
@@ -1070,7 +1070,7 @@ namespace UnrealBuildTool
 
 			FileItem TargetItem = LocalToRemoteFileItem(FileItem.GetItemByPath(TargetPath), false);
 
-			CopyAction.CommandArguments = string.Format("-c 'cp -f -R \"{0}\" \"{1}\"'", SourcePath, Path.GetDirectoryName(TargetPath) + "/");
+			CopyAction.CommandArguments = string.Format("-c 'cp -f -R \"{0}\" \"{1}\"; touch -c \"{1}\"'", SourcePath, Path.GetDirectoryName(TargetPath) + "/");
 			CopyAction.PrerequisiteItems.Add(Executable);
 			CopyAction.ProducedItems.Add(TargetItem);
 			CopyAction.StatusDescription = string.Format("Copying {0} to app bundle", Path.GetFileName(Resource.ResourcePath));
