UBT: preprocessor define for cross/native compilation.

- Applies to Linux only for now.

diff --git a/Engine/Source/Programs/UnrealBuildTool/Linux/UEBuildLinux.cs b/Engine/Source/Programs/UnrealBuildTool/Linux/UEBuildLinux.cs
--- a/Engine/Source/Programs/UnrealBuildTool/Linux/UEBuildLinux.cs
+++ b/Engine/Source/Programs/UnrealBuildTool/Linux/UEBuildLinux.cs
@@ -340,6 +340,7 @@ namespace UnrealBuildTool
         {
             InBuildTarget.GlobalCompileEnvironment.Config.Definitions.Add("PLATFORM_LINUX=1");
             InBuildTarget.GlobalCompileEnvironment.Config.Definitions.Add("LINUX=1");
+            InBuildTarget.GlobalCompileEnvironment.Config.Definitions.Add(String.Format("CROSS_COMPILATION={0}", ExternalExecution.GetRuntimePlatform() != UnrealTargetPlatform.Linux ? 1 : 0));
 
             InBuildTarget.GlobalCompileEnvironment.Config.Definitions.Add("WITH_DATABASE_SUPPORT=0");		//@todo linux: valid?
 
