Added code to handle COM crashes when detecting physical processor count in UBT

diff --git a/Engine/Source/Programs/UnrealBuildTool/System/LocalExecutor.cs b/Engine/Source/Programs/UnrealBuildTool/System/LocalExecutor.cs
--- a/Engine/Source/Programs/UnrealBuildTool/System/LocalExecutor.cs
+++ b/Engine/Source/Programs/UnrealBuildTool/System/LocalExecutor.cs
@@ -339,9 +339,21 @@ namespace UnrealBuildTool
 			int NumCores = 0;
 			if (!Utils.IsRunningOnMono)
 			{
-				foreach(var Item in new System.Management.ManagementObjectSearcher("Select * from Win32_Processor").Get())
+				try
+				{
+					using (var Mos = new System.Management.ManagementObjectSearcher("Select * from Win32_Processor"))
+					{
+						var MosCollection = Mos.Get();
+						foreach (var Item in MosCollection)
+						{
+							NumCores += int.Parse(Item["NumberOfCores"].ToString());
+						}
+					}
+				}
+				catch (Exception Ex)
 				{
-					NumCores += int.Parse(Item["NumberOfCores"].ToString());
+					Log.TraceWarning("Unable to get the number of Cores: {0}", Ex.ToString());
+					Log.TraceWarning("Falling back to processor count.");
 				}
 			}
 			// On some systems this requires a hot fix to work so we fall back to using the (logical) processor count.
