Proper fix for out of date check for Mac bundle resources

diff --git a/Engine/Source/Programs/UnrealBuildTool/Mac/MacToolChain.cs b/Engine/Source/Programs/UnrealBuildTool/Mac/MacToolChain.cs
--- a/Engine/Source/Programs/UnrealBuildTool/Mac/MacToolChain.cs
+++ b/Engine/Source/Programs/UnrealBuildTool/Mac/MacToolChain.cs
@@ -1070,7 +1070,7 @@ namespace UnrealBuildTool
 
 			FileItem TargetItem = LocalToRemoteFileItem(FileItem.GetItemByPath(TargetPath), false);
 
-			CopyAction.CommandArguments = string.Format("-c 'cp -f -R \"{0}\" \"{1}\"; touch -c \"{1}\"'", SourcePath, Path.GetDirectoryName(TargetPath) + "/");
+			CopyAction.CommandArguments = string.Format("-c 'cp -f -R \"{0}\" \"{1}\"; touch -c \"{2}\"'", SourcePath, Path.GetDirectoryName(TargetPath) + "/", TargetPath);
 			CopyAction.PrerequisiteItems.Add(Executable);
 			CopyAction.ProducedItems.Add(TargetItem);
 			CopyAction.StatusDescription = string.Format("Copying {0} to app bundle", Path.GetFileName(Resource.ResourcePath));
