UBT: Moved ScopedTimer from UAT to UBT, added ScopedCounter

diff --git a/Engine/Source/Programs/AutomationTool/Utils.cs b/Engine/Source/Programs/AutomationTool/Utils.cs
--- a/Engine/Source/Programs/AutomationTool/Utils.cs
+++ b/Engine/Source/Programs/AutomationTool/Utils.cs
@@ -996,39 +996,4 @@ namespace AutomationTool
 	}
 
 	#endregion
-
-	#region Scoped Timer
-
-	public class ScopedTimer : IDisposable
-	{
-		DateTime StartTime;
-		string TimerName;
-		static int Indent = 0;
-		static object IndentLock = new object();
-
-		public ScopedTimer(string Name)
-		{
-			TimerName = Name;
-			lock (IndentLock)
-			{
-				Indent++;
-			}
-			StartTime = DateTime.UtcNow;
-		}
-
-		public void Dispose()
-		{
-			var TotalSeconds = (DateTime.UtcNow - StartTime).TotalSeconds;
-			var LogIndent = 0;
-			lock (IndentLock)
-			{
-				LogIndent = --Indent;
-			}
-			var IndentText = new StringBuilder(LogIndent * 2);
-			IndentText.Append(' ', LogIndent * 2);
-			Log.TraceVerbose("{0}{1} took {2}s", IndentText.ToString(), TimerName, TotalSeconds);
-		}
-	}
-
-	#endregion
 }
diff --git a/Engine/Source/Programs/UnrealBuildTool/System/Utils.cs b/Engine/Source/Programs/UnrealBuildTool/System/Utils.cs
--- a/Engine/Source/Programs/UnrealBuildTool/System/Utils.cs
+++ b/Engine/Source/Programs/UnrealBuildTool/System/Utils.cs
@@ -1227,6 +1227,24 @@ namespace UnrealBuildTool
 		{
 			WriteLine(1, TraceEventType.Error, Message);
 		}
+
+		/// <summary>
+		/// Writes a message with the specified verbosity to the console.
+		/// </summary>
+		[MethodImplAttribute(MethodImplOptions.NoInlining)]
+		public static void TraceEvent(TraceEventType Verbosity, string Message)
+		{
+			WriteLine(1, Verbosity, Message);
+		}
+
+		/// <summary>
+		/// Writes a formatted message with the specified verbosity to the console.
+		/// </summary>
+		[MethodImplAttribute(MethodImplOptions.NoInlining)]
+		public static void TraceEvent(TraceEventType Verbosity, string Format, params object[] Args)
+		{
+			WriteLine(1, Verbosity, Format, Args);
+		}
 	}
 
     #region StreamUtils
@@ -1416,4 +1434,126 @@ namespace UnrealBuildTool
         }
     }
     #endregion
+
+	#region Scoped Timers
+
+	/// <summary>
+	/// Scoped timer, start is in the constructor, end in Dispose. Best used with using(var Timer = new ScopedTimer()). Suports nesting.
+	/// </summary>
+	public class ScopedTimer : IDisposable
+	{
+		DateTime StartTime;
+		string TimerName;
+		TraceEventType Verbosity;
+		static int Indent = 0;
+		static object IndentLock = new object();
+
+		public ScopedTimer(string Name, TraceEventType InVerbosity = TraceEventType.Verbose)
+		{
+			TimerName = Name;
+			lock (IndentLock)
+			{
+				Indent++;
+			}
+			Verbosity = InVerbosity;
+			StartTime = DateTime.UtcNow;
+		}
+
+		public void Dispose()
+		{
+			var TotalSeconds = (DateTime.UtcNow - StartTime).TotalSeconds;
+			var LogIndent = 0;
+			lock (IndentLock)
+			{
+				LogIndent = --Indent;
+			}
+			var IndentText = new StringBuilder(LogIndent * 2);
+			IndentText.Append(' ', LogIndent * 2);
+			
+			Log.TraceEvent(Verbosity, "{0}{1} took {2}s", IndentText.ToString(), TimerName, TotalSeconds);
+		}
+	}
+
+	/// <summary>
+	/// Scoped, accumulative timer. Can have multiple instances. Best used as a static variable with Start() and End() 
+	/// </summary>
+	public class ScopedCounter : IDisposable
+	{
+		class Accumulator
+		{
+			public double Time;
+			public TraceEventType Verbosity;
+		}
+		DateTime StartTime;
+		Accumulator ScopedAccumulator;
+		string CounterName;
+		
+		static Dictionary<string, Accumulator> Accumulators = new Dictionary<string, Accumulator>();
+
+		public ScopedCounter(string Name, TraceEventType InVerbosity = TraceEventType.Verbose)
+		{
+			CounterName = Name;
+			if (!Accumulators.TryGetValue(CounterName, out ScopedAccumulator))
+			{
+				ScopedAccumulator = new Accumulator();
+				Accumulators.Add(CounterName, ScopedAccumulator);
+			}
+			ScopedAccumulator.Verbosity = InVerbosity;
+			StartTime = DateTime.UtcNow;
+		}
+
+		/// <summary>
+		/// (Re-)Starts time measuring
+		/// </summary>
+		/// <returns></returns>
+		public ScopedCounter Start()
+		{
+			StartTime = DateTime.UtcNow;
+			return this;
+		}
+
+		/// <summary>
+		/// Ends measuring time
+		/// </summary>
+		public void Dispose()
+		{
+			var TotalSeconds = (DateTime.UtcNow - StartTime).TotalSeconds;
+			ScopedAccumulator.Time += TotalSeconds;
+		}
+
+		/// <summary>
+		/// Logs the specified counters timings to the output
+		/// </summary>
+		/// <param name="Name">Timer name</param>
+		static public void LogCounter(string Name)
+		{
+			Accumulator ExistingAccumulator;
+			if (Accumulators.TryGetValue(Name, out ExistingAccumulator))
+			{
+				LogAccumulator(Name, ExistingAccumulator);
+			}
+			else
+			{
+				throw new Exception(String.Format("ScopedCounter {0} does not exist!", Name));
+			}
+		}
+
+		static private void LogAccumulator(string Name, Accumulator Counter)
+		{
+			Log.TraceEvent(Counter.Verbosity, "{0} took {1}s", Name, Counter.Time);
+		}
+
+		/// <summary>
+		/// Dumps all counters to the log
+		/// </summary>
+		static public void DumpAllCounters()
+		{
+			foreach (var Counter in Accumulators)
+			{
+				LogAccumulator(Counter.Key, Counter.Value);
+			}
+		}
+	}
+
+	#endregion
 }
