UnrealBuildTool: Added UnrealHeaderTool version checking when checking if it's up to date.

#change UBT will check for API version of all UnrealHeaderTool binaries to detect partial syncs
#change Added BuildHostPlatform class for runtime platform abstraction

diff --git a/Engine/Source/Programs/AutomationTool/CommandUtils.cs b/Engine/Source/Programs/AutomationTool/CommandUtils.cs
--- a/Engine/Source/Programs/AutomationTool/CommandUtils.cs
+++ b/Engine/Source/Programs/AutomationTool/CommandUtils.cs
@@ -2346,7 +2346,7 @@ namespace AutomationTool
 			if (!Command.ParseParam("NoSign"))
 			{
 				CommandUtils.Log("Signing up to {0} files...", Files.Count);
-				UnrealBuildTool.UnrealTargetPlatform TargetPlatform = UnrealBuildTool.ExternalExecution.GetRuntimePlatform();
+				UnrealBuildTool.UnrealTargetPlatform TargetPlatform = UnrealBuildTool.BuildHostPlatform.Current.Platform;
 				if (TargetPlatform == UnrealBuildTool.UnrealTargetPlatform.Mac)
 				{
 					foreach (var File in Files)
diff --git a/Engine/Source/Programs/AutomationTool/HTML5/HTML5Platform.Automation.cs b/Engine/Source/Programs/AutomationTool/HTML5/HTML5Platform.Automation.cs
--- a/Engine/Source/Programs/AutomationTool/HTML5/HTML5Platform.Automation.cs
+++ b/Engine/Source/Programs/AutomationTool/HTML5/HTML5Platform.Automation.cs
@@ -264,12 +264,12 @@ public class HTML5Platform : Platform
             HTMLPath += "?cookonthefly=true";
         }
 
-        if (ExternalExecution.GetRuntimePlatform() == UnrealTargetPlatform.Mac)
+		if (BuildHostPlatform.Current.Platform == UnrealTargetPlatform.Mac)
         {
             FirefoxPath = "/bin/bash";
 			FirefoxProfileCommand = " -c  ' open -a Firefox.app  \"" + HTMLPath + "\"   ' ";
         }
-        else if (ExternalExecution.GetRuntimePlatform() == UnrealTargetPlatform.Win64)
+		else if (BuildHostPlatform.Current.Platform == UnrealTargetPlatform.Win64)
         {
             FirefoxProfileCommand = FirefoxProfileCommand + "\"" + HTMLPath + "\" ";
         }
diff --git a/Engine/Source/Programs/AutomationTool/HostPlatform.cs b/Engine/Source/Programs/AutomationTool/HostPlatform.cs
--- a/Engine/Source/Programs/AutomationTool/HostPlatform.cs
+++ b/Engine/Source/Programs/AutomationTool/HostPlatform.cs
@@ -27,7 +27,7 @@ namespace AutomationTool
 		/// </summary>
 		internal static void Initialize()
 		{
-			if (UnrealBuildTool.Utils.IsRunningOnMono)
+			if (UnrealBuildTool.BuildHostPlatform.Current.Platform == UnrealBuildTool.UnrealTargetPlatform.Mac)
 			{
 				RunningPlatform = new MacHostPlatform();
 			}
diff --git a/Engine/Source/Programs/AutomationTool/IOS/IOSPlatform.Automation.cs b/Engine/Source/Programs/AutomationTool/IOS/IOSPlatform.Automation.cs
--- a/Engine/Source/Programs/AutomationTool/IOS/IOSPlatform.Automation.cs
+++ b/Engine/Source/Programs/AutomationTool/IOS/IOSPlatform.Automation.cs
@@ -133,7 +133,7 @@ public class IOSPlatform : Platform
 		//   - IPP needs to be able to take a .app directory instead of a Payload directory when doing RepackageFromStage (which would probably be renamed)
 		//   - Some discrepancy in the loading screen pngs that are getting packaged, which needs to be investigated
 		//   - Code here probably needs to be updated to write 0 byte files as 1 byte (difference with IPP, was required at one point when using Ionic.Zip to prevent issues on device, maybe not needed anymore?)
-		if (UnrealBuildTool.ExternalExecution.GetRuntimePlatform() == UnrealTargetPlatform.Mac)
+		if (UnrealBuildTool.BuildHostPlatform.Current.Platform == UnrealTargetPlatform.Mac)
 		{
 			// copy in all of the artwork and plist
 			var DeployHandler = UEBuildDeploy.GetBuildDeploy(UnrealTargetPlatform.IOS);
@@ -161,9 +161,9 @@ public class IOSPlatform : Platform
 				InternalUtils.SafeCopyFile(Filename, DestFilename, true);
 			}
 		}
-		
 
-		if (UnrealBuildTool.ExternalExecution.GetRuntimePlatform () != UnrealTargetPlatform.Mac)
+
+		if (UnrealBuildTool.BuildHostPlatform.Current.Platform != UnrealTargetPlatform.Mac)
 		{
 			if (SC.StageTargetConfigurations.Count != 1) 
 			{
@@ -672,7 +672,7 @@ public class IOSPlatform : Platform
 
 	public override void GetFilesToDeployOrStage(ProjectParams Params, DeploymentContext SC)
 	{
-		if (UnrealBuildTool.ExternalExecution.GetRuntimePlatform () != UnrealTargetPlatform.Mac)
+		if (UnrealBuildTool.BuildHostPlatform.Current.Platform != UnrealTargetPlatform.Mac)
 		{
 			// copy the icons/launch screens from the engine
 			{
@@ -749,7 +749,7 @@ public class IOSPlatform : Platform
 		var TargetConfiguration = SC.StageTargetConfigurations[0];
 
 		string ProjectGameExeFilename = Params.ProjectGameExeFilename;
-		if (UnrealBuildTool.ExternalExecution.GetRuntimePlatform () == UnrealTargetPlatform.Mac)
+		if (UnrealBuildTool.BuildHostPlatform.Current.Platform == UnrealTargetPlatform.Mac)
 		{
 			ProjectGameExeFilename = CombinePaths (Path.GetDirectoryName(Params.RawProjectPath), "Binaries", "IOS", Path.GetFileName (Params.ProjectGameExeFilename));
 		}
@@ -775,7 +775,7 @@ public class IOSPlatform : Platform
 
 	public override void Deploy(ProjectParams Params, DeploymentContext SC)
     {
-		if (UnrealBuildTool.ExternalExecution.GetRuntimePlatform () != UnrealTargetPlatform.Mac)
+		if (UnrealBuildTool.BuildHostPlatform.Current.Platform != UnrealTargetPlatform.Mac)
 		{
 			if (SC.StageTargetConfigurations.Count != 1)
 			{
@@ -838,7 +838,7 @@ public class IOSPlatform : Platform
 
 	public override bool IsSupported { get { return true; } }
 
-	public override bool LaunchViaUFE { get { return UnrealBuildTool.ExternalExecution.GetRuntimePlatform () != UnrealTargetPlatform.Mac; } }
+	public override bool LaunchViaUFE { get { return UnrealBuildTool.BuildHostPlatform.Current.Platform != UnrealTargetPlatform.Mac; } }
 
 	public override string Remap(string Dest)
 	{
@@ -851,7 +851,7 @@ public class IOSPlatform : Platform
 
 	public override ProcessResult RunClient(ERunOptions ClientRunFlags, string ClientApp, string ClientCmdLine, ProjectParams Params)
 	{
-		if (UnrealBuildTool.ExternalExecution.GetRuntimePlatform () == UnrealTargetPlatform.Mac)
+		if (UnrealBuildTool.BuildHostPlatform.Current.Platform == UnrealTargetPlatform.Mac)
 		{
 			string AppDirectory = string.Format("{0}/Payload/{1}.app",
 				Path.GetDirectoryName(Params.ProjectGameExeFilename), 
@@ -916,7 +916,7 @@ public class IOSPlatform : Platform
 
 	public override void PreBuildAgenda(UE4Build Build, UE4Build.BuildAgenda Agenda)
 	{
-		if (UnrealBuildTool.ExternalExecution.GetRuntimePlatform () != UnrealTargetPlatform.Mac)
+		if (UnrealBuildTool.BuildHostPlatform.Current.Platform != UnrealTargetPlatform.Mac)
 		{
 			Agenda.DotNetProjects.Add (@"Engine\Source\Programs\IOS\iPhonePackager\iPhonePackager.csproj");
 		}
diff --git a/Engine/Source/Programs/AutomationTool/Mac/MacPlatform.Automation.cs b/Engine/Source/Programs/AutomationTool/Mac/MacPlatform.Automation.cs
--- a/Engine/Source/Programs/AutomationTool/Mac/MacPlatform.Automation.cs
+++ b/Engine/Source/Programs/AutomationTool/Mac/MacPlatform.Automation.cs
@@ -146,7 +146,7 @@ public class MacPlatform : Platform
     public override bool ShouldUseManifestForUBTBuilds(string AddArgs)
 	{
 		// don't use the manifest to set up build products if we are compiling Mac under Windows and we aren't going to copy anything back to the PC
-		bool bIsBuildingRemotely = UnrealBuildTool.ExternalExecution.GetRuntimePlatform() != UnrealTargetPlatform.Mac;
+		bool bIsBuildingRemotely = UnrealBuildTool.BuildHostPlatform.Current.Platform != UnrealTargetPlatform.Mac;
         bool bUseManifest = !bIsBuildingRemotely || AddArgs.IndexOf("-CopyAppBundleBackToDevice", StringComparison.InvariantCultureIgnoreCase) > 0;
 		return bUseManifest;
 	}
diff --git a/Engine/Source/Programs/AutomationTool/Scripts/AnalyzeThirdPartyLibs.Automation.cs b/Engine/Source/Programs/AutomationTool/Scripts/AnalyzeThirdPartyLibs.Automation.cs
--- a/Engine/Source/Programs/AutomationTool/Scripts/AnalyzeThirdPartyLibs.Automation.cs
+++ b/Engine/Source/Programs/AutomationTool/Scripts/AnalyzeThirdPartyLibs.Automation.cs
@@ -130,7 +130,7 @@ class AnalyzeThirdPartyLibs : BuildCommand
 		Log("************************* Analyze Third Party Libs");
 
 		// figure out what batch/script to run
-		switch (UnrealBuildTool.ExternalExecution.GetRuntimePlatform())
+		switch (UnrealBuildTool.BuildHostPlatform.Current.Platform)
 		{
 			case UnrealTargetPlatform.Win64:
 			case UnrealTargetPlatform.Mac:
diff --git a/Engine/Source/Programs/AutomationTool/Scripts/BuildThirdPartyLibs.Automation.cs b/Engine/Source/Programs/AutomationTool/Scripts/BuildThirdPartyLibs.Automation.cs
--- a/Engine/Source/Programs/AutomationTool/Scripts/BuildThirdPartyLibs.Automation.cs
+++ b/Engine/Source/Programs/AutomationTool/Scripts/BuildThirdPartyLibs.Automation.cs
@@ -67,7 +67,7 @@ class BuildThirdPartyLibs : BuildCommand
 
 		// figure out what batch/script to run
 		string CompileScriptName;
-		switch (UnrealBuildTool.ExternalExecution.GetRuntimePlatform())
+		switch (UnrealBuildTool.BuildHostPlatform.Current.Platform)
 		{
 			case UnrealTargetPlatform.Win64:
 				CompileScriptName = WindowsCompileScript;
diff --git a/Engine/Source/Programs/AutomationTool/Scripts/CopyBuildToStagingDirectory.Automation.cs b/Engine/Source/Programs/AutomationTool/Scripts/CopyBuildToStagingDirectory.Automation.cs
--- a/Engine/Source/Programs/AutomationTool/Scripts/CopyBuildToStagingDirectory.Automation.cs
+++ b/Engine/Source/Programs/AutomationTool/Scripts/CopyBuildToStagingDirectory.Automation.cs
@@ -630,7 +630,7 @@ public partial class Project : CommandUtils
 			{
 				FileHostParams += "-filehostip=";
 				bool FirstParam = true;
-				if (UnrealBuildTool.ExternalExecution.GetRuntimePlatform() == UnrealTargetPlatform.Mac)
+				if (UnrealBuildTool.BuildHostPlatform.Current.Platform == UnrealTargetPlatform.Mac)
 				{
 					NetworkInterface[] Interfaces = NetworkInterface.GetAllNetworkInterfaces();
 					foreach (NetworkInterface adapter in Interfaces)
diff --git a/Engine/Source/Programs/AutomationTool/Scripts/RunProjectCommand.Automation.cs b/Engine/Source/Programs/AutomationTool/Scripts/RunProjectCommand.Automation.cs
--- a/Engine/Source/Programs/AutomationTool/Scripts/RunProjectCommand.Automation.cs
+++ b/Engine/Source/Programs/AutomationTool/Scripts/RunProjectCommand.Automation.cs
@@ -590,7 +590,7 @@ public partial class Project : CommandUtils
 			{
 				TempCmdLine += "-filehostip=";
                 bool FirstParam = true;
-				if (UnrealBuildTool.ExternalExecution.GetRuntimePlatform () == UnrealTargetPlatform.Mac)
+				if (UnrealBuildTool.BuildHostPlatform.Current.Platform == UnrealTargetPlatform.Mac)
 				{
 					NetworkInterface[] Interfaces = NetworkInterface.GetAllNetworkInterfaces ();
 					foreach (NetworkInterface adapter in Interfaces)
diff --git a/Engine/Source/Programs/AutomationTool/UBTUtils.cs b/Engine/Source/Programs/AutomationTool/UBTUtils.cs
--- a/Engine/Source/Programs/AutomationTool/UBTUtils.cs
+++ b/Engine/Source/Programs/AutomationTool/UBTUtils.cs
@@ -26,7 +26,7 @@ namespace AutomationTool
 			{
 				CommandLine += " -rocket";
 			}
-			if (!IsBuildMachine && UnrealBuildTool.ExternalExecution.GetRuntimePlatform() == UnrealBuildTool.UnrealTargetPlatform.Mac)
+			if (!IsBuildMachine && UnrealBuildTool.BuildHostPlatform.Current.Platform == UnrealBuildTool.UnrealTargetPlatform.Mac)
 			{
 				CommandLine += " -nocreatestub";
 			}
diff --git a/Engine/Source/Programs/UnrealBuildTool/Configuration/BuildConfiguration.cs b/Engine/Source/Programs/UnrealBuildTool/Configuration/BuildConfiguration.cs
--- a/Engine/Source/Programs/UnrealBuildTool/Configuration/BuildConfiguration.cs
+++ b/Engine/Source/Programs/UnrealBuildTool/Configuration/BuildConfiguration.cs
@@ -493,7 +493,7 @@ namespace UnrealBuildTool
 			DistccExecutablesPath = "/usr/local/bin";
 
 			// The default for normal Mac users should be to use DistCode which installs as an Xcode plugin and provides dynamic host management
-			if (ExternalExecution.GetRuntimePlatform() == UnrealTargetPlatform.Mac)
+			if (BuildHostPlatform.Current.Platform == UnrealTargetPlatform.Mac)
 			{
 				string UserDir = Environment.GetEnvironmentVariable("HOME");
 				string MacDistccExecutablesPath = UserDir + "/Library/Application Support/Developer/Shared/Xcode/Plug-ins/Distcc 3.2.xcplugin/Contents/usr/bin";
diff --git a/Engine/Source/Programs/UnrealBuildTool/Configuration/UEBuildPlatform.cs b/Engine/Source/Programs/UnrealBuildTool/Configuration/UEBuildPlatform.cs
--- a/Engine/Source/Programs/UnrealBuildTool/Configuration/UEBuildPlatform.cs
+++ b/Engine/Source/Programs/UnrealBuildTool/Configuration/UEBuildPlatform.cs
@@ -404,7 +404,7 @@ namespace UnrealBuildTool
          */
         public String GetPathVarDelimiter()
         {
-            switch (ExternalExecution.GetRuntimePlatform())
+            switch (BuildHostPlatform.Current.Platform)
             {
                 case UnrealTargetPlatform.Linux:
                 case UnrealTargetPlatform.Mac:
@@ -413,7 +413,7 @@ namespace UnrealBuildTool
                 case UnrealTargetPlatform.Win64:
                     return ";";
                 default:
-                    Log.TraceWarning("PATH var delimiter unknown for platform " + ExternalExecution.GetRuntimePlatform().ToString() + " using ';'");
+                    Log.TraceWarning("PATH var delimiter unknown for platform " + BuildHostPlatform.Current.Platform.ToString() + " using ';'");
                     return ";";
             }
         }
@@ -908,7 +908,7 @@ namespace UnrealBuildTool
             {
                 if (SDKRoot != "")
                 {
-                    SDKPath = Path.Combine(SDKRoot, "Host" + ExternalExecution.GetRuntimePlatform(), GetSDKTargetPlatformName());
+                    SDKPath = Path.Combine(SDKRoot, "Host" + BuildHostPlatform.Current.Platform, GetSDKTargetPlatformName());
                 }
             }
             return SDKPath;
diff --git a/Engine/Source/Programs/UnrealBuildTool/Configuration/UEBuildTarget.cs b/Engine/Source/Programs/UnrealBuildTool/Configuration/UEBuildTarget.cs
--- a/Engine/Source/Programs/UnrealBuildTool/Configuration/UEBuildTarget.cs
+++ b/Engine/Source/Programs/UnrealBuildTool/Configuration/UEBuildTarget.cs
@@ -840,7 +840,7 @@ namespace UnrealBuildTool
 
 				UBTArguments.Append("UnrealHeaderTool");
 				// Which desktop platform do we need to clean UHT for?
-                UBTArguments.Append(" " + ExternalExecution.GetRuntimePlatform().ToString());
+                UBTArguments.Append(" " + BuildHostPlatform.Current.Platform.ToString());
 				UBTArguments.Append(" " + UnrealTargetConfiguration.Development.ToString());
 				// NOTE: We disable mutex when launching UBT from within UBT to clean UHT
 				UBTArguments.Append(" -NoMutex -Clean");
diff --git a/Engine/Source/Programs/UnrealBuildTool/IOS/IOSToolChain.cs b/Engine/Source/Programs/UnrealBuildTool/IOS/IOSToolChain.cs
--- a/Engine/Source/Programs/UnrealBuildTool/IOS/IOSToolChain.cs
+++ b/Engine/Source/Programs/UnrealBuildTool/IOS/IOSToolChain.cs
@@ -160,7 +160,7 @@ namespace UnrealBuildTool
 
 			IOSSDKVersionFloat = float.Parse(IOSSDKVersion, System.Globalization.CultureInfo.InvariantCulture);
 
-			if (ExternalExecution.GetRuntimePlatform() != UnrealTargetPlatform.Mac)
+			if (BuildHostPlatform.Current.Platform != UnrealTargetPlatform.Mac)
 			{
 				Log.TraceInformation("Compiling with IOS SDK {0} on Mac {1}", IOSSDKVersionFloat, RemoteServerName);
 			}
@@ -348,7 +348,7 @@ namespace UnrealBuildTool
 
 		string GetRemoteFrameworkZipPath( UEBuildFramework Framework )
 		{
-			if ( ExternalExecution.GetRuntimePlatform() != UnrealTargetPlatform.Mac )
+			if ( BuildHostPlatform.Current.Platform != UnrealTargetPlatform.Mac )
 			{
 				return ConvertPath( GetLocalFrameworkZipPath( Framework ) );
 			}
@@ -366,7 +366,7 @@ namespace UnrealBuildTool
 			string IntermediatePath = Framework.OwningModule.Target.ProjectDirectory + "/Intermediate/UnzippedFrameworks/" + Framework.OwningModule.Name;
 			IntermediatePath =  Path.GetFullPath( ( IntermediatePath + Framework.FrameworkZipPath ).Replace( ".zip", "" ) );
 
-			if ( ExternalExecution.GetRuntimePlatform() != UnrealTargetPlatform.Mac )
+			if ( BuildHostPlatform.Current.Platform != UnrealTargetPlatform.Mac )
 			{
 				return ConvertPath( IntermediatePath );
 			}
@@ -376,7 +376,7 @@ namespace UnrealBuildTool
 
 		void CleanIntermediateDirectory( string Path )
 		{
-			if ( ExternalExecution.GetRuntimePlatform() != UnrealTargetPlatform.Mac )
+			if ( BuildHostPlatform.Current.Platform != UnrealTargetPlatform.Mac )
 			{
 				// Delete the intermediate directory on the mac
 				RPCUtilHelper.Command( "/", String.Format( "rm -rf {0}", Path ), "", null );
@@ -468,7 +468,7 @@ namespace UnrealBuildTool
 			{
 				Arguments += string.Format(" -I\"{0}\"", ConvertPath(Path.GetFullPath(IncludePath)));
 
-				if (ExternalExecution.GetRuntimePlatform() != UnrealTargetPlatform.Mac)
+				if (BuildHostPlatform.Current.Platform != UnrealTargetPlatform.Mac)
 				{
 					// sync any third party headers we may need
 					if (IncludePath.Contains("ThirdParty"))
@@ -542,7 +542,7 @@ namespace UnrealBuildTool
 				// Add the C++ source file and its included files to the prerequisite item list.
 				CompileAction.PrerequisiteItems.Add(SourceFile);
 
-				if (ExternalExecution.GetRuntimePlatform() != UnrealTargetPlatform.Mac)
+				if (BuildHostPlatform.Current.Platform != UnrealTargetPlatform.Mac)
 				{
 					QueueFileForBatchUpload(SourceFile);
 				}
@@ -554,7 +554,7 @@ namespace UnrealBuildTool
 						CompileAction.PrerequisiteItems.Add(IncludedFile);
 
 						if (!BuildConfiguration.bUseExperimentalFastDependencyScan &&	// With fast dependency scanning, we will not have an exhaustive list of dependencies here.  We rely on PostCodeGeneration() to upload these files.
-							ExternalExecution.GetRuntimePlatform() != UnrealTargetPlatform.Mac)
+							BuildHostPlatform.Current.Platform != UnrealTargetPlatform.Mac)
 						{
 							QueueFileForBatchUpload(IncludedFile);
 						}
@@ -605,7 +605,7 @@ namespace UnrealBuildTool
 				FileArguments += string.Format(" \"{0}\"", ConvertPath(SourceFile.AbsolutePath), false);
 
 				string CompilerPath = XcodeDeveloperDir + "Toolchains/XcodeDefault.xctoolchain/usr/bin/" + IOSCompiler;
-				if (!Utils.IsRunningOnMono && ExternalExecution.GetRuntimePlatform() != UnrealTargetPlatform.Mac)
+				if (!Utils.IsRunningOnMono && BuildHostPlatform.Current.Platform != UnrealTargetPlatform.Mac)
 				{
 					CompileAction.ActionHandler = new Action.BlockingActionHandler(RPCUtilHelper.RPCActionHandler);
 				}
@@ -633,7 +633,7 @@ namespace UnrealBuildTool
 			// Create an action that invokes the linker.
 			Action LinkAction = new Action(ActionType.Link);
 
-			if (!Utils.IsRunningOnMono && ExternalExecution.GetRuntimePlatform() != UnrealTargetPlatform.Mac)
+			if (!Utils.IsRunningOnMono && BuildHostPlatform.Current.Platform != UnrealTargetPlatform.Mac)
 			{
 				LinkAction.ActionHandler = new Action.BlockingActionHandler(RPCUtilHelper.RPCActionHandler);
 			}
@@ -674,7 +674,7 @@ namespace UnrealBuildTool
 				}
 			}
 
-			if (ExternalExecution.GetRuntimePlatform() != UnrealTargetPlatform.Mac)
+			if (BuildHostPlatform.Current.Platform != UnrealTargetPlatform.Mac)
 			{
 				// Add any additional files that we'll need in order to link the app
 				foreach (string AdditionalShadowFile in LinkEnvironment.Config.AdditionalShadowFiles)
@@ -710,7 +710,7 @@ namespace UnrealBuildTool
 				RememberedAdditionalFrameworks.Add( Framework );
 
 				// Copy them to remote mac if needed
-				if ( ExternalExecution.GetRuntimePlatform() != UnrealTargetPlatform.Mac )
+				if ( BuildHostPlatform.Current.Platform != UnrealTargetPlatform.Mac )
 				{
 					FileItem ShadowFile = FileItem.GetExistingItemByPath( GetLocalFrameworkZipPath( Framework ) );
 
@@ -754,7 +754,7 @@ namespace UnrealBuildTool
 			else
 			{
 				string ResponsePath = Path.GetFullPath("..\\Intermediate\\IOS\\LinkFileList_" + Path.GetFileNameWithoutExtension(LinkEnvironment.Config.OutputFilePath) + ".tmp"); 
-				if (!Utils.IsRunningOnMono && ExternalExecution.GetRuntimePlatform () != UnrealTargetPlatform.Mac)
+				if (!Utils.IsRunningOnMono && BuildHostPlatform.Current.Platform != UnrealTargetPlatform.Mac)
 				{
 					ResponseFile.Create (ResponsePath, InputFileNames);
 					RPCUtilHelper.CopyFile (ResponsePath, ConvertPath (ResponsePath), true);
@@ -940,7 +940,7 @@ namespace UnrealBuildTool
 
 		public override void PreBuildSync()
 		{
-			if (ExternalExecution.GetRuntimePlatform() != UnrealTargetPlatform.Mac)
+			if (BuildHostPlatform.Current.Platform != UnrealTargetPlatform.Mac)
 			{
 				BuiltBinaries = new List<string>();
 			}
@@ -960,7 +960,7 @@ namespace UnrealBuildTool
 				// Assume that there is another directory inside the zip with the same name as the zip
 				ZipDstPath = ZipDstPath.Substring( 0, ZipDstPath.LastIndexOf( '/' ) );
 
-				if ( ExternalExecution.GetRuntimePlatform() == UnrealTargetPlatform.Mac )
+				if ( BuildHostPlatform.Current.Platform == UnrealTargetPlatform.Mac )
 				{
 					// If we're on the mac, just unzip using the shell
 					string ResultsText;
@@ -984,7 +984,7 @@ namespace UnrealBuildTool
 
 			string AppName = Target.Rules.Type == TargetRules.TargetType.Game ? Target.GameName : Target.AppName;
 
-			if (ExternalExecution.GetRuntimePlatform() == UnrealTargetPlatform.Mac)
+			if (BuildHostPlatform.Current.Platform == UnrealTargetPlatform.Mac)
 			{
 				string RemoteShadowDirectoryMac = Path.GetDirectoryName(Target.OutputPath);
 				string FinalRemoteExecutablePath = String.Format("{0}/Payload/{1}.app/{1}", RemoteShadowDirectoryMac, Target.GameName);
diff --git a/Engine/Source/Programs/UnrealBuildTool/IOS/UEBuildIOS.cs b/Engine/Source/Programs/UnrealBuildTool/IOS/UEBuildIOS.cs
--- a/Engine/Source/Programs/UnrealBuildTool/IOS/UEBuildIOS.cs
+++ b/Engine/Source/Programs/UnrealBuildTool/IOS/UEBuildIOS.cs
@@ -276,7 +276,7 @@ namespace UnrealBuildTool
 
         public override void SetupBinaries(UEBuildTarget InBuildTarget)
         {
-            if (ExternalExecution.GetRuntimePlatform () != UnrealTargetPlatform.Mac)
+            if (BuildHostPlatform.Current.Platform != UnrealTargetPlatform.Mac)
             {
                 // dangerously fast mode doesn't generate stub files
                 if (!IOSToolChain.bUseDangerouslyFastMode)
diff --git a/Engine/Source/Programs/UnrealBuildTool/IOS/UEDeployIOS.cs b/Engine/Source/Programs/UnrealBuildTool/IOS/UEDeployIOS.cs
--- a/Engine/Source/Programs/UnrealBuildTool/IOS/UEDeployIOS.cs
+++ b/Engine/Source/Programs/UnrealBuildTool/IOS/UEDeployIOS.cs
@@ -272,7 +272,7 @@ namespace UnrealBuildTool.IOS
 			string BuildPath = InTarget.ProjectDirectory + "/Binaries/IOS";
 			string ProjectDirectory = InTarget.ProjectDirectory;
 
-			if (ExternalExecution.GetRuntimePlatform() == UnrealTargetPlatform.Mac && Environment.GetEnvironmentVariable("UBT_NO_POST_DEPLOY") != "true")
+			if (BuildHostPlatform.Current.Platform == UnrealTargetPlatform.Mac && Environment.GetEnvironmentVariable("UBT_NO_POST_DEPLOY") != "true")
 			{
 				string DecoratedGameName;
 				if (InTarget.Configuration == UnrealTargetConfiguration.Development)
diff --git a/Engine/Source/Programs/UnrealBuildTool/Linux/LinuxToolChain.cs b/Engine/Source/Programs/UnrealBuildTool/Linux/LinuxToolChain.cs
--- a/Engine/Source/Programs/UnrealBuildTool/Linux/LinuxToolChain.cs
+++ b/Engine/Source/Programs/UnrealBuildTool/Linux/LinuxToolChain.cs
@@ -13,7 +13,7 @@ namespace UnrealBuildTool
     {
         protected static bool CrossCompiling()
         {
-            return ExternalExecution.GetRuntimePlatform() != UnrealTargetPlatform.Linux;
+            return BuildHostPlatform.Current.Platform != UnrealTargetPlatform.Linux;
         }
 
         protected static bool UsingClang()
@@ -517,7 +517,7 @@ namespace UnrealBuildTool
             // Create an archive action
             Action ArchiveAction = new Action(ActionType.Link);
             ArchiveAction.WorkingDirectory = Path.GetFullPath(".");
-            bool bUsingSh = ExternalExecution.GetRuntimePlatform() != UnrealTargetPlatform.Win64 && ExternalExecution.GetRuntimePlatform() != UnrealTargetPlatform.Win32;
+            bool bUsingSh = BuildHostPlatform.Current.Platform != UnrealTargetPlatform.Win64 && BuildHostPlatform.Current.Platform != UnrealTargetPlatform.Win32;
             if (bUsingSh)
             {
                 ArchiveAction.CommandPath = "/bin/sh";
@@ -576,7 +576,7 @@ namespace UnrealBuildTool
 
 			Log.TraceVerbose("Adding postlink step");
 
-            bool bUseCmdExe = ExternalExecution.GetRuntimePlatform() == UnrealTargetPlatform.Win64 || ExternalExecution.GetRuntimePlatform() == UnrealTargetPlatform.Win32;
+            bool bUseCmdExe = BuildHostPlatform.Current.Platform == UnrealTargetPlatform.Win64 || BuildHostPlatform.Current.Platform == UnrealTargetPlatform.Win32;
             string ShellBinary = bUseCmdExe ? "cmd.exe" : "/bin/sh";
             string ExecuteSwitch = bUseCmdExe ? " /C" : ""; // avoid -c so scripts don't need +x
             string ScriptName = bUseCmdExe ? "FixDependencies.bat" : "FixDependencies.sh";
@@ -798,7 +798,7 @@ namespace UnrealBuildTool
 			// are created. This script will be called by action created in FixDependencies()
 			if (LinkEnvironment.Config.bIsCrossReferenced && LinkEnvironment.Config.bIsBuildingDLL)
 			{
-                bool bUseCmdExe = ExternalExecution.GetRuntimePlatform() == UnrealTargetPlatform.Win64 || ExternalExecution.GetRuntimePlatform() == UnrealTargetPlatform.Win32;
+                bool bUseCmdExe = BuildHostPlatform.Current.Platform == UnrealTargetPlatform.Win64 || BuildHostPlatform.Current.Platform == UnrealTargetPlatform.Win32;
                 string ScriptName = bUseCmdExe ? "FixDependencies.bat" : "FixDependencies.sh";
 
 				string FixDepsScriptPath = Path.Combine(LinkEnvironment.Config.LocalShadowDirectory, ScriptName);
@@ -876,7 +876,7 @@ namespace UnrealBuildTool
         /** Converts the passed in path from UBT host to compiler native format. */
         public override String ConvertPath(String OriginalPath)
         {
-            if (ExternalExecution.GetRuntimePlatform() == UnrealTargetPlatform.Linux)
+            if (BuildHostPlatform.Current.Platform == UnrealTargetPlatform.Linux)
             {
                 return OriginalPath.Replace("\\", "/");
             }
diff --git a/Engine/Source/Programs/UnrealBuildTool/Linux/UEBuildLinux.cs b/Engine/Source/Programs/UnrealBuildTool/Linux/UEBuildLinux.cs
--- a/Engine/Source/Programs/UnrealBuildTool/Linux/UEBuildLinux.cs
+++ b/Engine/Source/Programs/UnrealBuildTool/Linux/UEBuildLinux.cs
@@ -80,7 +80,7 @@ namespace UnrealBuildTool
          */
         protected override SDKStatus HasRequiredManualSDKInternal()
         {
-            if (ExternalExecution.GetRuntimePlatform() == UnrealTargetPlatform.Linux)
+            if (BuildHostPlatform.Current.Platform == UnrealTargetPlatform.Linux)
             {
                 return SDKStatus.Valid;
             }
diff --git a/Engine/Source/Programs/UnrealBuildTool/Mac/MacToolChain.cs b/Engine/Source/Programs/UnrealBuildTool/Mac/MacToolChain.cs
--- a/Engine/Source/Programs/UnrealBuildTool/Mac/MacToolChain.cs
+++ b/Engine/Source/Programs/UnrealBuildTool/Mac/MacToolChain.cs
@@ -59,7 +59,7 @@ namespace UnrealBuildTool
 
 			if (!ProjectFileGenerator.bGenerateProjectFiles)
 			{
-				if (ExternalExecution.GetRuntimePlatform() != UnrealTargetPlatform.Mac)
+				if (BuildHostPlatform.Current.Platform != UnrealTargetPlatform.Mac)
 				{
 					Log.TraceInformation("Compiling with Mac SDK {0} on Mac {1}", MacOSSDKVersion, RemoteServerName);
 				}
@@ -247,7 +247,7 @@ namespace UnrealBuildTool
 			{
 				Arguments.Append(" -I\"");
 
-				if (ExternalExecution.GetRuntimePlatform() != UnrealTargetPlatform.Mac)
+				if (BuildHostPlatform.Current.Platform != UnrealTargetPlatform.Mac)
 				{
 					Arguments.Append(ConvertPath(Path.GetFullPath(IncludePath)));
 
@@ -331,7 +331,7 @@ namespace UnrealBuildTool
 				// Add the C++ source file and its included files to the prerequisite item list.
 				CompileAction.PrerequisiteItems.Add(SourceFile);
 
-				if (ExternalExecution.GetRuntimePlatform() != UnrealTargetPlatform.Mac)
+				if (BuildHostPlatform.Current.Platform != UnrealTargetPlatform.Mac)
 				{
 					QueueFileForBatchUpload(SourceFile);
 				}
@@ -343,7 +343,7 @@ namespace UnrealBuildTool
 						CompileAction.PrerequisiteItems.Add(IncludedFile);
 
 						if (!BuildConfiguration.bUseExperimentalFastDependencyScan &&	// With fast dependency scanning, we will not have an exhaustive list of dependencies here.  We rely on PostCodeGeneration() to upload these files.
-							 ExternalExecution.GetRuntimePlatform() != UnrealTargetPlatform.Mac)
+							 BuildHostPlatform.Current.Platform != UnrealTargetPlatform.Mac)
 						{
 							QueueFileForBatchUpload(IncludedFile);
 						}
@@ -393,7 +393,7 @@ namespace UnrealBuildTool
 				// Add the source file path to the command-line.
 				FileArguments += string.Format(" \"{0}\"", ConvertPath(SourceFile.AbsolutePath), false);
 
-				if (ExternalExecution.GetRuntimePlatform() != UnrealTargetPlatform.Mac)
+				if (BuildHostPlatform.Current.Platform != UnrealTargetPlatform.Mac)
 				{
 					CompileAction.ActionHandler = new Action.BlockingActionHandler(RPCUtilHelper.RPCActionHandler);
 				}
@@ -531,7 +531,7 @@ namespace UnrealBuildTool
 			// Create an action that invokes the linker.
 			Action LinkAction = new Action(ActionType.Link);
 
-			if (ExternalExecution.GetRuntimePlatform() != UnrealTargetPlatform.Mac)
+			if (BuildHostPlatform.Current.Platform != UnrealTargetPlatform.Mac)
 			{
 				LinkAction.ActionHandler = new Action.BlockingActionHandler(RPCUtilHelper.RPCActionHandler);
 			}
@@ -576,7 +576,7 @@ namespace UnrealBuildTool
 
 			List<string> ThirdPartyLibraries = new List<string>();
 
-			if (ExternalExecution.GetRuntimePlatform() != UnrealTargetPlatform.Mac)
+			if (BuildHostPlatform.Current.Platform != UnrealTargetPlatform.Mac)
 			{
 				// Add any additional files that we'll need in order to link the app
 				foreach (string AdditionalShadowFile in LinkEnvironment.Config.AdditionalShadowFiles)
@@ -617,7 +617,7 @@ namespace UnrealBuildTool
 							ThirdPartyLibraries.Add(AdditionalLibrary);
 						}
 
-						if (ExternalExecution.GetRuntimePlatform() != UnrealTargetPlatform.Mac)
+						if (BuildHostPlatform.Current.Platform != UnrealTargetPlatform.Mac)
 						{
 							// copy over libs we may need
 							FileItem ShadowFile = FileItem.GetExistingItemByPath(AdditionalLibrary);
@@ -704,7 +704,7 @@ namespace UnrealBuildTool
 				{
 					ResponseFile.Create(ResponsePath, InputFileNames);
 
-					if (ExternalExecution.GetRuntimePlatform() != UnrealTargetPlatform.Mac)
+					if (BuildHostPlatform.Current.Platform != UnrealTargetPlatform.Mac)
 					{
 						RPCUtilHelper.CopyFile(ResponsePath, ConvertPath(ResponsePath), true);
 					}
@@ -956,7 +956,7 @@ namespace UnrealBuildTool
 
 			LinkAction.CommandArguments = "-c 'chmod +x \"" + RemoteFixDylibDepsScript.AbsolutePath + "\"; \"" + RemoteFixDylibDepsScript.AbsolutePath + "\"; if [[ $? -ne 0 ]]; then exit 1; fi; ";
 
-			if (ExternalExecution.GetRuntimePlatform() != UnrealTargetPlatform.Mac)
+			if (BuildHostPlatform.Current.Platform != UnrealTargetPlatform.Mac)
 			{
 				LinkAction.ActionHandler = new Action.BlockingActionHandler(RPCUtilHelper.RPCActionHandler);
 			}
@@ -1000,7 +1000,7 @@ namespace UnrealBuildTool
 
 			// Make the compile action
 			Action GenDebugAction = new Action(ActionType.Link);
-			if (ExternalExecution.GetRuntimePlatform() != UnrealTargetPlatform.Mac)
+			if (BuildHostPlatform.Current.Platform != UnrealTargetPlatform.Mac)
 			{
 				GenDebugAction.ActionHandler = new Action.BlockingActionHandler(RPCUtilHelper.RPCActionHandler);
 			}
@@ -1043,7 +1043,7 @@ namespace UnrealBuildTool
 			FileItem BundleScript = FileItem.GetItemByFullPath(Path.Combine(LinkEnvironment.Config.IntermediateDirectory, "FinalizeAppBundle.sh"));
 			FileItem RemoteBundleScript = LocalToRemoteFileItem(BundleScript, true);
 
-			if (ExternalExecution.GetRuntimePlatform() != UnrealTargetPlatform.Mac)
+			if (BuildHostPlatform.Current.Platform != UnrealTargetPlatform.Mac)
 			{
 				FinalizeAppBundleAction.ActionHandler = new Action.BlockingActionHandler(RPCUtilHelper.RPCActionHandler);
 			}
@@ -1204,7 +1204,7 @@ namespace UnrealBuildTool
 		{
 			string ProjectDirectory = Path.GetDirectoryName(Path.GetFullPath(ProjectFileName));
 
-			if (ExternalExecution.GetRuntimePlatform() != UnrealTargetPlatform.Mac)
+			if (BuildHostPlatform.Current.Platform != UnrealTargetPlatform.Mac)
 			{
 				RPCUtilHelper.CopyFile(ProjectFileName, ConvertPath(ProjectFileName), true);
 				RPCUtilHelper.CopyFile("Engine/Source/Programs/DotNETCommon/MetaData.cs", ConvertPath("Engine/Source/Programs/DotNETCommon/MetaData.cs"), true);
@@ -1218,7 +1218,7 @@ namespace UnrealBuildTool
 
 			string XBuildArgs = "/verbosity:quiet /nologo /target:Rebuild /property:Configuration=Development /property:Platform=AnyCPU " + Path.GetFileName(ProjectFileName) + " /property:TargetFrameworkVersion=v4.5 /tv:4.0";
 
-			if (ExternalExecution.GetRuntimePlatform() != UnrealTargetPlatform.Mac)
+			if (BuildHostPlatform.Current.Platform != UnrealTargetPlatform.Mac)
 			{
 				RPCUtilHelper.Command(ConvertPath(ProjectDirectory), "xbuild", XBuildArgs, null);
 			}
@@ -1256,7 +1256,7 @@ namespace UnrealBuildTool
 			var OutputFiles = base.PostBuild(Executable, BinaryLinkEnvironment);
 
 			// If building for Mac on a Mac, use actions to finalize the builds (otherwise, we use Deploy)
-			if (ExternalExecution.GetRuntimePlatform() != UnrealTargetPlatform.Mac)
+			if (BuildHostPlatform.Current.Platform != UnrealTargetPlatform.Mac)
 			{
 				return OutputFiles;
 			}
diff --git a/Engine/Source/Programs/UnrealBuildTool/Mac/UEBuildMac.cs b/Engine/Source/Programs/UnrealBuildTool/Mac/UEBuildMac.cs
--- a/Engine/Source/Programs/UnrealBuildTool/Mac/UEBuildMac.cs
+++ b/Engine/Source/Programs/UnrealBuildTool/Mac/UEBuildMac.cs
@@ -159,8 +159,8 @@ namespace UnrealBuildTool
 		public override void ValidateBuildConfiguration(CPPTargetConfiguration Configuration, CPPTargetPlatform Platform, bool bCreateDebugInfo)
 		{
 //			BuildConfiguration.bUsePCHFiles = false;
-			BuildConfiguration.bCheckExternalHeadersForModification = ExternalExecution.GetRuntimePlatform() != UnrealTargetPlatform.Mac;
-			BuildConfiguration.bCheckSystemHeadersForModification = ExternalExecution.GetRuntimePlatform() != UnrealTargetPlatform.Mac;
+			BuildConfiguration.bCheckExternalHeadersForModification = BuildHostPlatform.Current.Platform != UnrealTargetPlatform.Mac;
+			BuildConfiguration.bCheckSystemHeadersForModification = BuildHostPlatform.Current.Platform != UnrealTargetPlatform.Mac;
 			BuildConfiguration.ProcessorCountMultiplier = MacToolChain.GetAdjustedProcessorCountMultiplier();
 			BuildConfiguration.bUseSharedPCHs = false;
 
diff --git a/Engine/Source/Programs/UnrealBuildTool/Mac/UEDeployMac.cs b/Engine/Source/Programs/UnrealBuildTool/Mac/UEDeployMac.cs
--- a/Engine/Source/Programs/UnrealBuildTool/Mac/UEDeployMac.cs
+++ b/Engine/Source/Programs/UnrealBuildTool/Mac/UEDeployMac.cs
@@ -41,7 +41,7 @@ namespace UnrealBuildTool.IOS
 
 			bool bIsStaticLibrary = InTarget.OutputPath.EndsWith(".a");
 
-			if (ExternalExecution.GetRuntimePlatform() != UnrealTargetPlatform.Mac)
+			if (BuildHostPlatform.Current.Platform != UnrealTargetPlatform.Mac)
 			{
 				if (!bIsStaticLibrary)
 				{
diff --git a/Engine/Source/Programs/UnrealBuildTool/System/BuildHostPlatform.cs b/Engine/Source/Programs/UnrealBuildTool/System/BuildHostPlatform.cs
new file mode 100644
--- /dev/null
+++ b/Engine/Source/Programs/UnrealBuildTool/System/BuildHostPlatform.cs
@@ -0,0 +1,196 @@
+﻿using System;
+using System.Collections.Generic;
+using System.Linq;
+using System.Text;
+using System.IO;
+using System.Runtime.InteropServices;
+
+namespace UnrealBuildTool
+{
+	/// <summary>
+	/// Host platform abstraction
+	/// </summary>
+	public abstract class BuildHostPlatform
+	{
+		private static BuildHostPlatform CurrentPlatform;
+		private static bool bIsMac = File.Exists("/System/Library/CoreServices/SystemVersion.plist");
+
+		/// <summary>
+		/// Returns the name of platform UBT is running on. Internal use only. If you need access this this enum, use BuildHostPlatform.Current.Platform */
+		/// </summary>
+		private static UnrealTargetPlatform GetRuntimePlatform()
+		{
+			PlatformID Platform = Environment.OSVersion.Platform;
+			switch (Platform)
+			{
+				case PlatformID.Win32NT:
+					return UnrealTargetPlatform.Win64;
+				case PlatformID.Unix:
+					return bIsMac ? UnrealTargetPlatform.Mac : UnrealTargetPlatform.Linux;
+				case PlatformID.MacOSX:
+					return UnrealTargetPlatform.Mac;
+				default:
+					throw new BuildException("Unhandled runtime platform " + Platform);
+			}
+		}
+
+		/// <summary>
+		/// Host platform singleton.
+		/// </summary>
+		static public BuildHostPlatform Current
+		{
+			get 
+			{
+				if (CurrentPlatform == null)
+				{
+					switch (GetRuntimePlatform())
+					{
+						case UnrealTargetPlatform.Win64:
+							CurrentPlatform = new WindowsBuildHostPlatform();
+							break;
+						case UnrealTargetPlatform.Mac:
+							CurrentPlatform = new MacBuildHostPlatform();
+							break;
+						case UnrealTargetPlatform.Linux:
+							CurrentPlatform = new LinuxBuildHostPlatform();
+							break;
+					}
+				}
+				return CurrentPlatform;
+			}
+		}
+
+		/// <summary>
+		/// Gets the current host platform type.
+		/// </summary>
+		abstract public UnrealTargetPlatform Platform { get; }
+
+		/// <summary>
+		/// Checks the API version of a dynamic library
+		/// </summary>
+		/// <param name="Filename">Filename of the library</param>
+		/// <returns>API version of -1 if not found.</returns>
+		abstract public int GetDllApiVersion(string Filename);
+	}
+
+	class WindowsBuildHostPlatform : BuildHostPlatform
+	{
+		public override UnrealTargetPlatform Platform
+		{
+			get { return UnrealTargetPlatform.Win64; }
+		}
+
+		[DllImport("kernel32.dll", SetLastError = true)]
+		static extern IntPtr LoadLibraryEx(string lpFileName, IntPtr hFile, uint dwFlags);
+		[DllImport("kernel32.dll", SetLastError = true)]
+		static extern IntPtr FindResource(IntPtr hModule, string lpName, string lpType);
+		[DllImport("kernel32.dll")]
+		static extern IntPtr FindResource(IntPtr hModule, IntPtr lpID, IntPtr lpType);
+		[DllImport("kernel32.dll", SetLastError = true)]
+		static extern IntPtr LoadResource(IntPtr hModule, IntPtr hResInfo);
+		[DllImport("kernel32.dll", SetLastError = true)]
+		static extern uint SizeofResource(IntPtr hModule, IntPtr hResInfo);
+		[DllImport("kernel32.dll", SetLastError = true)]
+		static extern IntPtr LockResource(IntPtr hResData);
+		[DllImport("kernel32.dll", SetLastError = true)]
+		static extern void UnlockResource(IntPtr hResInfo);
+		[DllImport("kernel32.dll", SetLastError = true)]
+		static extern void FreeLibrary(IntPtr hModule);
+
+		[Flags]
+		enum LoadLibraryFlags : uint
+		{
+			DONT_RESOLVE_DLL_REFERENCES = 0x00000001,
+			LOAD_LIBRARY_AS_DATAFILE = 0x00000002,
+			LOAD_WITH_ALTERED_SEARCH_PATH = 0x00000008,
+			LOAD_IGNORE_CODE_AUTHZ_LEVEL = 0x00000010,
+			LOAD_LIBRARY_AS_IMAGE_RESOURCE = 0x00000020,
+			LOAD_LIBRARY_AS_DATAFILE_EXCLUSIVE = 0x00000040,
+			LOAD_LIBRARY_REQUIRE_SIGNED_TARGET = 0x00000080
+		}
+
+		enum ResourceType
+		{
+			RT_CURSOR = 1,
+			RT_BITMAP = 2,
+			RT_ICON = 3,
+			RT_MENU = 4,
+			RT_DIALOG = 5,
+			RT_STRING = 6,
+			RT_FONTDIR = 7,
+			RT_FONT = 8,
+			RT_ACCELERATOR = 9,
+			RT_RCDATA = 10,
+			RT_MESSAGETABLE = 11
+		}
+
+		public override int GetDllApiVersion(string Filename)
+		{
+			int Result = -1;
+
+			try
+			{
+				const int ID_MODULE_API_VERSION_RESOURCE = 191;
+
+				// Retrieves the embedded API version from a DLL
+				IntPtr hModule = LoadLibraryEx(Filename, IntPtr.Zero, (uint)LoadLibraryFlags.LOAD_LIBRARY_AS_DATAFILE);
+				if (hModule != IntPtr.Zero)
+				{
+					IntPtr hResInfo = FindResource(hModule, new IntPtr(ID_MODULE_API_VERSION_RESOURCE), new IntPtr((int)ResourceType.RT_RCDATA));
+					if (hResInfo != IntPtr.Zero)
+					{
+						IntPtr hResGlobal = LoadResource(hModule, hResInfo);
+						if (hResGlobal != IntPtr.Zero)
+						{
+							IntPtr pResData = LockResource(hResGlobal);
+							if (pResData != IntPtr.Zero)
+							{
+								uint Length = SizeofResource(hModule, hResInfo);
+								if (Length > 0)
+								{
+									var Str = Marshal.PtrToStringAnsi(pResData);
+									Result = Int32.Parse(Str);
+								}
+							}
+						}
+					}
+					FreeLibrary(hModule);
+				}
+			}
+			catch (Exception Ex)
+			{
+				Log.TraceWarning("Failed to get DLL API version for {0}. Exception: {1}", Filename, Ex.Message);
+			}
+
+			return Result;
+		}
+	}
+
+	class MacBuildHostPlatform : BuildHostPlatform
+	{
+		public override UnrealTargetPlatform Platform
+		{
+			get { return UnrealTargetPlatform.Mac; }
+		}
+
+		public override int GetDllApiVersion(string Filename)
+		{
+			// @TODO: Implement GetDllApiVersion for Mac
+			return -1;
+		}
+	}
+
+	class LinuxBuildHostPlatform : BuildHostPlatform
+	{
+		public override UnrealTargetPlatform Platform
+		{
+			get { return UnrealTargetPlatform.Linux; }
+		}
+
+		public override int GetDllApiVersion(string Filename)
+		{
+			// @TODO: Implement GetDllApiVersion for Linux
+			return -1;
+		}
+	}
+}
diff --git a/Engine/Source/Programs/UnrealBuildTool/System/ExternalExecution.cs b/Engine/Source/Programs/UnrealBuildTool/System/ExternalExecution.cs
--- a/Engine/Source/Programs/UnrealBuildTool/System/ExternalExecution.cs
+++ b/Engine/Source/Programs/UnrealBuildTool/System/ExternalExecution.cs
@@ -8,6 +8,7 @@ using System.Diagnostics;
 using System.IO;
 using System.Text.RegularExpressions;
 using System.Reflection;
+using System.Runtime.InteropServices;
 
 namespace UnrealBuildTool
 {
@@ -114,6 +115,8 @@ namespace UnrealBuildTool
 	 */
 	public class ExternalExecution
 	{
+
+
 		static ExternalExecution()
 		{
 		}
@@ -123,17 +126,28 @@ namespace UnrealBuildTool
 		/// </summary>
 		static string GetHeaderToolPath()
 		{
-			UnrealTargetPlatform Platform = GetRuntimePlatform();
+			UnrealTargetPlatform Platform = BuildHostPlatform.Current.Platform;
 			string ExeExtension = UEBuildPlatform.GetBuildPlatform(Platform).GetBinaryExtension(UEBuildBinaryType.Executable);
 			string HeaderToolExeName = "UnrealHeaderTool";
 			string HeaderToolPath = Path.Combine("..", "Binaries", Platform.ToString(), HeaderToolExeName + ExeExtension);
 			return HeaderToolPath;
 		}
 
+		class VersionedBinary
+		{
+			public VersionedBinary(string InFilename, int InVersion)
+			{
+				Filename = InFilename;
+				Version = InVersion;
+			}
+			public string Filename;
+			public int Version;
+		}
+
 		/// <summary>
 		/// Finds all UnrealHeaderTool plugins in the plugin directory
 		/// </summary>
-		static void RecursivelyCollectHeaderToolPlugins(string RootPath, string Pattern, string Platform, List<string> PluginBinaries)
+		static void RecursivelyCollectHeaderToolPlugins(string RootPath, string Pattern, string Platform, List<VersionedBinary> PluginBinaries)
 		{
 			var SubDirectories = Directory.GetDirectories(RootPath);
 			foreach (var Dir in SubDirectories)
@@ -153,7 +167,7 @@ namespace UnrealBuildTool
 			{
 				if (Binary.Contains(Platform))
 				{
-					PluginBinaries.Add(Binary);
+					PluginBinaries.Add(new VersionedBinary(Binary, BuildHostPlatform.Current.GetDllApiVersion(Binary)));
 				}
 			}
 		}
@@ -161,21 +175,25 @@ namespace UnrealBuildTool
 		/// <summary>
 		/// Gets all UnrealHeaderTool binaries (including DLLs if it was not build monolithically)
 		/// </summary>
-		static string[] GetHeaderToolBinaries()
+		static VersionedBinary[] GetHeaderToolBinaries()
 		{
-			var Binaries = new List<string>();
+			var Binaries = new List<VersionedBinary>();
 			var HeaderToolExe = GetHeaderToolPath();
 			if (File.Exists(HeaderToolExe))
 			{
-				Binaries.Add(HeaderToolExe);
+				Binaries.Add(new VersionedBinary(HeaderToolExe, -1));
 
 				var HeaderToolLocation = Path.GetDirectoryName(HeaderToolExe);
-				var Platform = GetRuntimePlatform();
+				var Platform = BuildHostPlatform.Current.Platform;
 				var DLLExtension = UEBuildPlatform.GetBuildPlatform(Platform).GetBinaryExtension(UEBuildBinaryType.DynamicLinkLibrary);
 				var DLLSearchPattern = "UnrealHeaderTool-*" + DLLExtension;
 				var HeaderToolDLLs = Directory.GetFiles(HeaderToolLocation, DLLSearchPattern, SearchOption.TopDirectoryOnly);
-								
-				Binaries.AddRange(HeaderToolDLLs);
+
+
+				foreach (var Binary in HeaderToolDLLs)
+				{
+					Binaries.Add(new VersionedBinary(Binary, BuildHostPlatform.Current.GetDllApiVersion(Binary)));
+				}
 
 				var PluginDirectory = Path.Combine("..", "Plugins");
 				RecursivelyCollectHeaderToolPlugins(PluginDirectory, DLLSearchPattern, Platform.ToString(), Binaries);
@@ -186,42 +204,59 @@ namespace UnrealBuildTool
 		/// <summary>
 		/// Gets the latest write time of any of the UnrealHeaderTool binaries (including DLLs and Plugins) or DateTime.MaxValue if UnrealHeaderTool does not exist
 		/// </summary>
-		static DateTime GetHeaderToolTimestamp()
-		{			 
-			var HeaderToolBinaries = GetHeaderToolBinaries();
+		/// <returns>
+		/// Latest timestamp of UHT binaries or DateTime.MaxValue if UnrealHeaderTool is out of date and needs to be rebuilt.
+		/// </returns>
+		static DateTime CheckIfUnrealHeaderToolIsUpToDate()
+		{
 			var LatestWriteTime = DateTime.MinValue;
-			// Find the latest write time for all UnrealHeaderTool binaries
-			foreach (var Binary in HeaderToolBinaries)
+			int? MinVersion = null;
+			using (var TimestampTimer = new ScopedTimer("GetHeaderToolTimestamp"))
 			{
-				var BinaryInfo = new FileInfo(Binary);
-				if (BinaryInfo.Exists)
+				var HeaderToolBinaries = GetHeaderToolBinaries();				
+				// Find the latest write time for all UnrealHeaderTool binaries
+				foreach (var Binary in HeaderToolBinaries)
+				{
+					var BinaryInfo = new FileInfo(Binary.Filename);
+					if (BinaryInfo.Exists)
+					{
+						// Latest write time
+						if (BinaryInfo.LastWriteTime > LatestWriteTime)
+						{
+							LatestWriteTime = BinaryInfo.LastWriteTime;
+						}
+						// Minimum version
+						if (Binary.Version > -1)
+						{
+							MinVersion = MinVersion.HasValue ? Math.Min(MinVersion.Value, Binary.Version) : Binary.Version;
+						}
+					}
+				}
+				if (MinVersion.HasValue)
 				{
-					if (BinaryInfo.LastWriteTime > LatestWriteTime)
+					// If we were able to retrieve the minimal API version, go through all binaries one more time
+					// and delete all binaries that do not match the minimum version (which for local builds would be 0, but it will
+					// also detect bad or partial syncs)
+					foreach (var Binary in HeaderToolBinaries)
 					{
-						LatestWriteTime = BinaryInfo.LastWriteTime;
+						if (Binary.Version > -1)
+						{
+							if (Binary.Version != MinVersion.Value)
+							{
+								// Bad sync
+								File.Delete(Binary.Filename);
+								LatestWriteTime = DateTime.MaxValue;
+								Log.TraceWarning("Detected mismatched version in UHT binary {0} (API Version {1}, expected: {2})", Path.GetFileName(Binary.Filename), Binary.Version, MinVersion.Value);
+							}
+						}
 					}
 				}
 			}
-			// If UHT doesn't exist, force regenerate.
+			// If UHT doesn't exist or is out of date/mismatched, force regenerate.
 			return LatestWriteTime > DateTime.MinValue ? LatestWriteTime : DateTime.MaxValue;
 		}
 
-		private static bool bIsMac = File.Exists("/System/Library/CoreServices/SystemVersion.plist");
 
-		/** Returns the name of platform UBT is running on */
-		public static UnrealTargetPlatform GetRuntimePlatform()
-		{
-			PlatformID Platform = Environment.OSVersion.Platform;
-			switch (Platform)
-			{
-			case PlatformID.Win32NT:
-				return UnrealTargetPlatform.Win64;
-			case PlatformID.Unix:
-				return bIsMac ? UnrealTargetPlatform.Mac : UnrealTargetPlatform.Linux;
-			default:
-				throw new BuildException("Unhandled runtime platform " + Platform);
-			}
-		}
 
 		/// <summary>
 		/// Gets the timestamp of CoreUObject.generated.cpp file.
@@ -264,7 +299,7 @@ namespace UnrealBuildTool
 			bool bIsOutOfDate = false;
 
 			// Get UnrealHeaderTool timestamp. If it's newer than generated headers, they need to be rebuilt too.
-			var HeaderToolTimestamp = GetHeaderToolTimestamp();
+			var HeaderToolTimestamp = CheckIfUnrealHeaderToolIsUpToDate();
 
 			// Get CoreUObject.generated.cpp timestamp.  If the source files are older than the CoreUObject generated code, we'll
 			// need to regenerate code for the module
@@ -500,7 +535,7 @@ namespace UnrealBuildTool
 						UBTArguments.Append( "UnrealHeaderTool" );
 
 						// Which desktop platform do we need to compile UHT for?
-						UBTArguments.Append( " " + GetRuntimePlatform().ToString() );
+						UBTArguments.Append(" " + BuildHostPlatform.Current.Platform.ToString());
 						// NOTE: We force Development configuration for UHT so that it runs quickly, even when compiling debug
 						UBTArguments.Append( " " + UnrealTargetConfiguration.Development.ToString() );
 
diff --git a/Engine/Source/Programs/UnrealBuildTool/System/ProjectFileGenerator.cs b/Engine/Source/Programs/UnrealBuildTool/System/ProjectFileGenerator.cs
--- a/Engine/Source/Programs/UnrealBuildTool/System/ProjectFileGenerator.cs
+++ b/Engine/Source/Programs/UnrealBuildTool/System/ProjectFileGenerator.cs
@@ -269,7 +269,7 @@ namespace UnrealBuildTool
 
 				if (!Directory.Exists(MasterProjectRelativePath + "/Source"))
 				{
-					if (ExternalExecution.GetRuntimePlatform () == UnrealTargetPlatform.Mac)
+					if (BuildHostPlatform.Current.Platform == UnrealTargetPlatform.Mac)
 					{
 						MasterProjectRelativePath = Path.GetFullPath(Path.Combine (Utils.GetExecutingAssemblyDirectory (), "..", "..", "..", "Engine"));
 						GameProjectName = "UE4Game";
@@ -1563,7 +1563,7 @@ namespace UnrealBuildTool
 					// Create target rules for all of the platforms and configuration combinations that we want to enable support for.
 					// Just use the current platform as we only need to recover the target type and both should be supported for all targets...
 					string UnusedTargetFilePath;
-					var TargetRulesObject = RulesCompiler.CreateTargetRules(TargetName, new TargetInfo(ExternalExecution.GetRuntimePlatform(), UnrealTargetConfiguration.Development), false, out UnusedTargetFilePath);
+					var TargetRulesObject = RulesCompiler.CreateTargetRules(TargetName, new TargetInfo(BuildHostPlatform.Current.Platform, UnrealTargetConfiguration.Development), false, out UnusedTargetFilePath);
 
 					// Exclude client and server targets under binary Rocket; it's impossible to build without precompiled engine binaries
 					if (!UnrealBuildTool.RunningRocket() || (TargetRulesObject.Type != TargetRules.TargetType.Client && TargetRulesObject.Type != TargetRules.TargetType.Server))
@@ -1992,7 +1992,7 @@ namespace UnrealBuildTool
                     // TODO(sbc): See if we can just drop the Encoding.UTF8 argument on all
                     // platforms.  In this case UTF8 encoding will still be used but without the
                     // BOM, which is, AFAICT, desirable in almost all cases.
-					if (ExternalExecution.GetRuntimePlatform() == UnrealTargetPlatform.Linux || ExternalExecution.GetRuntimePlatform() == UnrealTargetPlatform.Mac)
+					if (BuildHostPlatform.Current.Platform == UnrealTargetPlatform.Linux || BuildHostPlatform.Current.Platform == UnrealTargetPlatform.Mac)
                         File.WriteAllText(FileName, NewFileContents);
                     else
                         File.WriteAllText(FileName, NewFileContents, Encoding.UTF8);
diff --git a/Engine/Source/Programs/UnrealBuildTool/System/UnrealBuildTool.cs b/Engine/Source/Programs/UnrealBuildTool/System/UnrealBuildTool.cs
--- a/Engine/Source/Programs/UnrealBuildTool/System/UnrealBuildTool.cs
+++ b/Engine/Source/Programs/UnrealBuildTool/System/UnrealBuildTool.cs
@@ -346,7 +346,7 @@ namespace UnrealBuildTool
         /// <returns>true if valid, false if not</returns>
         static public bool IsValidDesktopPlatform(UnrealTargetPlatform InPlatform)
         {
-            switch (ExternalExecution.GetRuntimePlatform())
+            switch (BuildHostPlatform.Current.Platform)
             {
             case UnrealTargetPlatform.Linux:
                 return InPlatform == UnrealTargetPlatform.Linux;
@@ -355,7 +355,7 @@ namespace UnrealBuildTool
             case UnrealTargetPlatform.Win64:
                 return ((InPlatform == UnrealTargetPlatform.Win32) || (InPlatform == UnrealTargetPlatform.Win64));
             default:
-                throw new BuildException("Invalid RuntimePlatform:" + ExternalExecution.GetRuntimePlatform());
+                throw new BuildException("Invalid RuntimePlatform:" + BuildHostPlatform.Current.Platform);
             }
         }
 
@@ -1180,7 +1180,7 @@ namespace UnrealBuildTool
                 }
 
                 bool CreateStub = Utils.ParseCommandLineFlag(Arguments, "-nocreatestub", out ArgumentIndex);
-                if (CreateStub || (String.IsNullOrEmpty(Environment.GetEnvironmentVariable("uebp_LOCAL_ROOT")) && ExternalExecution.GetRuntimePlatform() == UnrealTargetPlatform.Mac))
+                if (CreateStub || (String.IsNullOrEmpty(Environment.GetEnvironmentVariable("uebp_LOCAL_ROOT")) && BuildHostPlatform.Current.Platform == UnrealTargetPlatform.Mac))
                 {
                     BuildConfiguration.bCreateStubIPA = false;
                 }
diff --git a/Engine/Source/Programs/UnrealBuildTool/System/XcodeProjectFileGenerator.cs b/Engine/Source/Programs/UnrealBuildTool/System/XcodeProjectFileGenerator.cs
--- a/Engine/Source/Programs/UnrealBuildTool/System/XcodeProjectFileGenerator.cs
+++ b/Engine/Source/Programs/UnrealBuildTool/System/XcodeProjectFileGenerator.cs
@@ -1352,7 +1352,7 @@ namespace UnrealBuildTool
 		private void WriteProject(StringBuilder XcodeProjectFileContent, List<XcodeProjectTarget> ProjectTargets)
 		{
 			// @todo: Windows and Mac should store the project in the same folder
-			string PathBranch = (ExternalExecution.GetRuntimePlatform() != UnrealTargetPlatform.Mac) ? "/Engine/Intermediate/IOS" : "";
+			string PathBranch = (BuildHostPlatform.Current.Platform != UnrealTargetPlatform.Mac) ? "/Engine/Intermediate/IOS" : "";
 			var XcodeProjectPath = MasterProjectRelativePath + PathBranch + "/" + MasterProjectName + ".xcodeproj";
 			if (bGeneratingGameProjectFiles && GameProjectName == "UE4Game")
 			{
@@ -1388,7 +1388,7 @@ namespace UnrealBuildTool
 
 		private void WriteWorkspace()
 		{
-			string PathBranch = (ExternalExecution.GetRuntimePlatform() != UnrealTargetPlatform.Mac) ? "/Engine/Intermediate/IOS" : "";
+			string PathBranch = (BuildHostPlatform.Current.Platform != UnrealTargetPlatform.Mac) ? "/Engine/Intermediate/IOS" : "";
 			var XcodeProjectPath = MasterProjectRelativePath + PathBranch + "/" + MasterProjectName + ".xcodeproj";
 			if (bGeneratingGameProjectFiles && GameProjectName == "UE4Game")
 			{
diff --git a/Engine/Source/Programs/UnrealBuildTool/ToolChain/RemoteToolChain.cs b/Engine/Source/Programs/UnrealBuildTool/ToolChain/RemoteToolChain.cs
--- a/Engine/Source/Programs/UnrealBuildTool/ToolChain/RemoteToolChain.cs
+++ b/Engine/Source/Programs/UnrealBuildTool/ToolChain/RemoteToolChain.cs
@@ -195,7 +195,7 @@ namespace UnrealBuildTool
 				return InitializationErrorCode;
 			}
 
-            if (ExternalExecution.GetRuntimePlatform() != UnrealTargetPlatform.Mac)
+            if (BuildHostPlatform.Current.Platform != UnrealTargetPlatform.Mac)
             {
             	// If we don't care which machine we're going to build on, query and
             	// pick the one with the most free command slots available
@@ -354,7 +354,7 @@ namespace UnrealBuildTool
         /** Converts the passed in path from UBT host to compiler native format. */
         public override String ConvertPath(String OriginalPath)
         {
-            if (ExternalExecution.GetRuntimePlatform() != UnrealTargetPlatform.Mac)
+            if (BuildHostPlatform.Current.Platform != UnrealTargetPlatform.Mac)
             {
 				if (OriginalPath[1] != ':')
 				{
@@ -380,7 +380,7 @@ namespace UnrealBuildTool
 
 		protected string GetMacDevSrcRoot()
 		{
-			if (ExternalExecution.GetRuntimePlatform() != UnrealTargetPlatform.Mac)
+			if (BuildHostPlatform.Current.Platform != UnrealTargetPlatform.Mac)
 			{
 				// figure out the remote version of Engine/Source
 				return ConvertPath(Path.GetFullPath(Path.Combine(BranchDirectory, "Engine/Source/")));
@@ -464,7 +464,7 @@ namespace UnrealBuildTool
 
 		public override void PostCodeGeneration(UEBuildTarget Target, UHTManifest Manifest)
 		{
-			if (ExternalExecution.GetRuntimePlatform() != UnrealTargetPlatform.Mac)
+			if (BuildHostPlatform.Current.Platform != UnrealTargetPlatform.Mac)
 			{
 				// @todo UHT: Temporary workaround for UBT no longer being able to follow includes from generated headers unless
 				//  the headers already existed before the build started.  We're working on a proper fix.
@@ -549,7 +549,7 @@ namespace UnrealBuildTool
 		public override void PreBuildSync()
         {
 			// no need to sync on the Mac!
-			if (ExternalExecution.GetRuntimePlatform() == UnrealTargetPlatform.Mac)
+			if (BuildHostPlatform.Current.Platform == UnrealTargetPlatform.Mac)
 			{
 				return;
 			}
@@ -789,7 +789,7 @@ namespace UnrealBuildTool
 
 		static public Double GetAdjustedProcessorCountMultiplier()
 		{
-            if (ExternalExecution.GetRuntimePlatform() != UnrealTargetPlatform.Mac)
+            if (BuildHostPlatform.Current.Platform != UnrealTargetPlatform.Mac)
             {
                 Int32 RemoteCPUCount = RPCUtilHelper.GetCommandSlots();
                 if (RemoteCPUCount == 0)
diff --git a/Engine/Source/Programs/UnrealBuildTool/UnrealBuildTool.csproj b/Engine/Source/Programs/UnrealBuildTool/UnrealBuildTool.csproj
--- a/Engine/Source/Programs/UnrealBuildTool/UnrealBuildTool.csproj
+++ b/Engine/Source/Programs/UnrealBuildTool/UnrealBuildTool.csproj
@@ -215,6 +215,7 @@
     <Compile Include="System\ActionHistory.cs" />
     <Compile Include="System\BuildException.cs" />
     <Compile Include="Configuration\BuildConfiguration.cs" />
+    <Compile Include="System\BuildHostPlatform.cs" />
     <Compile Include="System\FlatCPPIncludeDepencencyCache.cs" />
     <Compile Include="System\Distcc.cs" />
     <Compile Include="System\ModuleProcessingException.cs" />
diff --git a/Engine/Source/Programs/UnrealBuildTool/UnrealBuildTool_Mono.csproj b/Engine/Source/Programs/UnrealBuildTool/UnrealBuildTool_Mono.csproj
--- a/Engine/Source/Programs/UnrealBuildTool/UnrealBuildTool_Mono.csproj
+++ b/Engine/Source/Programs/UnrealBuildTool/UnrealBuildTool_Mono.csproj
@@ -121,6 +121,7 @@
     <Compile Include="System\ActionHistory.cs" />
     <Compile Include="System\BuildException.cs" />
     <Compile Include="Configuration\BuildConfiguration.cs" />
+    <Compile Include="System\BuildHostPlatform.cs" />
     <Compile Include="System\DirectoryLookupCache.cs" />
     <Compile Include="System\CPPEnvironment.cs" />
     <Compile Include="System\CPPHeaders.cs" />
