Enabled all but few Clang warnings when building on Mac and cleaned up MacToolChain.cs a bit

diff --git a/Engine/Source/Programs/UnrealBuildTool/Mac/MacToolChain.cs b/Engine/Source/Programs/UnrealBuildTool/Mac/MacToolChain.cs
--- a/Engine/Source/Programs/UnrealBuildTool/Mac/MacToolChain.cs
+++ b/Engine/Source/Programs/UnrealBuildTool/Mac/MacToolChain.cs
@@ -51,8 +51,6 @@ namespace UnrealBuildTool
 
 		private static List<FileItem> BundleDependencies = new List<FileItem>();
 
-		public List<string> BuiltBinaries = new List<string>();
-
 		public override void SetUpGlobalEnvironment()
 		{
 			base.SetUpGlobalEnvironment();
@@ -76,25 +74,29 @@ namespace UnrealBuildTool
 
 			Result += " -fmessage-length=0";
 			Result += " -pipe";
-			Result += " -Wno-trigraphs";
 			Result += " -fpascal-strings";
-			Result += " -mdynamic-no-pic";
-			Result += " -Wreturn-type";
-			Result += " -Wno-format";
+
 			Result += " -fexceptions";
-			Result += " -fno-threadsafe-statics";
-			Result += " -Wno-deprecated-writable-strings";
+			Result += " -fasm-blocks";
+
+			Result += " -Wall -Werror";
+
+			Result += " -Wno-unused-variable";
 			Result += " -Wno-unused-value";
-			Result += " -Wno-switch-enum";
-			Result += " -Wno-logical-op-parentheses";	// needed for external headers we can't change
-			Result += " -Wno-null-arithmetic";			// needed for external headers we can't change
-			Result += " -Wno-return-type-c-linkage";	// needed for PhysX
-			Result += " -Wno-ignored-attributes";		// needed for nvtesslib
-			Result += " -Wno-uninitialized";
-			Result += " -Wno-tautological-compare";
+			// This will hide the warnings about static functions in headers that aren't used in every single .cpp file
+			Result += " -Wno-unused-function";
+			// This hides the "enumeration value 'XXXXX' not handled in switch [-Wswitch]" warnings - we should maybe remove this at some point and add UE_LOG(, Fatal, ) to default cases
 			Result += " -Wno-switch";
+			// This hides the "warning : comparison of unsigned expression < 0 is always false" type warnings due to constant comparisons, which are possible with template arguments
+			Result += " -Wno-tautological-compare";
+			// This will prevent the issue of warnings for unused private variables.
+			Result += " -Wno-unused-private-field";
 			Result += " -Wno-invalid-offsetof"; // needed to suppress warnings about using offsetof on non-POD types.
-			Result += " -fasm-blocks";
+
+			// @todo: Remove these two when the code is fixed and they're no longer needed
+			Result += " -Wno-logical-op-parentheses";
+			Result += " -Wno-unknown-pragmas";
+
 			Result += " -c";
 
 			Result += " -arch x86_64";
@@ -104,7 +106,14 @@ namespace UnrealBuildTool
 			// Optimize non- debug builds.
 			if (CompileEnvironment.Config.Target.Configuration != CPPTargetConfiguration.Debug)
 			{
-				Result += " -O3";
+				if (UEBuildConfiguration.bCompileForSize)
+				{
+					Result += " -Oz";
+				}
+				else
+				{
+					Result += " -O3";
+				}
 			}
 			else
 			{
@@ -191,7 +200,7 @@ namespace UnrealBuildTool
 			// Needed to make sure install_name_tool will be able to update paths in Mach-O headers
 			Result += " -headerpad_max_install_names";
 
-			Result += " -lc++"; // for STL used in HLSLCC
+			Result += " -lc++";
 
 			foreach (string Framework in LinkEnvironment.Config.Frameworks)
 			{
@@ -212,9 +221,7 @@ namespace UnrealBuildTool
 		static string GetArchiveArguments_Global(LinkEnvironment LinkEnvironment)
 		{
 			string Result = "";
-
 			Result += " -static";
-
 			return Result;
 		}
 
@@ -225,11 +232,6 @@ namespace UnrealBuildTool
 
 			Arguments.Append(GetCompileArguments_Global(CompileEnvironment));
 
-			if (CompileEnvironment.Config.PrecompiledHeaderAction != PrecompiledHeaderAction.Create)
-			{
-				Arguments.Append(" -Werror");
-			}
-
 			if (CompileEnvironment.Config.PrecompiledHeaderAction == PrecompiledHeaderAction.Include)
 			{
 				// Add the precompiled header file's path to the include path so GCC can find it.
@@ -624,7 +626,6 @@ namespace UnrealBuildTool
 							if (ShadowFile != null)
 							{
 								QueueFileForBatchUpload(ShadowFile);
-								//							LinkAction.PrerequisiteItems.Add(ShadowFile);
 							}
 						}
 					}
@@ -689,10 +690,6 @@ namespace UnrealBuildTool
 				{
 					LinkCommand += " " + Filename;
 				}
-				// @todo rocket lib: the -filelist command should take a response file (see else condition), except that it just says it can't
-				// find the file that's in there. Rocket.lib may overflow the commandline by putting all files on the commandline, so this 
-				// may be needed:
-				// LinkAction.CommandArguments += string.Format(" -filelist \"{0}\"", ConvertPath(ResponsePath));
 			}
 			else
 			{
@@ -931,13 +928,6 @@ namespace UnrealBuildTool
 				}
 			}
 
-			// For Mac, generate the dSYM file if the config file is set to do so
-			if (BuildConfiguration.bGeneratedSYMFile == true)
-			{
-				Log.TraceInformation("Generating the dSYM file - this will add some time to your build...");
-				RemoteOutputFile = GenerateDebugInfo(RemoteOutputFile);
-			}
-
 			return RemoteOutputFile;
 		}
 
@@ -984,43 +974,6 @@ namespace UnrealBuildTool
 			return RemoteOutputFile;
 		}
 
-		/**
-		 * Generates debug info for a given executable
-		 * 
-		 * @param MachOBinary FileItem describing the executable or dylib to generate debug info for
-		 */
-		public FileItem GenerateDebugInfo(FileItem MachOBinary)
-		{
-			// Make a file item for the source and destination files
-			string AdditionalExtension = Path.GetExtension(MachOBinary.AbsolutePath) == ".dylib" ? "" : ".app";
-			string FullDestPathRoot = MachOBinary.AbsolutePath + AdditionalExtension + ".dSYM";
-			string FullDestPath = FullDestPathRoot;
-			FileItem OutputFile = FileItem.GetItemByPath(FullDestPath);
-			FileItem DestFile = LocalToRemoteFileItem(OutputFile, false);
-
-			// Make the compile action
-			Action GenDebugAction = new Action(ActionType.Link);
-			if (BuildHostPlatform.Current.Platform != UnrealTargetPlatform.Mac)
-			{
-				GenDebugAction.ActionHandler = new Action.BlockingActionHandler(RPCUtilHelper.RPCActionHandler);
-			}
-			GenDebugAction.WorkingDirectory = Path.GetFullPath(".");
-			GenDebugAction.CommandPath = "sh";
-
-			// note that the source and dest are switched from a copy command
-			GenDebugAction.CommandArguments = string.Format("-c '{0}usr/bin/xcrun dsymutil {1} -o {2}'",
-				DeveloperDir,
-				MachOBinary.AbsolutePath,
-				FullDestPathRoot,
-				Path.GetFileName(MachOBinary.AbsolutePath)+AdditionalExtension);
-			GenDebugAction.PrerequisiteItems.Add(MachOBinary);
-			GenDebugAction.ProducedItems.Add(DestFile);
-			GenDebugAction.StatusDescription = GenDebugAction.CommandArguments;
-			GenDebugAction.bCanExecuteRemotely = false;
-
-			return DestFile;
-		}
-
 		/**
 		 * Creates app bundle for a given executable
 		 * 
@@ -1234,23 +1187,6 @@ namespace UnrealBuildTool
 			}
 		}
 
-		public override void PreBuildSync()
-		{
-			BuiltBinaries = new List<string>();
-
-			base.PreBuildSync();
-		}
-
-		public override void PostBuildSync(UEBuildTarget Target)
-		{
-			foreach (UEBuildBinary Binary in Target.AppBinaries)
-			{
-				BuiltBinaries.Add(Path.GetFullPath(Binary.ToString()));
-			}
-
-			base.PostBuildSync(Target);
-		}
-
 		public override ICollection<FileItem> PostBuild(FileItem Executable, LinkEnvironment BinaryLinkEnvironment)
 		{
 			var OutputFiles = base.PostBuild(Executable, BinaryLinkEnvironment);
