Restored BuiltBinaries in MacToolChain

diff --git a/Engine/Source/Programs/UnrealBuildTool/Mac/MacToolChain.cs b/Engine/Source/Programs/UnrealBuildTool/Mac/MacToolChain.cs
--- a/Engine/Source/Programs/UnrealBuildTool/Mac/MacToolChain.cs
+++ b/Engine/Source/Programs/UnrealBuildTool/Mac/MacToolChain.cs
@@ -51,6 +51,8 @@ namespace UnrealBuildTool
 
 		private static List<FileItem> BundleDependencies = new List<FileItem>();
 
+		public List<string> BuiltBinaries = new List<string>();
+
 		public override void SetUpGlobalEnvironment()
 		{
 			base.SetUpGlobalEnvironment();
@@ -1187,6 +1189,23 @@ namespace UnrealBuildTool
 			}
 		}
 
+		public override void PreBuildSync()
+		{
+			BuiltBinaries = new List<string>();
+
+			base.PreBuildSync();
+		}
+
+		public override void PostBuildSync(UEBuildTarget Target)
+		{
+			foreach (UEBuildBinary Binary in Target.AppBinaries)
+			{
+				BuiltBinaries.Add(Path.GetFullPath(Binary.ToString()));
+			}
+
+			base.PostBuildSync(Target);
+		}
+
 		public override ICollection<FileItem> PostBuild(FileItem Executable, LinkEnvironment BinaryLinkEnvironment)
 		{
 			var OutputFiles = base.PostBuild(Executable, BinaryLinkEnvironment);
