added quotes around paths in unzip and mkdir for the remote tool chain
added quotes around framework paths for remote tool chain
#ue4
#ios

diff --git a/Engine/Source/Programs/UnrealBuildTool/IOS/IOSToolChain.cs b/Engine/Source/Programs/UnrealBuildTool/IOS/IOSToolChain.cs
--- a/Engine/Source/Programs/UnrealBuildTool/IOS/IOSToolChain.cs
+++ b/Engine/Source/Programs/UnrealBuildTool/IOS/IOSToolChain.cs
@@ -379,10 +379,10 @@ namespace UnrealBuildTool
 			if ( BuildHostPlatform.Current.Platform != UnrealTargetPlatform.Mac )
 			{
 				// Delete the intermediate directory on the mac
-				RPCUtilHelper.Command( "/", String.Format( "rm -rf {0}", Path ), "", null );
+				RPCUtilHelper.Command( "/", String.Format( "rm -rf \"{0}\"", Path ), "", null );
 
 				// Create a fresh intermediate after we delete it
-				RPCUtilHelper.Command( "/", String.Format( "mkdir -p {0}", Path ), "", null );
+				RPCUtilHelper.Command( "/", String.Format( "mkdir -p \"{0}\"", Path ), "", null );
 			}
 			else
 			{
@@ -394,7 +394,7 @@ namespace UnrealBuildTool
 
 				// Create the intermediate local directory
 				string ResultsText;
-				RunExecutableAndWait( "mkdir", String.Format( "-p {0}", Path ), out ResultsText );
+				RunExecutableAndWait( "mkdir", String.Format( "-p \"{0}\"", Path ), out ResultsText );
 			}
 		}
 
@@ -426,7 +426,7 @@ namespace UnrealBuildTool
 				if ( Framework.OwningModule != null && Framework.FrameworkZipPath != null && Framework.FrameworkZipPath != "" )
 				{
 					// If this framework has a zip specified, we'll need to setup the path as well
-					Result += " -F " + GetRemoteIntermediateFrameworkZipPath( Framework );
+					Result += " -F \"" + GetRemoteIntermediateFrameworkZipPath( Framework ) + "\"";
 				}
 
 				Result += " -framework " + Framework.FrameworkName;
@@ -964,12 +964,12 @@ namespace UnrealBuildTool
 				{
 					// If we're on the mac, just unzip using the shell
 					string ResultsText;
-					RunExecutableAndWait( "unzip", String.Format( "-o {0} -d {1}", ZipSrcPath, ZipDstPath ), out ResultsText );
+					RunExecutableAndWait( "unzip", String.Format( "-o \"{0}\" -d \"{1}\"", ZipSrcPath, ZipDstPath ), out ResultsText );
 					continue;
 				}
 
 				// Use RPC utility if the zip is on remote mac
-				Hashtable Result = RPCUtilHelper.Command( "/", String.Format( "unzip -o {0} -d {1}", ZipSrcPath, ZipDstPath ), "", null );
+				Hashtable Result = RPCUtilHelper.Command( "/", String.Format( "unzip -o \"{0}\" -d \"{1}\"", ZipSrcPath, ZipDstPath ), "", null );
 
 				foreach ( DictionaryEntry Entry in Result )
 				{
