UE4: HTML5 code is now only conditionally included in mono UBT so branches that do not have html5 code can still use UBT mono.

diff --git a/Engine/Source/Programs/UnrealBuildTool/UnrealBuildTool_Mono.csproj b/Engine/Source/Programs/UnrealBuildTool/UnrealBuildTool_Mono.csproj
--- a/Engine/Source/Programs/UnrealBuildTool/UnrealBuildTool_Mono.csproj
+++ b/Engine/Source/Programs/UnrealBuildTool/UnrealBuildTool_Mono.csproj
@@ -86,6 +86,15 @@
     <Reference Include="System.Xml" />
     <Reference Include="System.Xml.Linq" />
   </ItemGroup>
+  <Choose>
+    <When Condition="Exists('HTML5\UEBuildHTML5.cs')">
+      <ItemGroup>
+        <Compile Include="HTML5\HTML5ProjectGenerator.cs" />
+        <Compile Include="HTML5\HTML5ToolChain.cs" />
+        <Compile Include="HTML5\UEBuildHTML5.cs" />
+      </ItemGroup>
+    </When>
+  </Choose>
   <ItemGroup>
     <Compile Include="..\DotNETCommon\MetaData.cs">
       <Link>Properties\MetaData.cs</Link>
@@ -105,9 +114,6 @@
     <Compile Include="fastJSON\JsonSerializer.cs" />
     <Compile Include="fastJSON\Reflection.cs" />
     <Compile Include="fastJSON\SafeDictionary.cs" />
-    <Compile Include="HTML5\HTML5ProjectGenerator.cs" />
-    <Compile Include="HTML5\HTML5ToolChain.cs" />
-    <Compile Include="HTML5\UEBuildHTML5.cs" />
     <Compile Include="IOS\IOSProjectGenerator.cs" />
     <Compile Include="IOS\IOSToolChain.cs" />
     <Compile Include="IOS\UEBuildIOS.cs" />
