Experimental UnrealBuildTool makefile support

UnrealBuildTool 'Makefiles' allow for very fast iterative builds.
    - New BuildConfiguration.xml setting added: "bUseExperimentalFastBuildIteration"   (disabled by default)
    - Turning this on causes Unreal Build Tool to emit 'UBT Makefiles' for targets when they're built the first time.
    - Subsequent builds will load these Makefiles and begin outdatedness checking and build invocation very quickly.
    - The caveat is that if source files are added or removed to the project, UBT will need to gather information about those in order for your build to complete successfully.
    - Currently, you must run the project file generator after adding/removing source files to tell UBT to re-gather this information.
    - Events that can invalidate the 'UBT Makefile':
           - Adding/removing .cpp files
           - Adding/removing .h files with UObjects
           - Adding new UObject types to a file that didn't previously have any
           - Changing global build settings (most settings in this file qualify.)
           - Changed code that affects how Unreal Header Tool works
    - You can force regeneration of the 'UBT Makefile' by passing the '-Gather' argument, or simply regenerating project files
    - New command-line parameters added:
           - "-Gather": Tells UBT to always perform the gather step (slower but will catch project structural changes)
           - "-NoGather": Disables the gather step, unless UBT detects that it must be done.  This is the default when bUseExperimentalFastBuildIteration is enabled
           - "-GatherOnly": Runs the gather step and saves a UBTMakefile, but doesn't build anything
           - "-Assemble": Tells UBT to also assemble build products.  This always defaults to enabled
           - "-NoAssemble": Tells UBT to skip the assemble step, whether we gathered build products or not
           - "-AssembleOnly": Tells UBT to only assemble build products and not to gather, unless UBT determines it must

Other changes:
- UBT now keeps track of which targets it was building in an intermediate file, to help it invalidate cached includes in subsequent runs when the targets are different
- C++ includes are now stored in a class separate from the C++ compile enviroment (for easier serialization)
- The method that UBT uses to find the CoreUObject module timestamp was rewritten
- Various '[REDACTED_HANDLE] ubtmake' comments added to tag possible remaining Makefile tasks
- The 'FileItem' class had some member variable comments and code cleaned up, while making it serializable
- Cleaned up the comments and member variables in the "Action" class, while making it serializable
- Some UBT classes are now "serializable".  This is because we need to store the data in UBTMakefiles.
- Removed support for Actions to tinker with Stdout and Stderror (was not used for anything)
- Moved PrecompileHeaderEnvironment class to the UEBuildModule.cs source file
- Plugin intermediate include directories are now selected on demand rather than cached early
- Toolchain code for gathering prerequisite headers is now shared in a single function (AddPrerequisiteSourceFile)
- Removed Action.StatusDetailedDescription, was not used for anything
- Removed UEBuildConfiguration.bExcludePlugins, was not used for anything
- Removed ECompilationResult.FailedDueToHeaderChange, was not used for anything

diff --git a/Engine/Source/Programs/UnrealBuildTool/Android/AndroidToolChain.cs b/Engine/Source/Programs/UnrealBuildTool/Android/AndroidToolChain.cs
--- a/Engine/Source/Programs/UnrealBuildTool/Android/AndroidToolChain.cs
+++ b/Engine/Source/Programs/UnrealBuildTool/Android/AndroidToolChain.cs
@@ -420,11 +420,11 @@ namespace UnrealBuildTool
 			}
 
 			// Add include paths to the argument list.
-			foreach (string IncludePath in CompileEnvironment.Config.SystemIncludePaths)
+			foreach (string IncludePath in CompileEnvironment.Config.CPPIncludeInfo.SystemIncludePaths)
 			{
 				Arguments += string.Format(" -I\"{0}\"", IncludePath);
 			}
-			foreach (string IncludePath in CompileEnvironment.Config.IncludePaths)
+			foreach (string IncludePath in CompileEnvironment.Config.CPPIncludeInfo.IncludePaths)
 			{
 				Arguments += string.Format(" -I\"{0}\"", IncludePath);
 			}
@@ -479,11 +479,7 @@ namespace UnrealBuildTool
 				}
 
 				// Add the C++ source file and its included files to the prerequisite item list.
-				CompileAction.PrerequisiteItems.Add(SourceFile);
-				{
-					var IncludedFileList = CPPEnvironment.FindAndCacheAllIncludedFiles( Target, SourceFile, BuildPlatform, CompileEnvironment.GetIncludesPathsToSearch( SourceFile ), CompileEnvironment.IncludeFileSearchDictionary, bOnlyCachedDependencies:BuildConfiguration.bUseExperimentalFastDependencyScan );
-					CompileAction.PrerequisiteItems.AddRange( IncludedFileList );
-				}
+				AddPrerequisiteSourceFile( Target, BuildPlatform, CompileEnvironment, SourceFile, CompileAction.PrerequisiteItems );
 
 				if (CompileEnvironment.Config.PrecompiledHeaderAction == PrecompiledHeaderAction.Create)
 				{
@@ -542,7 +538,6 @@ namespace UnrealBuildTool
 				CompileAction.CommandPath = ClangPath;
 				CompileAction.CommandArguments = ResponseArgument;
 				CompileAction.StatusDescription = string.Format("{0}", Path.GetFileName(SourceFile.AbsolutePath));
-				CompileAction.StatusDetailedDescription = SourceFile.Description;
 
 				CompileAction.OutputEventHandler = new DataReceivedEventHandler(CompileOutputReceivedDataEventHandler);
 
diff --git a/Engine/Source/Programs/UnrealBuildTool/Android/UEBuildAndroid.cs b/Engine/Source/Programs/UnrealBuildTool/Android/UEBuildAndroid.cs
--- a/Engine/Source/Programs/UnrealBuildTool/Android/UEBuildAndroid.cs
+++ b/Engine/Source/Programs/UnrealBuildTool/Android/UEBuildAndroid.cs
@@ -1,4 +1,4 @@
-﻿// Copyright 1998-2014 Epic Games, Inc. All Rights Reserved.
+// Copyright 1998-2014 Epic Games, Inc. All Rights Reserved.
 
 using System;
 using System.Collections.Generic;
@@ -301,21 +301,21 @@ namespace UnrealBuildTool
 
 			if (InBuildTarget.GlobalCompileEnvironment.Config.Target.Architecture == "-armv7")
 			{
-				InBuildTarget.GlobalCompileEnvironment.Config.SystemIncludePaths.Add("$(NDKROOT)/sources/cxx-stl/gnu-libstdc++/" + GccVersion + "/include");
-				InBuildTarget.GlobalCompileEnvironment.Config.SystemIncludePaths.Add("$(NDKROOT)/sources/cxx-stl/gnu-libstdc++/" + GccVersion + "/libs/armeabi-v7a/include");
+				InBuildTarget.GlobalCompileEnvironment.Config.CPPIncludeInfo.SystemIncludePaths.Add("$(NDKROOT)/sources/cxx-stl/gnu-libstdc++/" + GccVersion + "/include");
+				InBuildTarget.GlobalCompileEnvironment.Config.CPPIncludeInfo.SystemIncludePaths.Add("$(NDKROOT)/sources/cxx-stl/gnu-libstdc++/" + GccVersion + "/libs/armeabi-v7a/include");
 
                 InBuildTarget.GlobalLinkEnvironment.Config.LibraryPaths.Add("$(NDKROOT)/sources/cxx-stl/gnu-libstdc++/" + GccVersion + "/libs/armeabi-v7a");
 			}
 			else
 			{
-				InBuildTarget.GlobalCompileEnvironment.Config.SystemIncludePaths.Add("$(NDKROOT)/sources/cxx-stl/gnu-libstdc++/" + GccVersion + "/include");
-				InBuildTarget.GlobalCompileEnvironment.Config.SystemIncludePaths.Add("$(NDKROOT)/sources/cxx-stl/gnu-libstdc++/" + GccVersion + "/libs/x86/include");
+				InBuildTarget.GlobalCompileEnvironment.Config.CPPIncludeInfo.SystemIncludePaths.Add("$(NDKROOT)/sources/cxx-stl/gnu-libstdc++/" + GccVersion + "/include");
+				InBuildTarget.GlobalCompileEnvironment.Config.CPPIncludeInfo.SystemIncludePaths.Add("$(NDKROOT)/sources/cxx-stl/gnu-libstdc++/" + GccVersion + "/libs/x86/include");
 
                 InBuildTarget.GlobalLinkEnvironment.Config.LibraryPaths.Add("$(NDKROOT)/sources/cxx-stl/gnu-libstdc++/" + GccVersion + "/libs/x86");
             }
 
-			InBuildTarget.GlobalCompileEnvironment.Config.SystemIncludePaths.Add("$(NDKROOT)/sources/android/native_app_glue");
-			InBuildTarget.GlobalCompileEnvironment.Config.SystemIncludePaths.Add("$(NDKROOT)/sources/android/cpufeatures");
+			InBuildTarget.GlobalCompileEnvironment.Config.CPPIncludeInfo.SystemIncludePaths.Add("$(NDKROOT)/sources/android/native_app_glue");
+			InBuildTarget.GlobalCompileEnvironment.Config.CPPIncludeInfo.SystemIncludePaths.Add("$(NDKROOT)/sources/android/cpufeatures");
 
             // Add path to statically compiled version of cxa_demangle
             InBuildTarget.GlobalLinkEnvironment.Config.LibraryPaths.Add(UEBuildConfiguration.UEThirdPartySourceDirectory + "Android/cxa_demangle/armeabi-v7a");
diff --git a/Engine/Source/Programs/UnrealBuildTool/Configuration/BuildConfiguration.cs b/Engine/Source/Programs/UnrealBuildTool/Configuration/BuildConfiguration.cs
--- a/Engine/Source/Programs/UnrealBuildTool/Configuration/BuildConfiguration.cs
+++ b/Engine/Source/Programs/UnrealBuildTool/Configuration/BuildConfiguration.cs
@@ -356,6 +356,23 @@ namespace UnrealBuildTool
 		/// </summary>
 		[XmlConfig]
 		public static bool bUseExperimentalFastDependencyScan;
+
+		/// <summary>
+		/// Enables a new 'experimental' support for very fast iterative builds.  Turning this on causes Unreal Build Tool to emit
+		/// 'UBT Makefiles' for targets when they're built the first time.  Subsequent builds will load these Makefiles and begin
+		/// outdatedness checking and build invocation very quickly.  The caveat is that if source files are added or removed to
+		/// the project, UBT will need to gather information about those in order for your build to complete successfully.  Currently,
+		/// you must run the project file generator after adding/removing source files to tell UBT to re-gather this information.
+		/// Events that can invalidate the 'UBT Makefile':  
+		///		- Adding/removing .cpp files
+		///		- Adding/removing .h files with UObjects
+		///		- Adding new UObject types to a file that didn't previously have any
+		///		- Changing global build settings (most settings in this file qualify.)
+		///		- Changed code that affects how Unreal Header Tool works
+		///	You can force regeneration of the 'UBT Makefile' by passing the '-gather' argument, or simply regenerating project files
+		/// </summary>
+		[XmlConfig]
+		public static bool bUseExperimentalFastBuildIteration;
 	
 		/// <summary>
 		/// Whether DMUCS/Distcc may be used.
@@ -482,9 +499,14 @@ namespace UnrealBuildTool
 			bCheckLicenseViolations = false;
 			bBreakBuildOnLicenseViolation = true;
 
-			// The new experimental fast dependency scanning is still being tested, so we disable it by default
+			// This feature reduces the time spent tracking C++ header dependencies every run.  This feature is still being tested, 
+			// so we disable it by default
 			bUseExperimentalFastDependencyScan = false;
 
+			// Experimental support for 'gathering' and 'assembling' builds in separate invocations is still being tested, so we
+			// leave it disabled by default
+			bUseExperimentalFastBuildIteration = false;
+
 			// Distcc requires some setup - so by default disable it so we don't break local or remote building
 			bAllowDistcc = false;
 			bAllowDistccLocalFallback = true;
diff --git a/Engine/Source/Programs/UnrealBuildTool/Configuration/UEBuildConfiguration.cs b/Engine/Source/Programs/UnrealBuildTool/Configuration/UEBuildConfiguration.cs
--- a/Engine/Source/Programs/UnrealBuildTool/Configuration/UEBuildConfiguration.cs
+++ b/Engine/Source/Programs/UnrealBuildTool/Configuration/UEBuildConfiguration.cs
@@ -156,10 +156,6 @@ namespace UnrealBuildTool
 		[XmlConfig]
 		public static bool bUseLoggingInShipping;
 
-		/** True if plugins should be excluded when building target. */
-		[XmlConfig]
-		public static bool bExcludePlugins;
-
         /** True if we need to package up Android with the OBB in the APK file */
 		[XmlConfig]
         public static bool bOBBinAPK;
diff --git a/Engine/Source/Programs/UnrealBuildTool/Configuration/UEBuildModule.cs b/Engine/Source/Programs/UnrealBuildTool/Configuration/UEBuildModule.cs
--- a/Engine/Source/Programs/UnrealBuildTool/Configuration/UEBuildModule.cs
+++ b/Engine/Source/Programs/UnrealBuildTool/Configuration/UEBuildModule.cs
@@ -1110,8 +1110,8 @@ namespace UnrealBuildTool
 			{
 				// Update project file's set of preprocessor definitions and include paths
 				IntelliSenseGatherer.AddIntelliSensePreprocessorDefinitions( ModuleCompileEnvironment.Config.Definitions );
-				IntelliSenseGatherer.AddInteliiSenseIncludePaths( ModuleCompileEnvironment.Config.SystemIncludePaths, bAddingSystemIncludes: true );
-				IntelliSenseGatherer.AddInteliiSenseIncludePaths( ModuleCompileEnvironment.Config.IncludePaths, bAddingSystemIncludes: false );
+				IntelliSenseGatherer.AddInteliiSenseIncludePaths( ModuleCompileEnvironment.Config.CPPIncludeInfo.SystemIncludePaths, bAddingSystemIncludes: true );
+				IntelliSenseGatherer.AddInteliiSenseIncludePaths( ModuleCompileEnvironment.Config.CPPIncludeInfo.IncludePaths, bAddingSystemIncludes: false );
 
 				// Bail out.  We don't need to actually compile anything while generating project files.
 				return LinkInputFiles;
@@ -1158,7 +1158,7 @@ namespace UnrealBuildTool
 				// Make sure our RC files have cached includes.  
 				foreach( var RCFile in SourceFilesToBuild.RCFiles )
 				{ 
-					RCFile.CachedCPPEnvironment = ModuleCompileEnvironment;
+					RCFile.CachedCPPIncludeInfo = ModuleCompileEnvironment.Config.CPPIncludeInfo;
 				}
 			}
 
@@ -1453,7 +1453,7 @@ namespace UnrealBuildTool
 
 				CachePCHUsageForModuleSourceFile( this.Target, CPPCompileEnvironment, GeneratedCppFileItem );
 
-				// @todo fastubt: Check for ALL other places where we might be injecting .cpp or .rc files for compiling without caching CachedCPPEnvironment first (anything platform specfic?)
+				// @todo fastubt: Check for ALL other places where we might be injecting .cpp or .rc files for compiling without caching CachedCPPIncludeInfo first (anything platform specfic?)
 				LinkInputFiles.AddRange( CPPCompileEnvironment.CompileFiles( Target, new List<FileItem>{ GeneratedCppFileItem }, Name ).ObjectFiles );
 			}
 
@@ -1539,13 +1539,13 @@ namespace UnrealBuildTool
 			var PCHCacheTimerStart = DateTime.UtcNow;
 
 			var BuildPlatform = UEBuildPlatform.GetBuildPlatformForCPPTargetPlatform( ModuleCompileEnvironment.Config.Target.Platform );
-			var IncludePathsToSearch = ModuleCompileEnvironment.GetIncludesPathsToSearch( CPPFile );
+			var IncludePathsToSearch = ModuleCompileEnvironment.Config.CPPIncludeInfo.GetIncludesPathsToSearch( CPPFile );
 
 			// Store the module compile environment along with the .cpp file.  This is so that we can use it later on when looking
 			// for header dependencies
-			CPPFile.CachedCPPEnvironment = ModuleCompileEnvironment;
+			CPPFile.CachedCPPIncludeInfo = ModuleCompileEnvironment.Config.CPPIncludeInfo;
 
-			var PCHFile = CachePCHUsageForCPPFile( Target, CPPFile, BuildPlatform, IncludePathsToSearch, ModuleCompileEnvironment.IncludeFileSearchDictionary );
+			var PCHFile = CachePCHUsageForCPPFile( Target, CPPFile, BuildPlatform, IncludePathsToSearch, ModuleCompileEnvironment.Config.CPPIncludeInfo.IncludeFileSearchDictionary );
 
 			if( BuildConfiguration.bPrintPerformanceInfo )
 			{
@@ -1568,17 +1568,17 @@ namespace UnrealBuildTool
 				bool bFoundAProblemWithPCHs = false;
 
 				FileItem UniquePCH = null;
-				foreach( var CPPFile in SourceFilesFound.CPPFiles )
+				foreach( var CPPFile in SourceFilesFound.CPPFiles )	// @todo fastubt: We're not caching CPPEnvironments for .c/.mm files, etc.  Even though they don't use PCHs, they still have #includes!  This can break dependency checking!
 				{
 					// Build a single list of include paths to search.
-					var IncludePathsToSearch = ModuleCompileEnvironment.GetIncludesPathsToSearch( CPPFile );   // @todo fastubt: This call could be made faster (list caching)
+					var IncludePathsToSearch = ModuleCompileEnvironment.Config.CPPIncludeInfo.GetIncludesPathsToSearch( CPPFile );   // @todo fastubt: This call could be made faster (list caching)
 
 					// Store the module compile environment along with the .cpp file.  This is so that we can use it later on when looking
 					// for header dependencies
-					CPPFile.CachedCPPEnvironment = ModuleCompileEnvironment;
+					CPPFile.CachedCPPIncludeInfo = ModuleCompileEnvironment.Config.CPPIncludeInfo;
 
 					// Find headers used by the source file.
-					var PCH = UEBuildModuleCPP.CachePCHUsageForCPPFile(Target, CPPFile, BuildPlatform, IncludePathsToSearch, ModuleCompileEnvironment.IncludeFileSearchDictionary);
+					var PCH = UEBuildModuleCPP.CachePCHUsageForCPPFile(Target, CPPFile, BuildPlatform, IncludePathsToSearch, ModuleCompileEnvironment.Config.CPPIncludeInfo.IncludeFileSearchDictionary);
 					if( PCH == null )
 					{
 						throw new BuildException( "Source file \"{0}\" is not including any headers.  We expect all modules to include a header file for precompiled header generation.  Please add an #include statement.", CPPFile.AbsolutePath );
@@ -1603,7 +1603,7 @@ namespace UnrealBuildTool
 				{
 					// Map from pch header string to the source files that use that PCH
 					var UsageMapPCH = new Dictionary<string, List<FileItem>>( StringComparer.InvariantCultureIgnoreCase );
-					foreach( var CPPFile in SourceFilesToBuild.CPPFiles )
+					foreach( var CPPFile in SourceFilesToBuild.CPPFiles )	// @todo fastubt: Using SourceFilesToBuild here, but SourceFilesFound up above.  Why are there two arrays again?
 					{
 						// Create a new entry if not in the pch usage map
 						UsageMapPCH.GetOrAddNew( CPPFile.PrecompiledHeaderIncludeFilename ).Add( CPPFile );
@@ -1741,7 +1741,7 @@ namespace UnrealBuildTool
 			Result.Config.Definitions.AddRange(Definitions);
 
 			// Setup the compile environment for the module.
-			SetupPrivateCompileEnvironment(ref Result.Config.IncludePaths, ref Result.Config.SystemIncludePaths, ref Result.Config.Definitions, ref Result.Config.AdditionalFrameworks);
+			SetupPrivateCompileEnvironment(ref Result.Config.CPPIncludeInfo.IncludePaths, ref Result.Config.CPPIncludeInfo.SystemIncludePaths, ref Result.Config.Definitions, ref Result.Config.AdditionalFrameworks);
 
 			return Result;
 		}
@@ -1771,9 +1771,9 @@ namespace UnrealBuildTool
 
 			// Add all of the include paths
 			List<string> IncludePaths = new List<string>();
-			IncludePaths.AddRange(CompileEnvironment.Config.SystemIncludePaths);
+			IncludePaths.AddRange(CompileEnvironment.Config.CPPIncludeInfo.SystemIncludePaths);
 			IncludePaths.AddRange(PrivateSystemIncludePaths);
-			IncludePaths.AddRange(CompileEnvironment.Config.IncludePaths);
+			IncludePaths.AddRange(CompileEnvironment.Config.CPPIncludeInfo.IncludePaths);
 			IncludePaths.AddRange(PrivateIncludePaths);
 
 			// Build the full list of definitions
@@ -2083,11 +2083,26 @@ namespace UnrealBuildTool
 			else if (Plugins.IsPluginModule(ModuleName))
 			{
 				// Plugin module
-				BaseDirectory = Plugins.GetPluginInfoForModule(ModuleName).IntermediateIncPath;
+				string PluginIntermediateIncPath;
+				{ 
+					PluginInfo Plugin = Plugins.GetPluginInfoForModule(ModuleName);
+					if (Plugin.LoadedFrom == PluginInfo.LoadedFromType.Engine)
+					{
+						// Plugin folder is in the engine directory
+						PluginIntermediateIncPath = Path.GetFullPath(Path.Combine(BuildConfiguration.RelativeEnginePath, BuildConfiguration.PlatformIntermediateFolder));
+					}
+					else
+					{
+						// Plugin folder is in the project directory
+						PluginIntermediateIncPath   = Path.GetFullPath(Path.Combine( Target.ProjectDirectory, BuildConfiguration.PlatformIntermediateFolder));
+					}
+					PluginIntermediateIncPath = Path.Combine(PluginIntermediateIncPath, "Inc", "Plugins");
+				}
+				BaseDirectory = PluginIntermediateIncPath;
 			}
 			else
 			{
-				var AllProjectFolders = UEBuildTarget.DiscoverAllGameFolders();
+				var AllProjectFolders = UEBuildTarget.DiscoverAllGameFolders();		// @todo fastubt: This will be called again and again for every UObject module (80+)
 				BaseDirectory = AllProjectFolders.Find(ProjectFolder => Utils.IsFileUnderDirectory( ModuleDirectory, ProjectFolder ));
 				if (BaseDirectory == null)
 				{
@@ -2228,4 +2243,95 @@ namespace UnrealBuildTool
 		public string ResourcePath			= null;
 		public string BundleContentsSubdir	= null;
 	}
+
+	public class PrecompileHeaderEnvironment
+	{
+		/** The name of the module this PCH header is a member of */
+		public readonly string ModuleName;
+
+		/** PCH header file name as it appears in an #include statement in source code (might include partial, or no relative path.)
+			This is needed by some compilers to use PCH features. */
+		public string PCHHeaderNameInCode;
+
+		/** The source header file that this precompiled header will be generated for */
+		public readonly FileItem PrecompiledHeaderIncludeFilename;
+
+		/** Whether this precompiled header will be built with CLR features enabled.  We won't mix and match CLR PCHs with non-CLR PCHs */
+		public readonly CPPCLRMode CLRMode;
+
+		/** Whether this precompiled header will be built with code optimization enabled. */
+		public readonly ModuleRules.CodeOptimization OptimizeCode;
+
+		/** The PCH file we're generating */
+		public FileItem PrecompiledHeaderFile = null;
+
+		/** Object files emitted from the compiler when generating this precompiled header.  These will be linked into modules that
+			include this PCH */
+		public readonly List<FileItem> OutputObjectFiles = new List<FileItem>();
+
+		public PrecompileHeaderEnvironment( string InitModuleName, string InitPCHHeaderNameInCode, FileItem InitPrecompiledHeaderIncludeFilename, CPPCLRMode InitCLRMode, ModuleRules.CodeOptimization InitOptimizeCode )
+		{
+			ModuleName = InitModuleName;
+			PCHHeaderNameInCode = InitPCHHeaderNameInCode;
+			PrecompiledHeaderIncludeFilename = InitPrecompiledHeaderIncludeFilename;
+			CLRMode = InitCLRMode;
+			OptimizeCode = InitOptimizeCode;
+		}
+
+		/// <summary>
+		/// Creates a precompiled header action to generate a new pch file 
+		/// </summary>
+		/// <param name="PCHHeaderNameInCode">The precompiled header name as it appeared in an #include statement</param>
+		/// <param name="PrecompiledHeaderIncludeFilename">Name of the header used for pch.</param>
+		/// <param name="ProjectCPPEnvironment">The environment the C/C++ files in the project are compiled with.</param>
+		/// <param name="OutputDirectory">The folder to save the generated PCH file to</param>
+		/// <param name="ModuleName">Name of the module this PCH is being generated for</param>
+		/// <param name="bAllowDLLExports">True if we should allow DLLEXPORT definitions for this PCH</param>
+		/// <returns>the compilation output result of the created pch.</returns>
+		public static CPPOutput GeneratePCHCreationAction(UEBuildTarget Target, string PCHHeaderNameInCode, FileItem PrecompiledHeaderIncludeFilename, CPPEnvironment ProjectCPPEnvironment, string OutputDirectory, string ModuleName, bool bAllowDLLExports )
+		{
+			// Find the header file to be precompiled. Don't skip external headers
+			if (PrecompiledHeaderIncludeFilename.bExists)
+			{
+				// Create a Dummy wrapper around the PCH to avoid problems with #pragma once on clang
+				var ToolChain = UEToolChain.GetPlatformToolChain(ProjectCPPEnvironment.Config.Target.Platform);
+				string PCHGuardDefine = Path.GetFileNameWithoutExtension(PrecompiledHeaderIncludeFilename.AbsolutePath).ToUpper();
+				string LocalPCHHeaderNameInCode = ToolChain.ConvertPath(PrecompiledHeaderIncludeFilename.AbsolutePath);
+				string TmpPCHHeaderContents = String.Format("#ifndef __AUTO_{0}_H__\n#define __AUTO_{0}_H__\n//Last Write: {2}\n#include \"{1}\"\n#endif//__AUTO_{0}_H__", PCHGuardDefine, LocalPCHHeaderNameInCode, PrecompiledHeaderIncludeFilename.LastWriteTime);
+				string DummyPath = Path.Combine(
+					ProjectCPPEnvironment.Config.OutputDirectory,
+					Path.GetFileName(PrecompiledHeaderIncludeFilename.AbsolutePath));
+				FileItem DummyPCH = FileItem.CreateIntermediateTextFile(DummyPath, TmpPCHHeaderContents);
+
+				// Create a new C++ environment that is used to create the PCH.
+				var ProjectPCHEnvironment = ProjectCPPEnvironment.DeepCopy();
+				ProjectPCHEnvironment.Config.PrecompiledHeaderAction = PrecompiledHeaderAction.Create;
+				ProjectPCHEnvironment.Config.PrecompiledHeaderIncludeFilename = PrecompiledHeaderIncludeFilename.AbsolutePath;
+				ProjectPCHEnvironment.Config.PCHHeaderNameInCode = PCHHeaderNameInCode;
+				ProjectPCHEnvironment.Config.OutputDirectory = OutputDirectory;
+
+				if( !bAllowDLLExports )
+				{
+					for( var CurDefinitionIndex = 0; CurDefinitionIndex < ProjectPCHEnvironment.Config.Definitions.Count; ++CurDefinitionIndex )
+					{
+						// We change DLLEXPORT to DLLIMPORT for "shared" PCH headers
+						var OldDefinition = ProjectPCHEnvironment.Config.Definitions[ CurDefinitionIndex ];
+						if( OldDefinition.EndsWith( "=DLLEXPORT" ) )
+						{
+							ProjectPCHEnvironment.Config.Definitions[ CurDefinitionIndex ] = OldDefinition.Replace( "DLLEXPORT", "DLLIMPORT" );
+						}
+					}
+				}
+
+				// Cache our CPP environment so that we can check for outdatedness quickly.  Only files that have includes need this.
+				DummyPCH.CachedCPPIncludeInfo = ProjectPCHEnvironment.Config.CPPIncludeInfo;
+
+				Log.TraceVerbose( "Found PCH file \"{0}\".", PrecompiledHeaderIncludeFilename );
+
+				// Create the action to compile the PCH file.
+				return ProjectPCHEnvironment.CompileFiles(Target, new List<FileItem>() { DummyPCH }, ModuleName);
+			}
+			throw new BuildException( "Couldn't find PCH file \"{0}\".", PrecompiledHeaderIncludeFilename );
+		}
+	}
 }
diff --git a/Engine/Source/Programs/UnrealBuildTool/Configuration/UEBuildTarget.cs b/Engine/Source/Programs/UnrealBuildTool/Configuration/UEBuildTarget.cs
--- a/Engine/Source/Programs/UnrealBuildTool/Configuration/UEBuildTarget.cs
+++ b/Engine/Source/Programs/UnrealBuildTool/Configuration/UEBuildTarget.cs
@@ -215,6 +215,7 @@ namespace UnrealBuildTool
 				{
 					switch (Arguments[ArgumentIndex].ToUpperInvariant())
 					{
+						// @todo ubtmake: How does this work with UBTMake?  Do we even need to support it anymore?
 						case "-MODULE":
 							// Specifies a module to recompile.  Can be specified more than once on the command-line to compile multiple specific modules.
 							{
@@ -625,7 +626,7 @@ namespace UnrealBuildTool
 		/// <returns>true if it should, false if it shouldn't</returns>
 		public bool ShouldCompileMonolithic()
 		{
-			return bCompileMonolithic;
+			return bCompileMonolithic;	// @todo ubtmake: We need to make sure this function and similar things aren't called in assembler mode
 		}
 
 		/** 
@@ -767,7 +768,7 @@ namespace UnrealBuildTool
 		/// Attempts to delete a file. Will retry a few times before failing.
 		/// </summary>
 		/// <param name="Filename"></param>
-		void CleanFile(string Filename)
+		public static void CleanFile(string Filename)
 		{
 			const int RetryDelayStep = 200;
 			int RetryDelay = 1000;
@@ -997,6 +998,20 @@ namespace UnrealBuildTool
 					}
 				}
 
+				// Delete the UBT makefile
+				{
+					// Figure out what to call our action graph based on everything we're building
+					var Targets = new List<UEBuildTarget>();
+					Targets.Add( this );	// @todo ubtmake: Only supports cleaning one target at a time :(
+
+					var UBTMakefilePath = UnrealBuildTool.GetUBTMakefilePath( Targets );
+					if (File.Exists(UBTMakefilePath))
+					{
+						Log.TraceVerbose("\tDeleting " + UBTMakefilePath);
+						CleanFile(UBTMakefilePath);
+					}
+				}
+
 				// Delete the action history
 				{
 					var ActionHistoryFilename = ActionHistory.GeneratePathForTarget(this);
@@ -1205,8 +1220,11 @@ namespace UnrealBuildTool
 		}
 
 		/** Builds the target, appending list of output files and returns building result. */
-		public ECompilationResult Build(IUEToolChain TargetToolChain, List<FileItem> OutputItems, out string EULAViolationWarning)
+		public ECompilationResult Build(IUEToolChain TargetToolChain, out List<FileItem> OutputItems, out List<UHTModuleInfo> UObjectModules, out string EULAViolationWarning)
 		{
+			OutputItems = new List<FileItem>();
+			UObjectModules = new List<UHTModuleInfo>();
+
 			var SpecialRocketLibFilesThatAreBuildProducts = PreBuildSetup();
 
 			EULAViolationWarning = !ProjectFileGenerator.bGenerateProjectFiles
@@ -1309,21 +1327,19 @@ namespace UnrealBuildTool
 			// just an optimization
 			if( !ProjectFileGenerator.bGenerateProjectFiles )
 			{
-				var SharedPCHHeaderFiles = FindSharedPCHHeaders();
-				GlobalCompileEnvironment.SharedPCHHeaderFiles = SharedPCHHeaderFiles;
+				GlobalCompileEnvironment.SharedPCHHeaderFiles = FindSharedPCHHeaders();
 			}
 
 			if( (BuildConfiguration.bXGEExport && UEBuildConfiguration.bGenerateManifest) || (!UEBuildConfiguration.bGenerateManifest && !UEBuildConfiguration.bCleanProject && !ProjectFileGenerator.bGenerateProjectFiles) )
 			{
 				var UObjectDiscoveryStartTime = DateTime.UtcNow;
 
-				var UObjectModules = new List<UHTModuleInfo>();
-
 				// Figure out which modules have UObjects that we may need to generate headers for
 				foreach( var Binary in AppBinaries )
 				{
+					var LocalUObjectModules = UObjectModules;	// For lambda access
 					var DependencyModules = Binary.GetAllDependencyModules(bIncludeDynamicallyLoaded: false, bForceCircular: false);
-					foreach( var DependencyModuleCPP in DependencyModules.OfType<UEBuildModuleCPP>().Where( CPPModule => !UObjectModules.Any( Module => Module.ModuleName == CPPModule.Name ) ) )
+					foreach( var DependencyModuleCPP in DependencyModules.OfType<UEBuildModuleCPP>().Where( CPPModule => !LocalUObjectModules.Any( Module => Module.ModuleName == CPPModule.Name ) ) )
 					{
 						if( !DependencyModuleCPP.bIncludedInTarget )
 						{
@@ -1344,6 +1360,21 @@ namespace UnrealBuildTool
 
 							DependencyModuleCPP.AutoGenerateCppInfo = new UEBuildModuleCPP.AutoGenerateCppInfoClass(BuildInfo);
 
+							// If we're running in "gather" mode only, we'll go ahead and cache PCH information for each module right now, so that we don't
+							// have to do it in the assembling phase.  It's OK for gathering to take a bit longer, even if UObject headers are not out of
+							// date in order to save a lot of time in the assembling runs.
+							UHTModuleInfo.PCH = "";
+							if( UnrealBuildTool.IsGatheringBuild && !UnrealBuildTool.IsAssemblingBuild )
+							{
+								// We need to figure out which PCH header this module is including, so that UHT can inject an include statement for it into any .cpp files it is synthesizing
+								var ModuleCompileEnvironment = DependencyModuleCPP.CreateModuleCompileEnvironment(GlobalCompileEnvironment);
+								DependencyModuleCPP.CachePCHUsageForModuleSourceFiles(ModuleCompileEnvironment);
+								if (DependencyModuleCPP.ProcessedDependencies.UniquePCHHeaderFile != null)
+								{
+									UHTModuleInfo.PCH = DependencyModuleCPP.ProcessedDependencies.UniquePCHHeaderFile.AbsolutePath;
+								}
+							}
+
 							UObjectModules.Add( UHTModuleInfo );
 							Log.TraceVerbose( "Detected UObject module: " + UHTModuleInfo.ModuleName );
 						}
@@ -1356,16 +1387,16 @@ namespace UnrealBuildTool
 					Trace.TraceInformation( "UObject discovery time: " + UObjectDiscoveryTime + "s" );
 				}
 
+				// @todo ubtmake: Even in Gather mode, we need to run UHT to make sure the files exist for the static action graph to be setup correctly.  This is because UHT generates .cpp
+				// files that are injected as top level prerequisites.  If UHT only emitted included header files, we wouldn't need to run it during the Gather phase at all.
 				if( UObjectModules.Count > 0 )
 				{
 					// Execute the header tool
-					var TargetName = !String.IsNullOrEmpty( GameName ) ? GameName : AppName;
 					string ModuleInfoFileName = Path.GetFullPath( Path.Combine( ProjectIntermediateDirectory, "UnrealHeaderTool.manifest" ) );
-
 					ECompilationResult UHTResult = ECompilationResult.OtherCompilationError;
 					if (!ExternalExecution.ExecuteHeaderToolIfNecessary(this, GlobalCompileEnvironment, UObjectModules, ModuleInfoFileName, ref UHTResult))
 					{
-						Log.TraceInformation("UnrealHeaderTool failed for target '" + TargetName + "' (platform: " + Platform.ToString() + ", module info: " + ModuleInfoFileName + ").");
+						Log.TraceInformation("UnrealHeaderTool failed for target '" + GetTargetName() + "' (platform: " + Platform.ToString() + ", module info: " + ModuleInfoFileName + ").");
 						return UHTResult;
 					}
 				}
@@ -1385,7 +1416,7 @@ namespace UnrealBuildTool
 				OutputItems.AddRange(Binary.Build(TargetToolChain, GlobalCompileEnvironment, GlobalLinkEnvironment));
 			}
 
-			if (BuildConfiguration.WriteTargetInfoPath != null)
+			if (BuildConfiguration.WriteTargetInfoPath != null)	// @todo ubtmake: This has some overlap with what we are doing with the action graph.  It might be easier to support other build systems with formatting like this, but it is a lot of duplicated code
 			{
 				Log.TraceInformation("Writing build environment to " + BuildConfiguration.WriteTargetInfoPath + "...");
 				WriteTargetInfo();
@@ -1953,10 +1984,26 @@ namespace UnrealBuildTool
 			var ModuleSourceFolder = bHasModuleRules ? Path.GetDirectoryName(RulesCompiler.GetModuleFilename(Module.Name)) : ModuleFilename;
 			bool bShouldBeBuiltFromSource = bHasModuleRules && Directory.GetFiles(ModuleSourceFolder, "*.cpp", SearchOption.AllDirectories).Length > 0;
 
+			string PluginIntermediateBuildPath;
+			{ 
+				if (Plugin.LoadedFrom == PluginInfo.LoadedFromType.Engine)
+				{
+					// Plugin folder is in the engine directory
+					var PluginConfiguration = Configuration == UnrealTargetConfiguration.DebugGame ? UnrealTargetConfiguration.Development : Configuration;
+					PluginIntermediateBuildPath = Path.GetFullPath(Path.Combine(BuildConfiguration.RelativeEnginePath, BuildConfiguration.PlatformIntermediateFolder, AppName, PluginConfiguration.ToString()));
+				}
+				else
+				{
+					// Plugin folder is in the project directory
+					PluginIntermediateBuildPath = Path.GetFullPath(Path.Combine(ProjectDirectory, BuildConfiguration.PlatformIntermediateFolder, GetTargetName(), Configuration.ToString()));
+				}
+				PluginIntermediateBuildPath = Path.Combine(PluginIntermediateBuildPath, "Plugins", ShouldCompileMonolithic() ? "Static" : "Dynamic");
+			}
+
 			// Create the binary
 			UEBuildBinaryConfiguration Config = new UEBuildBinaryConfiguration( InType:                  BinaryType,
 																				InOutputFilePath:        OutputFilePath,
-																				InIntermediateDirectory: Plugin.IntermediateBuildPath,
+																				InIntermediateDirectory: PluginIntermediateBuildPath,
 																				bInAllowExports:         true,
 																				bInAllowCompilation:     bShouldBeBuiltFromSource,
 																				bInHasModuleRules:       bHasModuleRules,
@@ -2156,62 +2203,42 @@ namespace UnrealBuildTool
 		/** Sets up the plugins for this target */
 		protected virtual void SetupPlugins()
 		{
-			if (!UEBuildConfiguration.bExcludePlugins)
+			// Filter the plugins list by the current project
+			List<PluginInfo> ValidPlugins = new List<PluginInfo>();
+			foreach(PluginInfo Plugin in Plugins.AllPlugins)
 			{
-				// Filter the plugins list by the current project
-				List<PluginInfo> ValidPlugins = new List<PluginInfo>();
-				foreach(PluginInfo Plugin in Plugins.AllPlugins)
+				if(Plugin.LoadedFrom != PluginInfo.LoadedFromType.GameProject || Plugin.Directory.StartsWith(ProjectDirectory, StringComparison.InvariantCultureIgnoreCase))
 				{
-					if(Plugin.LoadedFrom != PluginInfo.LoadedFromType.GameProject || Plugin.Directory.StartsWith(ProjectDirectory, StringComparison.InvariantCultureIgnoreCase))
-					{
-						if (Plugin.LoadedFrom == PluginInfo.LoadedFromType.Engine)
-						{
-							// Plugin folder is in the engine directory
-							var PluginConfiguration = Configuration == UnrealTargetConfiguration.DebugGame ? UnrealTargetConfiguration.Development : Configuration;
-							Plugin.IntermediateIncPath   = Path.GetFullPath(Path.Combine(BuildConfiguration.RelativeEnginePath, BuildConfiguration.PlatformIntermediateFolder));
-							Plugin.IntermediateBuildPath = Path.GetFullPath(Path.Combine(BuildConfiguration.RelativeEnginePath, BuildConfiguration.PlatformIntermediateFolder, AppName, PluginConfiguration.ToString()));
-						}
-						else
-						{
-							// Plugin folder is in the project directory
-							Plugin.IntermediateIncPath   = Path.GetFullPath(Path.Combine(ProjectDirectory, BuildConfiguration.PlatformIntermediateFolder));
-							Plugin.IntermediateBuildPath = Path.GetFullPath(Path.Combine(ProjectDirectory, BuildConfiguration.PlatformIntermediateFolder, GetTargetName(), Configuration.ToString()));
-						}
+					ValidPlugins.Add(Plugin);
+				}
+			}
 
-						Plugin.IntermediateIncPath   = Path.Combine(Plugin.IntermediateIncPath,   "Inc", "Plugins");
-						Plugin.IntermediateBuildPath = Path.Combine(Plugin.IntermediateBuildPath, "Plugins", ShouldCompileMonolithic() ? "Static" : "Dynamic");
+			// Build the enabled plugin list
+			if (ShouldCompileMonolithic() || Rules.Type == TargetRules.TargetType.Program)
+			{
+				var FilterPluginNames = new List<string>(Rules.AdditionalPlugins);
 
-						ValidPlugins.Add(Plugin);
-					}
+				// Add the list of plugins enabled by default
+				if (UEBuildConfiguration.bCompileAgainstEngine)
+				{
+					FilterPluginNames.AddRange(ValidPlugins.Where(x => x.bEnabledByDefault).Select(x => x.Name));
 				}
 
-				// Build the enabled plugin list
-				if (ShouldCompileMonolithic() || Rules.Type == TargetRules.TargetType.Program)
+				// Update the plugin list for game targets
+				if(Rules.Type != TargetRules.TargetType.Program && UnrealBuildTool.HasUProjectFile())
 				{
-					var FilterPluginNames = new List<string>(Rules.AdditionalPlugins);
-
-					// Add the list of plugins enabled by default
-					if (UEBuildConfiguration.bCompileAgainstEngine)
-					{
-						FilterPluginNames.AddRange(ValidPlugins.Where(x => x.bEnabledByDefault).Select(x => x.Name));
-					}
+					// Enable all the game specific plugins by default
+					FilterPluginNames.AddRange(ValidPlugins.Where(x => x.LoadedFrom == PluginInfo.LoadedFromType.GameProject).Select(x => x.Name));
 
-					// Update the plugin list for game targets
-					if(Rules.Type != TargetRules.TargetType.Program && UnrealBuildTool.HasUProjectFile())
-					{
-						// Enable all the game specific plugins by default
-						FilterPluginNames.AddRange(ValidPlugins.Where(x => x.LoadedFrom == PluginInfo.LoadedFromType.GameProject).Select(x => x.Name));
-
-						// Use the project settings to update the plugin list for this target
-						FilterPluginNames = UProjectInfo.GetEnabledPlugins(UnrealBuildTool.GetUProjectFile(), FilterPluginNames, Platform);
-					}
-
-					EnabledPlugins.AddRange(ValidPlugins.Where(x => FilterPluginNames.Contains(x.Name)).Distinct());
-				}
-				else
-				{
-					EnabledPlugins.AddRange(ValidPlugins);
+					// Use the project settings to update the plugin list for this target
+					FilterPluginNames = UProjectInfo.GetEnabledPlugins(UnrealBuildTool.GetUProjectFile(), FilterPluginNames, Platform);
 				}
+
+				EnabledPlugins.AddRange(ValidPlugins.Where(x => FilterPluginNames.Contains(x.Name)).Distinct());
+			}
+			else
+			{
+				EnabledPlugins.AddRange(ValidPlugins);
 			}
 		}
 
@@ -2319,7 +2346,7 @@ namespace UnrealBuildTool
 			{
 				throw new BuildException( "Couldn't find Engine/Source directory using relative path" );
 			}
-			GlobalCompileEnvironment.Config.IncludePaths.Add( EngineSourceDirectory );
+			GlobalCompileEnvironment.Config.CPPIncludeInfo.IncludePaths.Add( EngineSourceDirectory );
 
 			[REDACTED_EMAIL]: Do any platform specific tool chain initialization here if required
 
diff --git a/Engine/Source/Programs/UnrealBuildTool/HTML5/HTML5ToolChain.cs b/Engine/Source/Programs/UnrealBuildTool/HTML5/HTML5ToolChain.cs
--- a/Engine/Source/Programs/UnrealBuildTool/HTML5/HTML5ToolChain.cs
+++ b/Engine/Source/Programs/UnrealBuildTool/HTML5/HTML5ToolChain.cs
@@ -318,11 +318,11 @@ namespace UnrealBuildTool
 			CPPOutput Result = new CPPOutput();
 
 			// Add include paths to the argument list.
-			foreach (string IncludePath in CompileEnvironment.Config.IncludePaths)
+			foreach (string IncludePath in CompileEnvironment.Config.CPPIncludeInfo.IncludePaths)
 			{
 				Arguments += string.Format(" -I\"{0}\"", IncludePath);
 			}
-			foreach (string IncludePath in CompileEnvironment.Config.SystemIncludePaths)
+			foreach (string IncludePath in CompileEnvironment.Config.CPPIncludeInfo.SystemIncludePaths)
 			{
 				Arguments += string.Format(" -I\"{0}\"", IncludePath);
 			}
@@ -348,11 +348,7 @@ namespace UnrealBuildTool
 				bool bIsPlainCFile = Path.GetExtension(SourceFile.AbsolutePath).ToUpperInvariant() == ".C";
                 
 				// Add the C++ source file and its included files to the prerequisite item list.
-				CompileAction.PrerequisiteItems.Add(SourceFile);
-				{
-					var IncludedFileList = CPPEnvironment.FindAndCacheAllIncludedFiles( Target, SourceFile, BuildPlatform, CompileEnvironment.GetIncludesPathsToSearch( SourceFile ), CompileEnvironment.IncludeFileSearchDictionary, bOnlyCachedDependencies:BuildConfiguration.bUseExperimentalFastDependencyScan );
-					CompileAction.PrerequisiteItems.AddRange( IncludedFileList );
-				}
+				AddPrerequisiteSourceFile( Target, BuildPlatform, CompileEnvironment, SourceFile, CompileAction.PrerequisiteItems );
 
 				// Add the source file path to the command-line.
                 string FileArguments = string.Format(" \"{0}\"", SourceFile.AbsolutePath);
@@ -384,7 +380,6 @@ namespace UnrealBuildTool
 
                 System.Console.WriteLine(CompileAction.CommandArguments); 
 				CompileAction.StatusDescription = Path.GetFileName(SourceFile.AbsolutePath);
-				CompileAction.StatusDetailedDescription = SourceFile.Description;
 				CompileAction.OutputEventHandler = new DataReceivedEventHandler(CompileOutputReceivedDataEventHandler);
 
 				// Don't farm out creation of precomputed headers as it is the critical path task.
diff --git a/Engine/Source/Programs/UnrealBuildTool/IOS/IOSToolChain.cs b/Engine/Source/Programs/UnrealBuildTool/IOS/IOSToolChain.cs
--- a/Engine/Source/Programs/UnrealBuildTool/IOS/IOSToolChain.cs
+++ b/Engine/Source/Programs/UnrealBuildTool/IOS/IOSToolChain.cs
@@ -462,8 +462,8 @@ namespace UnrealBuildTool
 			}
 
 			// Add include paths to the argument list.
-			List<string> AllIncludes = CompileEnvironment.Config.IncludePaths;
-			AllIncludes.AddRange(CompileEnvironment.Config.SystemIncludePaths);
+			List<string> AllIncludes = CompileEnvironment.Config.CPPIncludeInfo.IncludePaths;
+			AllIncludes.AddRange(CompileEnvironment.Config.CPPIncludeInfo.SystemIncludePaths);
 			foreach (string IncludePath in AllIncludes)
 			{
 				Arguments += string.Format(" -I\"{0}\"", ConvertPath(Path.GetFullPath(IncludePath)));
@@ -540,26 +540,7 @@ namespace UnrealBuildTool
 				}
 
 				// Add the C++ source file and its included files to the prerequisite item list.
-				CompileAction.PrerequisiteItems.Add(SourceFile);
-
-				if (BuildHostPlatform.Current.Platform != UnrealTargetPlatform.Mac)
-				{
-					QueueFileForBatchUpload(SourceFile);
-				}
-
-				{
-					var IncludedFileList = CPPEnvironment.FindAndCacheAllIncludedFiles( Target, SourceFile, BuildPlatform, CompileEnvironment.GetIncludesPathsToSearch( SourceFile ), CompileEnvironment.IncludeFileSearchDictionary, bOnlyCachedDependencies:BuildConfiguration.bUseExperimentalFastDependencyScan );
-					foreach (FileItem IncludedFile in IncludedFileList)
-					{
-						CompileAction.PrerequisiteItems.Add(IncludedFile);
-
-						if (!BuildConfiguration.bUseExperimentalFastDependencyScan &&	// With fast dependency scanning, we will not have an exhaustive list of dependencies here.  We rely on PostCodeGeneration() to upload these files.
-							BuildHostPlatform.Current.Platform != UnrealTargetPlatform.Mac)
-						{
-							QueueFileForBatchUpload(IncludedFile);
-						}
-					}
-				}
+				AddPrerequisiteSourceFile( Target, BuildPlatform, CompileEnvironment, SourceFile, CompileAction.PrerequisiteItems );
 
 				if (CompileEnvironment.Config.PrecompiledHeaderAction == PrecompiledHeaderAction.Create)
 				{
@@ -615,7 +596,6 @@ namespace UnrealBuildTool
 				CompileAction.CommandPath = CompilerPath;
 				CompileAction.CommandArguments = Arguments + FileArguments + CompileEnvironment.Config.AdditionalArguments;
 				CompileAction.StatusDescription = string.Format("{0}", Path.GetFileName(SourceFile.AbsolutePath));
-				CompileAction.StatusDetailedDescription = SourceFile.Description;
 				CompileAction.bIsGCCCompiler = true;
 				// We're already distributing the command by execution on Mac.
 				CompileAction.bCanExecuteRemotely = false;
diff --git a/Engine/Source/Programs/UnrealBuildTool/Linux/LinuxToolChain.cs b/Engine/Source/Programs/UnrealBuildTool/Linux/LinuxToolChain.cs
--- a/Engine/Source/Programs/UnrealBuildTool/Linux/LinuxToolChain.cs
+++ b/Engine/Source/Programs/UnrealBuildTool/Linux/LinuxToolChain.cs
@@ -383,11 +383,11 @@ namespace UnrealBuildTool
             }
 
             // Add include paths to the argument list.
-            foreach (string IncludePath in CompileEnvironment.Config.IncludePaths)
+            foreach (string IncludePath in CompileEnvironment.Config.CPPIncludeInfo.IncludePaths)
             {
                 Arguments += string.Format(" -I\"{0}\"", IncludePath);
             }
-            foreach (string IncludePath in CompileEnvironment.Config.SystemIncludePaths)
+            foreach (string IncludePath in CompileEnvironment.Config.CPPIncludeInfo.SystemIncludePaths)
             {
                 Arguments += string.Format(" -I\"{0}\"", IncludePath);
             }
@@ -441,12 +441,8 @@ namespace UnrealBuildTool
                     FileArguments += PCHArguments;
                 }
 
-                // Add the C++ source file and its included files to the prerequisite item list.
-                CompileAction.PrerequisiteItems.Add(SourceFile);
-				{
-					var IncludedFileList = CPPEnvironment.FindAndCacheAllIncludedFiles( Target, SourceFile, BuildPlatform, CompileEnvironment.GetIncludesPathsToSearch( SourceFile ), CompileEnvironment.IncludeFileSearchDictionary, bOnlyCachedDependencies:BuildConfiguration.bUseExperimentalFastDependencyScan );
-					CompileAction.PrerequisiteItems.AddRange( IncludedFileList );
-				}
+				// Add the C++ source file and its included files to the prerequisite item list.
+				AddPrerequisiteSourceFile( Target, BuildPlatform, CompileEnvironment, SourceFile, CompileAction.PrerequisiteItems );
 
                 if (CompileEnvironment.Config.PrecompiledHeaderAction == PrecompiledHeaderAction.Create)
                 {
@@ -502,7 +498,6 @@ namespace UnrealBuildTool
                 CompileAction.CommandArguments = Arguments + FileArguments + CompileEnvironment.Config.AdditionalArguments;
                 CompileAction.CommandDescription = "Compile";
                 CompileAction.StatusDescription = Path.GetFileName(SourceFile.AbsolutePath);
-                CompileAction.StatusDetailedDescription = SourceFile.Description;
                 CompileAction.bIsGCCCompiler = true;
 
                 // Don't farm out creation of pre-compiled headers as it is the critical path task.
diff --git a/Engine/Source/Programs/UnrealBuildTool/Mac/MacToolChain.cs b/Engine/Source/Programs/UnrealBuildTool/Mac/MacToolChain.cs
--- a/Engine/Source/Programs/UnrealBuildTool/Mac/MacToolChain.cs
+++ b/Engine/Source/Programs/UnrealBuildTool/Mac/MacToolChain.cs
@@ -245,8 +245,8 @@ namespace UnrealBuildTool
 			}
 
 			// Add include paths to the argument list.
-			List<string> AllIncludes = CompileEnvironment.Config.IncludePaths;
-			AllIncludes.AddRange(CompileEnvironment.Config.SystemIncludePaths);
+			List<string> AllIncludes = CompileEnvironment.Config.CPPIncludeInfo.IncludePaths;
+			AllIncludes.AddRange(CompileEnvironment.Config.CPPIncludeInfo.SystemIncludePaths);
 			foreach (string IncludePath in AllIncludes)
 			{
 				Arguments.Append(" -I\"");
@@ -333,26 +333,7 @@ namespace UnrealBuildTool
 				}
 
 				// Add the C++ source file and its included files to the prerequisite item list.
-				CompileAction.PrerequisiteItems.Add(SourceFile);
-
-				if (BuildHostPlatform.Current.Platform != UnrealTargetPlatform.Mac)
-				{
-					QueueFileForBatchUpload(SourceFile);
-				}
-
-				{
-					var IncludedFileList = CPPEnvironment.FindAndCacheAllIncludedFiles( Target, SourceFile, BuildPlatform, CompileEnvironment.GetIncludesPathsToSearch( SourceFile ), CompileEnvironment.IncludeFileSearchDictionary, bOnlyCachedDependencies:BuildConfiguration.bUseExperimentalFastDependencyScan );
-					foreach (FileItem IncludedFile in IncludedFileList)
-					{
-						CompileAction.PrerequisiteItems.Add(IncludedFile);
-
-						if (!BuildConfiguration.bUseExperimentalFastDependencyScan &&	// With fast dependency scanning, we will not have an exhaustive list of dependencies here.  We rely on PostCodeGeneration() to upload these files.
-							 BuildHostPlatform.Current.Platform != UnrealTargetPlatform.Mac)
-						{
-							QueueFileForBatchUpload(IncludedFile);
-						}
-					}
-				}
+				AddPrerequisiteSourceFile( Target, BuildPlatform, CompileEnvironment, SourceFile, CompileAction.PrerequisiteItems );
 
 				if (CompileEnvironment.Config.PrecompiledHeaderAction == PrecompiledHeaderAction.Create)
 				{
@@ -407,7 +388,6 @@ namespace UnrealBuildTool
 				CompileAction.CommandArguments = MacCompiler + Arguments + FileArguments + CompileEnvironment.Config.AdditionalArguments;
 				CompileAction.CommandDescription = "Compile";
 				CompileAction.StatusDescription = Path.GetFileName(SourceFile.AbsolutePath);
-				CompileAction.StatusDetailedDescription = SourceFile.Description;
 				CompileAction.bIsGCCCompiler = true;
 				// We're already distributing the command by execution on Mac.
 				CompileAction.bCanExecuteRemotely = false;
diff --git a/Engine/Source/Programs/UnrealBuildTool/System/ActionGraph.cs b/Engine/Source/Programs/UnrealBuildTool/System/ActionGraph.cs
--- a/Engine/Source/Programs/UnrealBuildTool/System/ActionGraph.cs
+++ b/Engine/Source/Programs/UnrealBuildTool/System/ActionGraph.cs
@@ -25,70 +25,116 @@ namespace UnrealBuildTool
 	}
 
 	/** A build action. */
+	[Serializable]
 	public class Action
 	{
-        /** Unique action identifier */
-        public int UniqueId;
+		///
+		/// Preparation and Assembly (serialized)
+		/// 
+
+		/** The type of this action (for debugging purposes). */
+		public ActionType ActionType;
 
+		/** Every file this action depends on.  These files need to exist and be up to date in order for this action to even be considered */
 		public List<FileItem> PrerequisiteItems = new List<FileItem>();
+
+		/** The files that this action produces after completing */
 		public List<FileItem> ProducedItems = new List<FileItem>();
-		public delegate void EventHandler(Action A);
 
-		/** Total number of actions depending on this one. */
-		public int NumTotalDependentActions = 0;
-		/** Relative cost of producing items for this action. */
-		public long RelativeCost = 0;
+		/** Directory from which to execute the program to create produced items */
 		public string WorkingDirectory = null;
+
+		/** True if we should log extra information when we run a program to create produced items */
         public bool bPrintDebugInfo = false;
+
+		/** The command to run to create produced items */
         public string CommandPath = null;
+
+		/** Command-line parameters to pass to the program */
 		public string CommandArguments = null;
+
+		/** Optional friendly description of the type of command being performed, for example "Compile" or "Link".  Displayed by some executors. */
 		public string CommandDescription = null;
+
+		/** Human-readable description of this action that may be displayed as status while invoking the action.  This is often the name of the file being compiled, or an executable file name being linked.  Displayed by some executors. */
 		public string StatusDescription = "...";
-		public string StatusDetailedDescription = "";
+
+		/** True if this action is allowed to be run on a remote machine when a distributed build system is being used, such as XGE */
 		public bool bCanExecuteRemotely = false;
+
+		/** True if this action is using the Visual C++ compiler.  Some build systems may be able to optimize for this case. */
 		public bool bIsVCCompiler = false;
+
+		/** True if this action is using the GCC compiler.  Some build systems may be able to optimize for this case. */
 		public bool bIsGCCCompiler = false;
+
 		/** Whether the action is using a pre-compiled header to speed it up. */
 		public bool bIsUsingPCH = false;
-		/** Whether to delete the prerequisite files on completion */
-		public bool bShouldDeletePrereqs = false;
+		
 		/** Whether the files in ProducedItems should be deleted before executing this action, when the action is outdated */
 		public bool bShouldDeleteProducedItems = false;
+
 		/**
 		 * Whether we should log this action, whether executed locally or remotely.  This is useful for actions that take time
 		 * but invoke tools without any console output.
 		 */
 		public bool bShouldOutputStatusDescription = true;
-		/** True if we should redirect standard input such that input will only come from the builder (which is none) */
-		public bool bShouldBlockStandardInput = false;
-		/** True if we should redirect standard output such that text will not be logged */
-		public bool bShouldBlockStandardOutput = false;
-		/** Start time of action, optionally set by executor. */
-		public DateTimeOffset StartTime = DateTimeOffset.MinValue;
-		/** End time of action, optionally set by executor. */
-		public DateTimeOffset EndTime = DateTimeOffset.MinValue;
+
+		/** True if any libraries produced by this action should be considered 'import libraries' */
+		public bool bProducesImportLibrary = false;
+
 		/** Optional custom event handler for standard output. */
-		public DataReceivedEventHandler OutputEventHandler = null;
+		public DataReceivedEventHandler OutputEventHandler = null;	// @todo ubtmake urgent: Delegate variables are not saved, but we are comparing against this in ExecutActions() for XGE!
+
 		/** Callback used to perform a special action instead of a generic command line */
 		public delegate void BlockingActionHandler(Action Action, out int ExitCode, out string Output);
-		public BlockingActionHandler ActionHandler = null;
-		/** The type of this action (for debugging purposes). */
-		public ActionType ActionType;
+		public BlockingActionHandler ActionHandler = null;		// @todo ubtmake urgent: Delegate variables are not saved, but we are comparing against this in ExecutActions() for XGE!
 
-		/** True if any libraries produced by this action should be considered 'import libraries' */
-		public bool bProducesImportLibrary = false;
+
+
+		///
+		/// Preparation only (not serialized)
+		///
+
+		/** Unique action identifier.  Used for displaying helpful info about detected cycles in the graph. */
+		[NonSerialized]
+        public int UniqueId;
 
         /** Always-incremented unique id */
         private static int NextUniqueId = 0;
 
+		/** Total number of actions depending on this one. */
+		[NonSerialized]
+		public int NumTotalDependentActions = 0;
+		
+		/** Relative cost of producing items for this action. */
+		[NonSerialized]
+		public long RelativeCost = 0;
+
+
+		///
+		/// Assembly only (not serialized)
+		///
+
+		/** Start time of action, optionally set by executor. */
+		[NonSerialized]
+		public DateTimeOffset StartTime = DateTimeOffset.MinValue;
+		
+		/** End time of action, optionally set by executor. */
+		[NonSerialized]
+		public DateTimeOffset EndTime = DateTimeOffset.MinValue;
+		
+
+
 		public Action(ActionType InActionType)
 		{
 			ActionType = InActionType;
 
-			UnrealBuildTool.AllActions.Add(this);
+			ActionGraph.AllActions.Add(this);
             UniqueId = ++NextUniqueId;
 		}
 
+
 		/**
 		 * Compares two actions based on total number of dependent items, descending.
 		 * 
@@ -145,7 +191,7 @@ namespace UnrealBuildTool
 		}
 	};
 
-	partial class UnrealBuildTool
+	class ActionGraph
 	{
 		public static List<Action> AllActions = new List<Action>();
 
@@ -154,38 +200,28 @@ namespace UnrealBuildTool
 			AllActions = new List<Action>();
 		}
 
-		/** Builds a list of actions that need to be executed to produce the specified output items. */
-		static List<Action> GetActionsToExecute(List<FileItem> OutputItems, List<UEBuildTarget> Targets, out Dictionary<UEBuildTarget, List<FileItem>> TargetToOutdatedPrerequisitesMap)
+		public static void FinalizeActionGraph()
 		{
-			// @todo fastubt: We really want to be able to generate actions in a separate phase from actually building.  For example, we could generate
-			// actions at GPF-time, then save them out to be reloaded quickly later when building.  It means we would separate C++ include dependencies
-			// out from the normal action graph and check outdatedness at build time.  The module relationships and link dependencies would all still
-			// be in the static part of the action graph that was loaded from disk.	 You would probably need to re-run GPF after adding or removing source files,
-			// but ideally not after changing which headers are included in an existing source file!  Maybe we could make adding/removing source files work
-			// too, but it affects the Unity/PCH/Linker input, so it would be easier to not deal with this.  Still need to run UHT during build phase
-			// as well, which means we need to know which files have UObjects for the manifest.  This could be cached up front, and kept up to date as
-			// we rescan modified source files looking for changed includes.  UHT does not require spidering includes.
-
 			// @todo fastubt: Can we use directory changed notifications or directory timestamps to accelerate C++ file outdatedness checking?
 			
 			// Link producing actions to the items they produce.
 			LinkActionsAndItems();
 
-			DeleteStaleHotReloadDLLs();
-
 			// Detect cycles in the action graph.
 			DetectActionGraphCycles();
 
 			// Sort action list by "cost" in descending order to improve parallelism.
 			SortActionList();
+		}
 
+		/** Builds a list of actions that need to be executed to produce the specified output items. */
+		public static List<Action> GetActionsToExecute(Action[] PrerequisiteActions, List<UEBuildTarget> Targets, out Dictionary<UEBuildTarget, List<FileItem>> TargetToOutdatedPrerequisitesMap)
+		{
 			// Build a set of all actions needed for this target.
-			var ActionsNeededForThisTarget = new Dictionary<Action, bool>();
-
-			// For now simply treat all object files as the root target.
-			foreach (FileItem OutputItem in OutputItems)
+			var IsActionOutdatedMap = new Dictionary<Action, bool>();
+			foreach (var Action in PrerequisiteActions)
 			{
-				GatherPrerequisiteActions(OutputItem, ref ActionsNeededForThisTarget);
+				IsActionOutdatedMap.Add( Action, true );
 			}
 			
 			// For all targets, build a set of all actions that are outdated.
@@ -193,10 +229,10 @@ namespace UnrealBuildTool
 			var HistoryList = new List<ActionHistory>();
 			var OpenHistoryFiles = new HashSet<string>(StringComparer.InvariantCultureIgnoreCase);
 			TargetToOutdatedPrerequisitesMap = new Dictionary<UEBuildTarget,List<FileItem>>();
-			foreach (var BuildTarget in Targets)
+			foreach (var BuildTarget in Targets)	// @todo ubtmake: Optimization: Ideally we don't even need to know about targets for ubtmake -- everything would come from the files
 			{
 				var HistoryFilename = ActionHistory.GeneratePathForTarget(BuildTarget);
-				if (!OpenHistoryFiles.Contains(HistoryFilename))
+				if (!OpenHistoryFiles.Contains(HistoryFilename))		// @todo ubtmake: Optimization: We should be able to move the command-line outdatedness and build product deletion over to the 'gather' phase, as the command-lines won't change between assembler runs
 				{
 					var History = new ActionHistory(HistoryFilename);
 					HistoryList.Add(History);
@@ -223,7 +259,7 @@ namespace UnrealBuildTool
 			List<Action> ActionsToExecute = new List<Action>();
 			foreach (Action Action in AllActions)
 			{
-				if (Action.CommandPath != null && ActionsNeededForThisTarget.ContainsKey(Action) && OutdatedActionDictionary[Action])
+				if (Action.CommandPath != null && IsActionOutdatedMap.ContainsKey(Action) && OutdatedActionDictionary[Action])
 				{
 					ActionsToExecute.Add(Action);
 				}
@@ -233,7 +269,7 @@ namespace UnrealBuildTool
 		}
 
 		/** Executes a list of actions. */
-		static bool ExecuteActions(List<Action> ActionsToExecute, out string ExecutorName)
+		public static bool ExecuteActions(List<Action> ActionsToExecute, out string ExecutorName)
 		{
 			bool Result = true;
 			bool bUsedXGE = false;
@@ -332,25 +368,6 @@ namespace UnrealBuildTool
 				Log.TraceInformation("Target is up to date.");
 			}
 
-			// Perform any cleanup
-			foreach (Action Action in ActionsToExecute)
-			{
-				if (Action.bShouldDeletePrereqs)
-				{
-					foreach (FileItem FileItem in Action.PrerequisiteItems)
-					{
-						if (bUsedXGE && BuildConfiguration.bXGEExport)
-						{
-							throw new BuildException("We are exporting XGE with a request to delete prerequisites; we need a delete prerequisites thing or just roll this into the XGE actions.");
-						}
-						else
-						{
-							FileItem.Delete();
-						}
-					}
-				}
-			}
-
 			return Result;
 		}
 
@@ -403,8 +420,10 @@ namespace UnrealBuildTool
 			}
 			return WorkingString;
 		}
+
+
 		/** Finds and deletes stale hot reload DLLs. */
-		static void DeleteStaleHotReloadDLLs()
+		public static void DeleteStaleHotReloadDLLs()
 		{
 			foreach (Action BuildAction in AllActions)
 			{
@@ -612,16 +631,16 @@ namespace UnrealBuildTool
 		 * @param OutputItem - The item to be built.
 		 * @param PrerequisiteActions - The actions that must be built and the root action are 
 		 */
-		static void GatherPrerequisiteActions(
+		public static void GatherPrerequisiteActions(
 			FileItem OutputItem,
-			ref Dictionary<Action, bool> PrerequisiteActions
+			ref HashSet<Action> PrerequisiteActions
 			)
 		{
 			if (OutputItem != null && OutputItem.ProducingAction != null)
 			{
-				if (!PrerequisiteActions.ContainsKey(OutputItem.ProducingAction))
+				if (!PrerequisiteActions.Contains(OutputItem.ProducingAction))
 				{
-					PrerequisiteActions.Add(OutputItem.ProducingAction, true);
+					PrerequisiteActions.Add(OutputItem.ProducingAction);
 					foreach (FileItem PrerequisiteItem in OutputItem.ProducingAction.PrerequisiteItems)
 					{
 						GatherPrerequisiteActions(PrerequisiteItem, ref PrerequisiteActions);
@@ -710,6 +729,72 @@ namespace UnrealBuildTool
 				Log.WriteLineIf(BuildConfiguration.bLogDetailedActionStats && !String.IsNullOrEmpty( LatestUpdatedProducedItemName ),
 					TraceEventType.Verbose, "{0}: Oldest produced item is {1}", RootAction.StatusDescription, LatestUpdatedProducedItemName);
 
+				bool bFindCPPIncludePrerequisites = false;
+				if( RootAction.ActionType == ActionType.Compile )
+				{
+					// Up to date targets don't need their headers scanned, because presumably they would already be out of dated based on already-cached includes before getting this far
+					if( !bIsOutdated && 
+						BuildConfiguration.bUseExperimentalFastBuildIteration && 
+						!UnrealBuildTool.IsGatheringBuild && UnrealBuildTool.IsAssemblingBuild && 
+						RootAction.ActionType == ActionType.Compile )
+					{
+						bFindCPPIncludePrerequisites = true;
+					}
+
+					// Were we asked to force an update of our cached includes BEFORE we try to build?  This may be needed if our cache can no longer
+					// be trusted and we need to fill it with perfectly valid data (even if we're in assembler only mode)
+					if( BuildConfiguration.bUseExperimentalFastDependencyScan && 
+						UnrealBuildTool.bNeedsFullCPPIncludeRescan )
+					{
+						bFindCPPIncludePrerequisites = true;
+					}
+				}
+
+
+				if( bFindCPPIncludePrerequisites )
+				{
+					// Scan this file for included headers that may be out of date.  Note that it's OK if we break out early because we found
+					// the action to be outdated.  For outdated actions, we kick off a separate include scan in a background thread later on to
+					// catch all of the other includes and form an exhaustive set.
+					foreach (FileItem PrerequisiteItem in RootAction.PrerequisiteItems)
+					{
+						// @todo ubtmake: Make sure we are catching RC files here too.  Anything that the toolchain would have tried it on.  Logic should match the CACHING stuff below
+						if( PrerequisiteItem.CachedCPPIncludeInfo != null )
+						{
+							var BuildPlatform = UEBuildPlatform.GetBuildPlatform( Target.GetTargetInfo().Platform );
+							var IncludedFileList = CPPEnvironment.FindAndCacheAllIncludedFiles( Target, PrerequisiteItem, BuildPlatform, PrerequisiteItem.CachedCPPIncludeInfo, bOnlyCachedDependencies:BuildConfiguration.bUseExperimentalFastDependencyScan );
+							foreach( var IncludedFile in IncludedFileList )	// @todo fastubt: @todo ubtmake: Optimization: This is "retesting" a lot of the same files over and over in a single run (common indirect includes)
+							{
+								if( IncludedFile.bExists )
+								{
+									// allow a 1 second slop for network copies
+									TimeSpan TimeDifference = IncludedFile.LastWriteTime - LastExecutionTime;
+									bool bPrerequisiteItemIsNewerThanLastExecution = TimeDifference.TotalSeconds > 1;
+									if (bPrerequisiteItemIsNewerThanLastExecution)
+									{
+										Log.TraceVerbose(
+											"{0}: Included file {1} is newer than the last execution of the action: {2} vs {3}",
+											RootAction.StatusDescription,
+											Path.GetFileName(IncludedFile.AbsolutePath),
+											IncludedFile.LastWriteTime.LocalDateTime,
+											LastExecutionTime.LocalDateTime
+											);
+										bIsOutdated = true;
+
+										// Don't bother checking every single include if we've found one that is out of date
+										break;
+									}
+								}
+							}
+						}
+
+						if( bIsOutdated )
+						{
+							break;
+						}
+					}
+				}
+
 				if(!bIsOutdated)
 				{
 					// Check if any of the prerequisite items are produced by outdated actions, or have changed more recently than
@@ -768,7 +853,7 @@ namespace UnrealBuildTool
 						}
 					}
 				}
-
+				
 				// For compile actions, we have C++ files that are actually dependent on header files that could have been changed.  We only need to
 				// know about the set of header files that are included for files that are already determined to be out of date (such as if the file
 				// is missing or was modified.)  In the case that the file is out of date, we'll perform a deep scan to update our cached set of
@@ -778,12 +863,12 @@ namespace UnrealBuildTool
 					var DeepIncludeScanStartTime = DateTime.UtcNow;
 
 					// @todo fastubt: we may be scanning more files than we need to here -- indirectly outdated files are bIsOutdated=true by this point (for example basemost includes when deeper includes are dirty)
-					if( bIsOutdated && RootAction.ActionType == ActionType.Compile )
+					if( bIsOutdated && RootAction.ActionType == ActionType.Compile )	// @todo fastubt: Does this work with RC files?  See above too.
 					{
 						Log.TraceVerbose( "Outdated action: {0}", RootAction.StatusDescription );
 						foreach (FileItem PrerequisiteItem in RootAction.PrerequisiteItems)
 						{
-							if( PrerequisiteItem.CachedCPPEnvironment != null )
+							if( PrerequisiteItem.CachedCPPIncludeInfo != null )
 							{
 								if( !IsCPPFile( PrerequisiteItem ) )
 								{
@@ -802,7 +887,7 @@ namespace UnrealBuildTool
 							}
 							else if( IsCPPImplementationFile( PrerequisiteItem ) || IsCPPResourceFile( PrerequisiteItem ) )
 							{
-								if( PrerequisiteItem.CachedCPPEnvironment == null )
+								if( PrerequisiteItem.CachedCPPIncludeInfo == null )
 								{
 									Log.TraceVerbose( "  -> WARNING: No CachedCPPEnvironment: {0}", PrerequisiteItem.AbsolutePath );
 								}
@@ -813,7 +898,7 @@ namespace UnrealBuildTool
 					if( BuildConfiguration.bPrintPerformanceInfo )
 					{
 						double DeepIncludeScanTime = ( DateTime.UtcNow - DeepIncludeScanStartTime ).TotalSeconds;
-						TotalDeepIncludeScanTime += DeepIncludeScanTime;
+						UnrealBuildTool.TotalDeepIncludeScanTime += DeepIncludeScanTime;
 					}
 				}
 
diff --git a/Engine/Source/Programs/UnrealBuildTool/System/CPPEnvironment.cs b/Engine/Source/Programs/UnrealBuildTool/System/CPPEnvironment.cs
--- a/Engine/Source/Programs/UnrealBuildTool/System/CPPEnvironment.cs
+++ b/Engine/Source/Programs/UnrealBuildTool/System/CPPEnvironment.cs
@@ -227,13 +227,7 @@ namespace UnrealBuildTool
 		public CPPCLRMode CLRMode = CPPCLRMode.CLRDisabled;
 
 		/** The include paths to look for included files in. */
-		public List<string> IncludePaths = new List<string>();
-
-		/**
-		 * The include paths where changes to contained files won't cause dependent C++ source files to
-		 * be recompiled, unless BuildConfiguration.bCheckSystemHeadersForModification==true.
-		 */
-		public List<string> SystemIncludePaths = new List<string>();
+		public readonly CPPIncludeInfo CPPIncludeInfo = new CPPIncludeInfo();
 
 		/** Paths where .NET framework assembly references are found, when compiling CLR applications. */
 		public List<string> SystemDotNetAssemblyPaths = new List<string>();
@@ -277,8 +271,8 @@ namespace UnrealBuildTool
 			bIsBuildingLibrary                     = InCopyEnvironment.bIsBuildingLibrary;
 			bIsBuildingDLL                         = InCopyEnvironment.bIsBuildingDLL;
 			CLRMode                                = InCopyEnvironment.CLRMode;
-			IncludePaths                 .AddRange(InCopyEnvironment.IncludePaths);
-			SystemIncludePaths           .AddRange(InCopyEnvironment.SystemIncludePaths);
+			CPPIncludeInfo.IncludePaths                 .AddRange(InCopyEnvironment.CPPIncludeInfo.IncludePaths);
+			CPPIncludeInfo.SystemIncludePaths           .AddRange(InCopyEnvironment.CPPIncludeInfo.SystemIncludePaths);
 			SystemDotNetAssemblyPaths    .AddRange(InCopyEnvironment.SystemDotNetAssemblyPaths);
 			FrameworkAssemblyDependencies.AddRange(InCopyEnvironment.FrameworkAssemblyDependencies);
  			PrivateAssemblyDependencies  .AddRange(InCopyEnvironment.PrivateAssemblyDependencies);
@@ -388,25 +382,5 @@ namespace UnrealBuildTool
 		}
 
 
-		/// <summary>
-		/// Give a C++ source file, returns a list of include paths we should search to resolve #includes for this path
-		/// </summary>
-		/// <param name="SourceFile">C++ source file we're going to check #includes for.</param>
-		/// <returns>Ordered list of paths to search</returns>
-		public List<string> GetIncludesPathsToSearch( FileItem SourceFile )
-		{
-			// Build a single list of include paths to search.
-			var IncludePathsToSearch = new List<string>();
-			string SourceFilesDirectory = Path.GetDirectoryName( SourceFile.AbsolutePath);
-			IncludePathsToSearch.Add( SourceFilesDirectory );
-			IncludePathsToSearch.AddRange( Config.IncludePaths );
-			if( BuildConfiguration.bCheckSystemHeadersForModification )
-			{
-				IncludePathsToSearch.AddRange( Config.SystemIncludePaths );
-			}
-			
-			return IncludePathsToSearch;
-		}
-
 	};
 }
diff --git a/Engine/Source/Programs/UnrealBuildTool/System/CPPHeaders.cs b/Engine/Source/Programs/UnrealBuildTool/System/CPPHeaders.cs
--- a/Engine/Source/Programs/UnrealBuildTool/System/CPPHeaders.cs
+++ b/Engine/Source/Programs/UnrealBuildTool/System/CPPHeaders.cs
@@ -9,6 +9,52 @@ using System.IO;
 
 namespace UnrealBuildTool
 {
+	/** For C++ source file items, this structure is used to cache data that will be used for include dependency scanning */
+	[Serializable]
+	public class CPPIncludeInfo
+	{
+		/** Ordered list of include paths for the module  */
+		public List<string> IncludePaths = new List<string>();
+
+		/**
+		 * The include paths where changes to contained files won't cause dependent C++ source files to
+		 * be recompiled, unless BuildConfiguration.bCheckSystemHeadersForModification==true.
+		 */
+		public List<string> SystemIncludePaths = new List<string>();
+
+		/** Contains a mapping from filename to the full path of the header in this environment.  This is used to optimized include path lookups at runtime for any given single module. */
+		[NonSerialized]
+		public Dictionary<string, FileItem> IncludeFileSearchDictionary = new Dictionary<string, FileItem>();
+
+
+		/// <summary>
+		/// Given a C++ source file, returns a list of include paths we should search to resolve #includes for this path
+		/// </summary>
+		/// <param name="SourceFile">C++ source file we're going to check #includes for.</param>
+		/// <returns>Ordered list of paths to search</returns>
+		public List<string> GetIncludesPathsToSearch( FileItem SourceFile )
+		{
+			// Build a single list of include paths to search.
+			var IncludePathsToSearch = new List<string>();
+			
+			if( SourceFile != null )
+			{ 
+				string SourceFilesDirectory = Path.GetDirectoryName( SourceFile.AbsolutePath );
+				IncludePathsToSearch.Add( SourceFilesDirectory );
+			}
+
+			IncludePathsToSearch.AddRange( IncludePaths );
+			if( BuildConfiguration.bCheckSystemHeadersForModification )
+			{
+				IncludePathsToSearch.AddRange( SystemIncludePaths );
+			}
+			
+			return IncludePathsToSearch;
+		}
+	}
+
+
+
 	/// <summary>
 	/// List of all files included in a file and helper class for handling circular dependencies.
 	/// </summary>
@@ -21,6 +67,7 @@ namespace UnrealBuildTool
 		public List<FileItem> CircularDependencies = new List<FileItem>();
 	}
 
+
 	public partial class CPPEnvironment
 	{
 		/** Contains a cache of include dependencies (direct and indirect), one for each target we're building. */
@@ -29,10 +76,8 @@ namespace UnrealBuildTool
 		/** Contains a cache of include dependencies (direct and indirect), one for each target we're building. */
 		public static readonly Dictionary<UEBuildTarget, FlatCPPIncludeDependencyCache> FlatCPPIncludeDependencyCache = new Dictionary<UEBuildTarget, FlatCPPIncludeDependencyCache>();
 
-		/** Contains a mapping from filename to the full path of the header in this environment. */
-		public readonly Dictionary<string, FileItem> IncludeFileSearchDictionary = new Dictionary<string, FileItem>();
-
 		public static int TotalFindIncludedFileCalls = 0;
+
 		public static int IncludePathSearchAttempts = 0;
 
 
@@ -107,7 +152,7 @@ namespace UnrealBuildTool
 				IncludeFileSearchDictionary.Add( InvariantPath, Result );
 			}
 
-			// @todo fastubt: The old UBT tried to skip 'external' headers here.  But it didn't work.  We might want to do this though!  Skip system headers and source/thirdparty headers!
+			// @todo fastubt: The old UBT tried to skip 'external' (STABLE) headers here.  But it didn't work.  We might want to do this though!  Skip system headers and source/thirdparty headers!
 
 			if( Result != null )
 			{
@@ -128,11 +173,20 @@ namespace UnrealBuildTool
 		static Dictionary<FileItem, IncludedFilesSet> OnlyCachedIncludedFilesMap = new Dictionary<FileItem, IncludedFilesSet>();
 
 
-		public static List<FileItem> FindAndCacheAllIncludedFiles( UEBuildTarget Target, FileItem SourceFile, IUEBuildPlatform BuildPlatform, List<string> IncludePathsToSearch, Dictionary<string, FileItem> IncludeFileSearchDictionary, bool bOnlyCachedDependencies )
+		public static List<FileItem> FindAndCacheAllIncludedFiles( UEBuildTarget Target, FileItem SourceFile, IUEBuildPlatform BuildPlatform, CPPIncludeInfo CPPIncludeInfo, bool bOnlyCachedDependencies )
 		{
 			var Result = new List<FileItem>();
+			
+			if( CPPIncludeInfo.IncludeFileSearchDictionary == null )
+			{
+				CPPIncludeInfo.IncludeFileSearchDictionary = new Dictionary<string,FileItem>();
+			}
+
+			bool bUseFlatCPPIncludeDependencyCache =
+				( BuildConfiguration.bUseExperimentalFastDependencyScan &&
+				  ( !BuildConfiguration.bUseExperimentalFastBuildIteration || UnrealBuildTool.IsAssemblingBuild ) );
 
-			if( BuildConfiguration.bUseExperimentalFastDependencyScan && bOnlyCachedDependencies )
+			if( bUseFlatCPPIncludeDependencyCache && bOnlyCachedDependencies )
 			{ 
 				var Dependencies = FlatCPPIncludeDependencyCache[Target].GetDependenciesForFile( SourceFile.AbsolutePath );
 				if( Dependencies != null )
@@ -152,17 +206,17 @@ namespace UnrealBuildTool
 				// @todo fastubt: HeaderParser.h is missing from the include set for Module.UnrealHeaderTool.cpp (failed to find include using:  FileItem DirectIncludeResolvedFile = CPPEnvironment.FindIncludedFile(DirectInclude.IncludeName, !BuildConfiguration.bCheckExternalHeadersForModification, IncludePathsToSearch, IncludeFileSearchDictionary );)
 
 				var IncludedFileList = new IncludedFilesSet();
-				CPPEnvironment.FindAndCacheAllIncludedFiles( Target, SourceFile, BuildPlatform, IncludePathsToSearch, IncludeFileSearchDictionary, ref IncludedFileList, bOnlyCachedDependencies:bOnlyCachedDependencies );
+				CPPEnvironment.FindAndCacheAllIncludedFiles( Target, SourceFile, BuildPlatform, CPPIncludeInfo, ref IncludedFileList, bOnlyCachedDependencies:bOnlyCachedDependencies );
 				foreach( FileItem IncludedFile in IncludedFileList )
 				{
 					Result.Add( IncludedFile );
 				}
 			
 				// Update cache
-				if( BuildConfiguration.bUseExperimentalFastDependencyScan && !bOnlyCachedDependencies )
+				if( bUseFlatCPPIncludeDependencyCache && !bOnlyCachedDependencies )
 				{ 
 					var Dependencies = new List<string>();
-					foreach( var IncludedFile in Result )	// @todo fastubt: Have to convert Hash to List here, try to avoid this.  Other locations too.
+					foreach( var IncludedFile in Result )
 					{
 						Dependencies.Add( IncludedFile.AbsolutePath );
 					}
@@ -181,7 +235,7 @@ namespace UnrealBuildTool
 		/// <param name="CPPFile">C++ file to get the dependencies for.</param>
 		/// <param name="Result">List of CPPFile dependencies.</param>
 		/// <returns>false if CPPFile is still being processed further down the callstack, true otherwise.</returns>
-		public static bool FindAndCacheAllIncludedFiles(UEBuildTarget Target, FileItem CPPFile, IUEBuildPlatform BuildPlatform, List<string> IncludePathsToSearch, Dictionary<string, FileItem> IncludeFileSearchDictionary, ref IncludedFilesSet Result, bool bOnlyCachedDependencies)
+		public static bool FindAndCacheAllIncludedFiles(UEBuildTarget Target, FileItem CPPFile, IUEBuildPlatform BuildPlatform, CPPIncludeInfo CPPIncludeInfo, ref IncludedFilesSet Result, bool bOnlyCachedDependencies)
 		{
 			IncludedFilesSet IncludedFileList;
 			var IncludedFilesMap = bOnlyCachedDependencies ? OnlyCachedIncludedFilesMap : ExhaustiveIncludedFilesMap;
@@ -215,7 +269,7 @@ namespace UnrealBuildTool
 						++TotalDirectIncludeResolveCacheMisses;
 
 						// search the include paths to resolve the file
-						FileItem DirectIncludeResolvedFile = CPPEnvironment.FindIncludedFile(DirectInclude.IncludeName, !BuildConfiguration.bCheckExternalHeadersForModification, IncludePathsToSearch, IncludeFileSearchDictionary );
+						FileItem DirectIncludeResolvedFile = CPPEnvironment.FindIncludedFile(DirectInclude.IncludeName, !BuildConfiguration.bCheckExternalHeadersForModification, CPPIncludeInfo.GetIncludesPathsToSearch( CPPFile ), CPPIncludeInfo.IncludeFileSearchDictionary );
 						if (DirectIncludeResolvedFile != null)
 						{
 							DirectlyIncludedFiles.Add(DirectIncludeResolvedFile);
@@ -240,7 +294,7 @@ namespace UnrealBuildTool
 					IncludedFileList.Add( DirectlyIncludedFile );
 
 					// Also add all of the indirectly included files!
-					if (FindAndCacheAllIncludedFiles(Target, DirectlyIncludedFile, BuildPlatform, IncludePathsToSearch, IncludeFileSearchDictionary, ref IncludedFileList, bOnlyCachedDependencies:bOnlyCachedDependencies) == false)
+					if (FindAndCacheAllIncludedFiles(Target, DirectlyIncludedFile, BuildPlatform, CPPIncludeInfo, ref IncludedFileList, bOnlyCachedDependencies:bOnlyCachedDependencies) == false)
 					{
 						// DirectlyIncludedFile is a circular dependency which is still being processed
 						// further down the callstack. Add this file to its circular dependencies list 
diff --git a/Engine/Source/Programs/UnrealBuildTool/System/Distcc.cs b/Engine/Source/Programs/UnrealBuildTool/System/Distcc.cs
--- a/Engine/Source/Programs/UnrealBuildTool/System/Distcc.cs
+++ b/Engine/Source/Programs/UnrealBuildTool/System/Distcc.cs
@@ -166,7 +166,7 @@ namespace UnrealBuildTool
 				}
 
 				Log.WriteLineIf(BuildConfiguration.bLogDetailedActionStats, TraceEventType.Information, "-------- Begin Detailed Action Stats ----------------------------------------------------------");
-				Log.WriteLineIf(BuildConfiguration.bLogDetailedActionStats, TraceEventType.Information, "^Action Type^Duration (seconds)^Tool^Task^Using PCH^Description");
+				Log.WriteLineIf(BuildConfiguration.bLogDetailedActionStats, TraceEventType.Information, "^Action Type^Duration (seconds)^Tool^Task^Using PCH");
 
 				double TotalThreadSeconds = 0;
 
@@ -192,13 +192,12 @@ namespace UnrealBuildTool
 
 					Log.WriteLineIf(BuildConfiguration.bLogDetailedActionStats,
 						TraceEventType.Information,
-						"^{0}^{1:0.00}^{2}^{3}^{4}^{5}", 
+						"^{0}^{1:0.00}^{2}^{3}^{4}", 
 						Action.ActionType.ToString(),
 						ThreadSeconds,
 						Path.GetFileName(Action.CommandPath), 
 						Action.StatusDescription,
-						Action.bIsUsingPCH,
-						Action.StatusDetailedDescription);
+						Action.bIsUsingPCH);
 
 					// Update statistics
 					switch (Action.ActionType)
diff --git a/Engine/Source/Programs/UnrealBuildTool/System/ExternalExecution.cs b/Engine/Source/Programs/UnrealBuildTool/System/ExternalExecution.cs
--- a/Engine/Source/Programs/UnrealBuildTool/System/ExternalExecution.cs
+++ b/Engine/Source/Programs/UnrealBuildTool/System/ExternalExecution.cs
@@ -19,11 +19,11 @@ namespace UnrealBuildTool
 	public enum ECompilationResult
 	{
 		Succeeded = 0,
-		FailedDueToHeaderChange = 1,
 		OtherCompilationError = 2
 	}
 
 	/** Information about a module that needs to be passed to UnrealHeaderTool for code generation */
+	[Serializable]
 	public class UHTModuleInfo
 	{
 		/** Module name */
@@ -262,7 +262,7 @@ namespace UnrealBuildTool
 		/// Gets the timestamp of CoreUObject.generated.cpp file.
 		/// </summary>
 		/// <returns>Last write time of CoreUObject.generated.cpp or DateTime.MaxValue if it doesn't exist.</returns>
-		private static DateTime GetCoreGeneratedTimestamp(UEBuildTarget Target)
+		private static DateTime GetCoreGeneratedTimestamp(UEBuildTarget Target, string ModuleName, string ModuleDirectory)
 		{			
 			DateTime Timestamp;
 			if( UnrealBuildTool.RunningRocket() )
@@ -272,8 +272,7 @@ namespace UnrealBuildTool
 			}
 			else
 			{
-				var CoreUObjectModule = (UEBuildModuleCPP)Target.GetModuleByName( "CoreUObject" );
-				string CoreGeneratedFilename = Path.Combine(UEBuildModuleCPP.GetGeneratedCodeDirectoryForModule(Target, CoreUObjectModule.ModuleDirectory, CoreUObjectModule.Name), CoreUObjectModule.Name + ".generated.cpp");
+				string CoreGeneratedFilename = Path.Combine(UEBuildModuleCPP.GetGeneratedCodeDirectoryForModule(Target, ModuleDirectory, ModuleName), ModuleName + ".generated.cpp");
 				if (File.Exists(CoreGeneratedFilename))
 				{
 					Timestamp = new FileInfo(CoreGeneratedFilename).LastWriteTime;
@@ -303,7 +302,23 @@ namespace UnrealBuildTool
 
 			// Get CoreUObject.generated.cpp timestamp.  If the source files are older than the CoreUObject generated code, we'll
 			// need to regenerate code for the module
-			var CoreGeneratedTimestamp = GetCoreGeneratedTimestamp(Target);
+			DateTime? CoreGeneratedTimestamp = null;
+			{ 
+				// Find the CoreUObject module
+				foreach( var Module in UObjectModules )
+				{
+					if( Module.ModuleName.Equals( "CoreUObject", StringComparison.InvariantCultureIgnoreCase ) )
+					{
+						CoreGeneratedTimestamp = GetCoreGeneratedTimestamp(Target, Module.ModuleName, Module.ModuleDirectory);
+						break;
+					}
+				}
+				if( CoreGeneratedTimestamp == null )
+				{
+					throw new BuildException( "Could not find CoreUObject in list of all UObjectModules" );
+				}
+			}
+
 
 			foreach( var Module in UObjectModules )
 			{
@@ -349,16 +364,24 @@ namespace UnrealBuildTool
 									break;
 								}
 
-								// Also check the timestamp on the directory the source file is in.  If the directory timestamp has
-								// changed, new source files may have been added or deleted.  We don't know whether the new/deleted
-								// files were actually UObject headers, but because we don't know all of the files we processed
-								// in the previous run, we need to assume our generated code is out of date if the directory timestamp
-								// is newer.
-								var HeaderDirectoryTimestamp = new DirectoryInfo( Path.GetDirectoryName( HeaderFile.AbsolutePath ) ).LastWriteTime;
-								if( SavedTimestamp.CompareTo( HeaderDirectoryTimestamp) < 0 )
+								// When we're running in assembler mode, outdatedness cannot be inferred by checking the directory timestamp
+								// of the source headers.  We don't care if source files were added or removed in this mode, because we're only
+								// able to process the known UObject headers that are in the Makefile.  If UObject header files are added/removed,
+								// we expect the user to re-run GenerateProjectFiles which will force UBTMakefile outdatedness.
+								// @todo fastubt: Possibly, we should never be doing this check these days.
+								if( UnrealBuildTool.IsGatheringBuild || !UnrealBuildTool.IsAssemblingBuild )
 								{
-									bIsOutOfDate = true;
-									break;
+									// Also check the timestamp on the directory the source file is in.  If the directory timestamp has
+									// changed, new source files may have been added or deleted.  We don't know whether the new/deleted
+									// files were actually UObject headers, but because we don't know all of the files we processed
+									// in the previous run, we need to assume our generated code is out of date if the directory timestamp
+									// is newer.
+									var HeaderDirectoryTimestamp = new DirectoryInfo( Path.GetDirectoryName( HeaderFile.AbsolutePath ) ).LastWriteTime;
+									if( SavedTimestamp.CompareTo( HeaderDirectoryTimestamp) < 0 )
+									{
+										bIsOutOfDate = true;
+										break;
+									}
 								}
 							}
 						}
@@ -506,25 +529,32 @@ namespace UnrealBuildTool
 					// want to save this work until we know that UHT actually needs to be run to speed up best-case iteration times.
 					foreach( var UHTModuleInfo in UObjectModules )
 					{
-						UHTModuleInfo.PCH = "";
-
-						// We need to figure out which PCH header this module is including, so that UHT can inject an include statement for it into any .cpp files it is synthesizing
-						var DependencyModuleCPP = (UEBuildModuleCPP)Target.GetModuleByName( UHTModuleInfo.ModuleName );
-						var ModuleCompileEnvironment = DependencyModuleCPP.CreateModuleCompileEnvironment(GlobalCompileEnvironment);
-						DependencyModuleCPP.CachePCHUsageForModuleSourceFiles(ModuleCompileEnvironment);
-						if (DependencyModuleCPP.ProcessedDependencies.UniquePCHHeaderFile != null)
+						// Only cache the PCH name if we don't already have one.  When running in 'gather only' mode, this will have already been cached
+						if( string.IsNullOrEmpty( UHTModuleInfo.PCH ) )
 						{
-							UHTModuleInfo.PCH = DependencyModuleCPP.ProcessedDependencies.UniquePCHHeaderFile.AbsolutePath;
+							UHTModuleInfo.PCH = "";
+
+							// We need to figure out which PCH header this module is including, so that UHT can inject an include statement for it into any .cpp files it is synthesizing
+							var DependencyModuleCPP = (UEBuildModuleCPP)Target.GetModuleByName( UHTModuleInfo.ModuleName );
+							var ModuleCompileEnvironment = DependencyModuleCPP.CreateModuleCompileEnvironment(GlobalCompileEnvironment);
+							DependencyModuleCPP.CachePCHUsageForModuleSourceFiles(ModuleCompileEnvironment);
+							if (DependencyModuleCPP.ProcessedDependencies.UniquePCHHeaderFile != null)
+							{
+								UHTModuleInfo.PCH = DependencyModuleCPP.ProcessedDependencies.UniquePCHHeaderFile.AbsolutePath;
+							}
 						}
 					}
 				}
 
+				// @todo fastubt: @todo ubtmake: Optimization: Ideally we could avoid having to generate this data in the case where UHT doesn't even need to run!  Can't we use the existing copy?  (see below use of Manifest)
 				UHTManifest Manifest = new UHTManifest(Target, RootLocalPath, ToolChain.ConvertPath(RootLocalPath + '\\'), UObjectModules);
 
 				if( !bIsBuildingUHT && bUHTNeedsToRun )
 				{
 					// Always build UnrealHeaderTool if header regeneration is required, unless we're running within a Rocket ecosystem
-					if (UnrealBuildTool.RunningRocket() == false && UEBuildConfiguration.bDoNotBuildUHT == false)
+					if (UnrealBuildTool.RunningRocket() == false && 
+						UEBuildConfiguration.bDoNotBuildUHT == false &&
+						!( !UnrealBuildTool.IsGatheringBuild && UnrealBuildTool.IsAssemblingBuild ) )	// If running in "assembler only" mode, we assume UHT is already up to date for much faster iteration!
 					{
 						// If it is out of date or not there it will be built.
 						// If it is there and up to date, it will add 0.8 seconds to the build time.
@@ -546,7 +576,7 @@ namespace UnrealBuildTool
 						{
 							UBTArguments.Append(" -noxge");
 						}
-
+						
 						if ( RunExternalExecutable( UnrealBuildTool.GetUBTPath(), UBTArguments.ToString() ) != 0 )
 						{ 
 							return false;
@@ -613,6 +643,7 @@ namespace UnrealBuildTool
 					// Allow generated code to be sync'd to remote machines if needed. This needs to be done even if UHT did not run because
 					// generated headers include other generated headers using absolute paths which in case of building remotely are already
 					// the remote machine absolute paths. Because of that parsing headers will not result in finding all includes properly.
+					// @todo ubtmake: Need to figure out what this does in the assembler case, and whether we need to run it
 					ToolChain.PostCodeGeneration(Target, Manifest);
 				}
 
diff --git a/Engine/Source/Programs/UnrealBuildTool/System/FileItem.cs b/Engine/Source/Programs/UnrealBuildTool/System/FileItem.cs
--- a/Engine/Source/Programs/UnrealBuildTool/System/FileItem.cs
+++ b/Engine/Source/Programs/UnrealBuildTool/System/FileItem.cs
@@ -6,6 +6,7 @@ using System.Text;
 using System.IO;
 using System.Diagnostics;
 using System.Threading;
+using System.Runtime.Serialization;
 
 namespace UnrealBuildTool
 {
@@ -13,59 +14,36 @@ namespace UnrealBuildTool
 	 * Represents a file on disk that is used as an input or output of a build action.
 	 * FileItems are created by calling FileItem.GetItemByPath, which creates a single FileItem for each unique file path.
 	 */
-	public class FileItem
+	[Serializable]
+	public class FileItem : ISerializable
 	{
+		///
+		/// Preparation and Assembly (serialized)
+		/// 
+
 		/** The action that produces the file. */
 		public Action ProducingAction = null;
 
 		/** The absolute path of the file. */
 		public readonly string AbsolutePath;
 
-		/** PCH header file name as it appears in an #include statement in source code (might include partial, or no relative path.)
-		    This is needed by some compilers to use PCH features. */
-		public string PCHHeaderNameInCode;
-
-		/** The pch file that this file will use */
-		public string PrecompiledHeaderIncludeFilename;
-
-		/** Description of file, used for logging. */
-		public string Description;
-
-		/** The information about the file. */
-		public FileInfo Info;
-
-		/** This is true if this item is actually a directory. Consideration for Mac application bundles. Note that Info will be null if true! */
-		public bool IsDirectory;
-
-		/** Relative cost of action associated with producing this file. */
-		public long RelativeCost = 0;
-
-		/** Used for performance debugging */
-		public static long TotalFileItemCount = 0;
-		public static long MissingFileItemCount = 0;
-
 		/** True if any DLLs produced by this  */
 		public bool bNeedsHotReloadNumbersDLLCleanUp = false;
 
-		/** A case-insensitive dictionary that's used to map each unique file name to a single FileItem object. */
-		static Dictionary<string, FileItem> UniqueSourceFileMap = new Dictionary<string, FileItem>();
-
 		/** Whether or not this is a remote file, in which case we can't access it directly */
 		public bool bIsRemoteFile = false;
 
-		/** A list of remote file items that have been created but haven't needed the remote info yet, so we can gang up many into one request */
-		static List<FileItem> DelayedRemoteLookupFiles = new List<FileItem>();
 
-		/** The compile environment for this file.  This is only set on C++ source file file items, or C++-related files that have include paths.  It's part of UBT's dependency caching optimizations. */
-		public CPPEnvironment CachedCPPEnvironment
+		/** For C++ file items, this stores cached information about the include paths needed in order to include header files from these C++ files.  This is part of UBT's dependency caching optimizations. */
+		public CPPIncludeInfo CachedCPPIncludeInfo
 		{
 			get
 			{
-				return _CachedCPPEnvironment;
+				return _CachedCPPIncludeInfo;
 			}
 			set
 			{
-				if( value != null && _CachedCPPEnvironment != null && _CachedCPPEnvironment != value )
+				if( value != null && _CachedCPPIncludeInfo != null && _CachedCPPIncludeInfo != value )
 				{
 					// Uh oh.  We're clobbering our cached CompileEnvironment for this file with a different CompileEnvironment.  This means
 					// that the same source file is being compiled into more than one module. (e.g. PCLaunch.rc)
@@ -79,7 +57,7 @@ namespace UnrealBuildTool
 					{ 					
 						// Let's make sure the include paths are the same
 						// @todo fastubt: We have not seen examples of this actually firing off, so we could probably remove the check for matching includes and simply always make this an error case
-						var CachedIncludePathsToSearch = _CachedCPPEnvironment.GetIncludesPathsToSearch( this );
+						var CachedIncludePathsToSearch = _CachedCPPIncludeInfo.GetIncludesPathsToSearch( this );
 						var NewIncludePathsToSearch = value.GetIncludesPathsToSearch( this );
 
 						bool bIncludesAreDifferent = false;
@@ -105,10 +83,36 @@ namespace UnrealBuildTool
 						}
 					}
 				}
-				_CachedCPPEnvironment = value;
+				_CachedCPPIncludeInfo = value;
 			}
 		}
-		public CPPEnvironment _CachedCPPEnvironment;
+		public CPPIncludeInfo _CachedCPPIncludeInfo;
+
+
+		///
+		/// Preparation only (not serialized)
+		/// 
+
+		/** PCH header file name as it appears in an #include statement in source code (might include partial, or no relative path.)
+		    This is needed by some compilers to use PCH features. */
+		public string PCHHeaderNameInCode;
+
+		/** The PCH file that this file will use */
+		public string PrecompiledHeaderIncludeFilename;
+
+
+		///
+		/// Transients (not serialized)
+		///
+
+		/** The information about the file. */
+		public FileInfo Info;
+
+		/** This is true if this item is actually a directory. Consideration for Mac application bundles. Note that Info will be null if true! */
+		public bool IsDirectory;
+
+		/** Relative cost of action associated with producing this file. */
+		public long RelativeCost = 0;
 
 		/** The last write time of the file. */
 		public DateTimeOffset _LastWriteTime;
@@ -116,7 +120,10 @@ namespace UnrealBuildTool
 		{
 			get
 			{
-				LookupOutstandingFiles();
+				if (bIsRemoteFile)
+				{
+					LookupOutstandingFiles();
+				}
 				return _LastWriteTime;
 			}
 			set { _LastWriteTime = value; }
@@ -128,7 +135,10 @@ namespace UnrealBuildTool
 		{
 			get
 			{
-				LookupOutstandingFiles();
+				if (bIsRemoteFile)
+				{
+					LookupOutstandingFiles();
+				}
 				return _bExists;
 			}
 			set { _bExists = value; }
@@ -140,12 +150,32 @@ namespace UnrealBuildTool
 		{
 			get
 			{
-				LookupOutstandingFiles();
+				if (bIsRemoteFile)
+				{
+					LookupOutstandingFiles();
+				}
 				return _Length;
 			}
 			set { _Length = value; }
 		}
 
+
+		///
+		/// Statics
+		///
+
+		/** Used for performance debugging */
+		public static long TotalFileItemCount = 0;
+		public static long MissingFileItemCount = 0;
+
+		/** A case-insensitive dictionary that's used to map each unique file name to a single FileItem object. */
+		static Dictionary<string, FileItem> UniqueSourceFileMap = new Dictionary<string, FileItem>();
+
+		/** A list of remote file items that have been created but haven't needed the remote info yet, so we can gang up many into one request */
+		static List<FileItem> DelayedRemoteLookupFiles = new List<FileItem>();
+
+
+
 		/**
 		 * Resolve any outstanding remote file info lookups
 		 */
@@ -319,6 +349,53 @@ namespace UnrealBuildTool
 			UniqueSourceFileMap[ AbsolutePath.ToLowerInvariant() ] = this;
 		}
 
+
+		/** ISerializable: Constructor called when this object is deserialized */
+		protected FileItem( SerializationInfo SerializationInfo, StreamingContext StreamingContext )
+		{
+			ProducingAction = (Action)SerializationInfo.GetValue( "pa", typeof( Action ) );
+			AbsolutePath = SerializationInfo.GetString( "ap" );
+			bIsRemoteFile = SerializationInfo.GetBoolean( "rf" );
+			bNeedsHotReloadNumbersDLLCleanUp = SerializationInfo.GetBoolean( "hr" );
+			CachedCPPIncludeInfo = (CPPIncludeInfo)SerializationInfo.GetValue( "ci", typeof( CPPIncludeInfo ) );
+
+			// Go ahead and init normally now
+			{ 
+				ResetFileInfo();
+
+				++TotalFileItemCount;
+				if (!_bExists)
+				{
+					++MissingFileItemCount;
+					// Log.TraceInformation( "Missing: " + FileAbsolutePath );
+				}
+
+				if( bIsRemoteFile )
+				{
+					lock (DelayedRemoteLookupFiles)
+					{
+						DelayedRemoteLookupFiles.Add(this);
+					}
+				}
+				else
+				{ 
+					UniqueSourceFileMap[ AbsolutePath.ToLowerInvariant() ] = this;
+				}
+			}
+		}
+
+
+		/** ISerializable: Called when serialized to report additional properties that should be saved */
+		public void GetObjectData( SerializationInfo SerializationInfo, StreamingContext StreamingContext )
+		{
+			SerializationInfo.AddValue( "pa", ProducingAction );
+			SerializationInfo.AddValue( "ap", AbsolutePath );
+			SerializationInfo.AddValue( "rf", bIsRemoteFile );
+			SerializationInfo.AddValue( "hr", bNeedsHotReloadNumbersDLLCleanUp );
+			SerializationInfo.AddValue( "ci", CachedCPPIncludeInfo );
+		}
+
+
 		/// <summary>
 		/// (Re-)set file information for this FileItem
 		/// </summary>
diff --git a/Engine/Source/Programs/UnrealBuildTool/System/FlatCPPIncludeDepencencyCache.cs b/Engine/Source/Programs/UnrealBuildTool/System/FlatCPPIncludeDepencencyCache.cs
--- a/Engine/Source/Programs/UnrealBuildTool/System/FlatCPPIncludeDepencencyCache.cs
+++ b/Engine/Source/Programs/UnrealBuildTool/System/FlatCPPIncludeDepencencyCache.cs
@@ -78,7 +78,7 @@ namespace UnrealBuildTool
 			{
 				using (FileStream Stream = new FileStream(Cache.AbsolutePath, FileMode.Open, FileAccess.Read))
 				{	
-					// @todo fastubt: We can store the cache in a cheaper/smaller way using hash file names and indicies into included headers, but it might actually slow down load times
+					// @todo fastubt: We can store the cache in a cheaper/smaller way using hash file names and indices into included headers, but it might actually slow down load times
 					// @todo fastubt: If we can index PCHs here, we can avoid storing all of the PCH's included headers (PCH's action should have been invalidated, so we shouldn't even have to report the PCH's includes as our indirect includes)
 					BinaryFormatter Formatter = new BinaryFormatter();
 					Result = Formatter.Deserialize(Stream) as FlatCPPIncludeDependencyCache;
@@ -105,7 +105,7 @@ namespace UnrealBuildTool
 		protected FlatCPPIncludeDependencyCache(FileItem Cache)
 		{
 			CacheFileItem = Cache;
-			DependencyMap = new Dictionary<string, FlatCPPIncludeDependencyInfo>(StringComparer.InvariantCultureIgnoreCase);
+			DependencyMap = new Dictionary<string, FlatCPPIncludeDependencyInfo>(StringComparer.InvariantCultureIgnoreCase);	// @todo fastubt: Apparently StringComparer is not very fast on Mono, better to use ToLowerInvariant(), even though it is an extra string copy
 		}
 
 
diff --git a/Engine/Source/Programs/UnrealBuildTool/System/LocalExecutor.cs b/Engine/Source/Programs/UnrealBuildTool/System/LocalExecutor.cs
--- a/Engine/Source/Programs/UnrealBuildTool/System/LocalExecutor.cs
+++ b/Engine/Source/Programs/UnrealBuildTool/System/LocalExecutor.cs
@@ -161,9 +161,9 @@ namespace UnrealBuildTool
 				}
 
 				ActionStartInfo.UseShellExecute = false;
-				ActionStartInfo.RedirectStandardInput = Action.bShouldBlockStandardInput;
-				ActionStartInfo.RedirectStandardOutput = Action.bShouldBlockStandardOutput;
-				ActionStartInfo.RedirectStandardError = Action.bShouldBlockStandardOutput;
+				ActionStartInfo.RedirectStandardInput = false;
+				ActionStartInfo.RedirectStandardOutput = false;
+				ActionStartInfo.RedirectStandardError = false;
 
 				// Log command-line used to execute task if debug info printing is enabled.
 				if (BuildConfiguration.bPrintDebugInfo)
@@ -495,7 +495,7 @@ namespace UnrealBuildTool
 			}
 
 			Log.WriteLineIf(BuildConfiguration.bLogDetailedActionStats, TraceEventType.Information, "-------- Begin Detailed Action Stats ----------------------------------------------------------");
-			Log.WriteLineIf(BuildConfiguration.bLogDetailedActionStats, TraceEventType.Information, "^Action Type^Duration (seconds)^Tool^Task^Using PCH^Description");
+			Log.WriteLineIf(BuildConfiguration.bLogDetailedActionStats, TraceEventType.Information, "^Action Type^Duration (seconds)^Tool^Task^Using PCH");
 
 			double TotalThreadSeconds = 0;
 
@@ -522,13 +522,12 @@ namespace UnrealBuildTool
 
 				Log.WriteLineIf(BuildConfiguration.bLogDetailedActionStats,
 					TraceEventType.Information,
-					"^{0}^{1:0.00}^{2}^{3}^{4}^{5}", 
+					"^{0}^{1:0.00}^{2}^{3}^{4}", 
 					Action.ActionType.ToString(),
 					ThreadSeconds,
 					Path.GetFileName(Action.CommandPath), 
                       Action.StatusDescription,
-					Action.bIsUsingPCH,
-					Action.StatusDetailedDescription);
+					Action.bIsUsingPCH);
 
 				// Update statistics
 				switch (Action.ActionType)
diff --git a/Engine/Source/Programs/UnrealBuildTool/System/Plugins.cs b/Engine/Source/Programs/UnrealBuildTool/System/Plugins.cs
--- a/Engine/Source/Programs/UnrealBuildTool/System/Plugins.cs
+++ b/Engine/Source/Programs/UnrealBuildTool/System/Plugins.cs
@@ -51,12 +51,6 @@ namespace UnrealBuildTool
 		// Path to the plugin's root directory
 		public string Directory;
 
-		// Path to the plugin's intermediate build folder
-		public string IntermediateBuildPath;
-
-		// Path to the plugin's Inc folder
-		public string IntermediateIncPath;
-
 		// List of modules in this plugin
 		public readonly List<PluginModuleInfo> Modules = new List<PluginModuleInfo>();
 
diff --git a/Engine/Source/Programs/UnrealBuildTool/System/ResponseFile.cs b/Engine/Source/Programs/UnrealBuildTool/System/ResponseFile.cs
--- a/Engine/Source/Programs/UnrealBuildTool/System/ResponseFile.cs
+++ b/Engine/Source/Programs/UnrealBuildTool/System/ResponseFile.cs
@@ -16,7 +16,6 @@ namespace UnrealBuildTool
 		/// <param name="Lines">List of lines to write to the response file</param>
 		public static string Create(string TempFileName, List<string> Lines)
 		{
-			// @todo fastubt: Make sure we aren't spitting out response files more often than we have to.  Also look at Unity.cpp files, injected .cpp files, etc.
 			FileInfo TempFileInfo = new FileInfo( TempFileName );
 			// Delete the existing file if it exists
 			if( TempFileInfo.Exists )
diff --git a/Engine/Source/Programs/UnrealBuildTool/System/Unity.cs b/Engine/Source/Programs/UnrealBuildTool/System/Unity.cs
--- a/Engine/Source/Programs/UnrealBuildTool/System/Unity.cs
+++ b/Engine/Source/Programs/UnrealBuildTool/System/Unity.cs
@@ -149,7 +149,7 @@ namespace UnrealBuildTool
 					// dependency cache, as it is now an "indirect include" from the unity file.  We'll clear out the compile environment
 					// attached to this file.  This prevents us from having to cache all of the indirect includes from these files inside our
 					// dependency cache, which speeds up iterative builds a lot!
-					CPPFile.CachedCPPEnvironment = null;
+					CPPFile.CachedCPPIncludeInfo = null;
 				}
 				AllUnityFiles = CPPUnityFileBuilder.GetUnityFiles();
 			}
@@ -225,7 +225,6 @@ namespace UnrealBuildTool
 				}
 
 				UnityCPPFile.RelativeCost        = UnityFile.TotalLength;
-				UnityCPPFile.Description         = string.Join(" + ", UnityFile.Files.Select(F => Path.GetFileName(F.AbsolutePath)));
 				UnityCPPFile.PCHHeaderNameInCode = PCHHeaderNameInCode;
                 UnityCPPFiles.Add(UnityCPPFile);
 
diff --git a/Engine/Source/Programs/UnrealBuildTool/System/UnrealBuildTool.cs b/Engine/Source/Programs/UnrealBuildTool/System/UnrealBuildTool.cs
--- a/Engine/Source/Programs/UnrealBuildTool/System/UnrealBuildTool.cs
+++ b/Engine/Source/Programs/UnrealBuildTool/System/UnrealBuildTool.cs
@@ -8,6 +8,7 @@ using System.Diagnostics;
 using System.Threading;
 using System.Reflection;
 using System.Linq;
+using System.Runtime.Serialization.Formatters.Binary;
 
 namespace UnrealBuildTool
 {
@@ -68,6 +69,51 @@ namespace UnrealBuildTool
         /** The Unreal project directory.  This should always be valid when compiling game projects.  Engine targets and program targets may not have a UProject set. */
         static string UProjectPath = "";
 
+		/** This is set to true during UBT startup, after it is safe to be accessing the IsGatheringBuild and IsAssemblingBuild properties.  Never access this directly. */
+		private static bool bIsSafeToCheckIfGatheringOrAssemblingBuild = false;
+
+		/** True if we should gather module dependencies for building in this run.  If this is false, then we'll expect to be loading information from disk about
+		    the target's modules before we'll be able to build anything.  One or both of IsGatheringBuild or IsAssemblingBuild must be true. */
+		static public bool IsGatheringBuild
+		{
+			get
+			{
+				if( bIsSafeToCheckIfGatheringOrAssemblingBuild )
+				{ 
+					return bIsGatheringBuild_Unsafe;
+				}
+				else
+				{
+					throw new BuildException( "Accessed UnrealBuildTool.IsGatheringBuild before it was fully initialized at startup (bIsSafeToCheckIfGatheringOrAssemblingBuild)" );
+				}
+			}
+		}
+		static private bool bIsGatheringBuild_Unsafe = true;
+
+		/** True if we should go ahead and actually compile and link in this run.  If this is false, no actions will be executed and we'll stop after gathering
+		    and saving information about dependencies.  One or both of IsGatheringBuild or IsAssemblingBuild must be true. */
+		static public bool IsAssemblingBuild
+		{
+			get
+			{
+				if( bIsSafeToCheckIfGatheringOrAssemblingBuild )
+				{ 
+					return bIsAssemblingBuild_Unsafe;
+				}
+				else
+				{
+					throw new BuildException( "Accessed UnrealBuildTool.IsGatheringBuild before it was fully initialized at startup (bIsSafeToCheckIfGatheringOrAssemblingBuild)" );
+				}
+			}
+		}
+		static private bool bIsAssemblingBuild_Unsafe = true;	// @todo ubtmake urgent: Doesn't seem to work with hotreload (tries to overwrite same DLL and not use a hotreload DLL name)
+
+		/** Used when BuildConfiguration.bUseExperimentalFastDependencyScan is enabled.  If true, it means that our cached includes may not longer be
+		    valid (or never were), and we need to recover by forcibly scanning included headers for all build prerequisites to make sure that our
+		    cached set of includes is actually correct, before determining which files are outdated. */
+		static public bool bNeedsFullCPPIncludeRescan = false;
+
+
         /// <summary>
         /// Returns true if UnrealBuildTool should skip everything other than the host platform as far as finding other modules that need to be compiled.
         /// </summary>
@@ -672,6 +718,13 @@ namespace UnrealBuildTool
                         }
                     }
 
+					// If we were asked to enable fast build iteration, we want the 'gather' phase to default to off (unless it is overridden below
+					// using a command-line option.)
+					if( BuildConfiguration.bUseExperimentalFastBuildIteration )
+					{
+						bIsGatheringBuild_Unsafe = false;
+					}
+
                     foreach (var Arg in Arguments)
                     {
                         string LowercaseArg = Arg.ToLowerInvariant();
@@ -793,6 +846,32 @@ namespace UnrealBuildTool
                             BuildConfiguration.bXGEExport = true;
                             BuildConfiguration.bAllowXGE = true;
                         }
+						else if (LowercaseArg == "-gather")
+						{
+							UnrealBuildTool.bIsGatheringBuild_Unsafe = true;
+						}
+						else if (LowercaseArg == "-nogather")
+						{
+							UnrealBuildTool.bIsGatheringBuild_Unsafe = false;
+						}
+						else if (LowercaseArg == "-gatheronly")
+						{
+							UnrealBuildTool.bIsGatheringBuild_Unsafe = true;
+							UnrealBuildTool.bIsAssemblingBuild_Unsafe = false;
+						}
+						else if (LowercaseArg == "-assemble")
+						{
+							UnrealBuildTool.bIsAssemblingBuild_Unsafe = true;
+						}
+						else if (LowercaseArg == "-noassemble")
+						{
+							UnrealBuildTool.bIsAssemblingBuild_Unsafe = false;
+						}
+						else if (LowercaseArg == "-assembleonly")
+						{
+							UnrealBuildTool.bIsGatheringBuild_Unsafe = false;
+							UnrealBuildTool.bIsAssemblingBuild_Unsafe = true;
+						}
                         else if (LowercaseArg == "-nosimplygon")
                         {
                             UEBuildConfiguration.bCompileSimplygon = false;
@@ -844,6 +923,11 @@ namespace UnrealBuildTool
 						}
                     }
 
+					if( !bIsGatheringBuild_Unsafe && !bIsAssemblingBuild_Unsafe )
+					{
+						throw new BuildException( "UnrealBuildTool: At least one of either IsGatheringBuild or IsAssemblingBuild must be true.  Did you pass '-NoGather' with '-NoAssemble'?" );
+					}
+
                     // Send an event with basic usage dimensions
                     Telemetry.SendEvent("CommonAttributes.2",
                         // @todo No platform independent way to do processor speed and count. Defer for now. Mono actually has ways to do this but needs to be tested.
@@ -937,8 +1021,16 @@ namespace UnrealBuildTool
 						{
 							if (UEBuildConfiguration.bPrepForDeployment == false)
 							{
-								// If we are only prepping for deployment, assume the build already occurred.
-								Result = RunUBT(Arguments);
+								int PassCount = 1;	// @todo ubtmake: It's useful to set this to a high number when running UBT under a sampling profiler
+								for( int PassIndex = 0; PassIndex < PassCount; ++PassIndex )
+								{ 
+									// If we are only prepping for deployment, assume the build already occurred.
+									Result = RunUBT(Arguments);
+									if( Result != ECompilationResult.Succeeded )
+									{
+										break;
+									}
+								}
 							}
 							else
 							{
@@ -1136,7 +1228,7 @@ namespace UnrealBuildTool
             bool bSuccess = true;
 
             // Reset global configurations
-            ResetAllActions();
+            ActionGraph.ResetAllActions();
 
             ParseBuildConfigurationFlags(Arguments);
 
@@ -1185,86 +1277,199 @@ namespace UnrealBuildTool
                     BuildConfiguration.bCreateStubIPA = false;
                 }
 
-                // Build action lists for all passed in targets.
-                var OutputItemsForAllTargets = new List<FileItem>();
+
                 foreach (string[] TargetSetting in TargetSettings)
                 {
-                    var TargetStartTime = DateTime.UtcNow;
-                    var Target = UEBuildTarget.CreateTarget(TargetSetting);
+                    var Target = UEBuildTarget.CreateTarget(TargetSetting);	// @todo ubtmake: Optimization: Ideally we don't have to do this in the assembler.  It is pretty slow to scan all rules files.  Maybe we can skip the slow parts internally?  It has to scan for *.Target.cs files, and also Load the rules assembly up. :(
                     if ((Target == null) && (UEBuildConfiguration.bCleanProject))
                     {
                         continue;
                     }
                     Targets.Add(Target);
+				}
 
-					// Load the include dependency cache.
-					// @todo fastubt: Ideally load this asynchronously at startup and only block when it is first needed and not finished loading
-					{	
-						if (!CPPEnvironment.IncludeDependencyCache.ContainsKey(Target))
+				UBTMakefile UBTMakefile = null;
+
+
+				// If UBT is being used to build something different than it did last time
+				if( !ProjectFileGenerator.bGenerateProjectFiles )
+				{
+					bool bIsBuildingSameTargetsAsLastTime = false;
+					{
+						// If we're building UHT without a Mutex, we'll need to assume that we're building the same targets and that no caches
+						// should be invalidated for this run.  This is important when UBT is invoked from within UBT in order to compile
+						// UHT.  In that (very common) case we definitely don't want to have to rebuild our cache from scratch.
+						bool bMustAssumeSameTargets = false;
+
+						if( Targets[0].GetTargetName().Equals( "UnrealHeaderTool", StringComparison.InvariantCultureIgnoreCase ) )
 						{
-							CPPEnvironment.IncludeDependencyCache.Add( Target, DependencyCache.Create( DependencyCache.GetDependencyCachePathForTarget(Target) ) );
+							int NoMutexArgumentIndex;
+							if (Utils.ParseCommandLineFlag(Arguments, "-NoMutex", out NoMutexArgumentIndex))
+							{
+								bMustAssumeSameTargets = true;
+							}
+						}
+
+						if( bMustAssumeSameTargets )
+						{
+							bIsBuildingSameTargetsAsLastTime = true;
+						}
+						else
+						{
+							string TargetCollectionName = MakeTargetCollectionName( Targets );
+							string LastBuiltTargetsFilePath = Path.Combine( BuildConfiguration.BaseIntermediatePath, "LastBuiltTargets.txt" );
+							if( File.Exists( LastBuiltTargetsFilePath ) && Utils.ReadAllText( LastBuiltTargetsFilePath ) == TargetCollectionName )
+							{
+								bIsBuildingSameTargetsAsLastTime = true;
+							}
+
+							// Save out the name of the targets we're building
+							if( !bIsBuildingSameTargetsAsLastTime )
+							{
+								Directory.CreateDirectory( Path.GetDirectoryName( LastBuiltTargetsFilePath ) );
+								File.WriteAllText( LastBuiltTargetsFilePath, TargetCollectionName, Encoding.UTF8 );
+							}
+
+
+							if( !bIsBuildingSameTargetsAsLastTime )
+							{
+								// Can't use super fast include checking unless we're building the same set of targets as last time, because
+								// we might not know about all of the C++ include dependencies for already-up-to-date shared build products
+								// between the targets
+								bNeedsFullCPPIncludeRescan = true;
+							}
 						}
 					}
-					if( BuildConfiguration.bUseExperimentalFastDependencyScan )
-					{ 
-						if (!CPPEnvironment.FlatCPPIncludeDependencyCache.ContainsKey(Target))
+				}
+			
+
+
+				{ 
+					// If we're generating project files, then go ahead and wipe out the existing UBTMakefile for every target, to make sure that
+					// it gets a full dependency scan next time.
+					// NOTE: This is just a safeguard and doesn't have to be perfect.  We also check for newer project file timestamps in LoadUBTMakefile()
+					if( ProjectFileGenerator.bGenerateProjectFiles )	// @todo ubtmake: This is only hit when generating IntelliSense for project files.  Probably should be done right inside ProjectFileGenerator.bat
+					{													// @todo ubtmake: Won't catch multi-target cases as GPF always builds one target at a time for Intellisense
+						// Delete the UBT makefile
+						var UBTMakefilePath = UnrealBuildTool.GetUBTMakefilePath( Targets );
+						if (File.Exists(UBTMakefilePath))
 						{
-							CPPEnvironment.FlatCPPIncludeDependencyCache.Add( Target, FlatCPPIncludeDependencyCache.Create( Target ) );
+							UEBuildTarget.CleanFile(UBTMakefilePath);
 						}
 					}
 
-                    var TargetOutputItems = new List<FileItem>();
-                    BuildResult = Target.Build(ToolChain, TargetOutputItems, out EULAViolationWarning);
+					// Make sure the gather phase is executed if we're not actually building anything
+					if( ProjectFileGenerator.bGenerateProjectFiles || UEBuildConfiguration.bGenerateManifest || UEBuildConfiguration.bCleanProject )
+					{
+						UnrealBuildTool.bIsGatheringBuild_Unsafe = true;						
+					}
 
-                    if(BuildResult != ECompilationResult.Succeeded)
-                    {
-                        break;
-                    }
+					// Were we asked to run in 'assembler only' mode?  If so, let's check to see if that's even possible by seeing if
+					// we have a valid UBT Makefile already saved to disk, ready for us to load.
+					if( UnrealBuildTool.bIsAssemblingBuild_Unsafe && !UnrealBuildTool.bIsGatheringBuild_Unsafe )
+					{
+						// @todo ubtmake: Mildly terrified of BuildConfiguration/UEBuildConfiguration globals that were set during the Prepare phase but are not set now.  We may need to save/load all of these, otherwise
+						//		we'll need to call SetupGlobalEnvironment on all of the targets (maybe other stuff, too.  See PreBuildStep())
+							
+						UBTMakefile = LoadUBTMakefile( Targets );
+
+						if( UBTMakefile == null )
+						{ 
+							// If the Makefile couldn't be loaded, then we're not going to be able to continue in "assembler only" mode.  We'll do both
+							// a 'gather' and 'assemble' in the same run.  This will take a while longer, but subsequent runs will be fast!
+							UnrealBuildTool.bIsGatheringBuild_Unsafe = true;
+
+							Log.TraceInformation( "Forcing generation of new UBTMakefile for {0}{1}", Targets[0].GetTargetName(), Targets.Count > 1 ? ( " (and " + ( Targets.Count - 1 ).ToString() + " more)" ) : "" );
+						}
+					}
 
-                    OutputItemsForAllTargets.AddRange( TargetOutputItems );
+					// OK, after this point it is safe to access the UnrealBuildTool.IsGatheringBuild and UnrealBuildTool.IsAssemblingBuild properties.
+					// These properties will not be changing again during this session/
+					bIsSafeToCheckIfGatheringOrAssemblingBuild = true;
+				}
 
-                    if ( (BuildConfiguration.bXGEExport && UEBuildConfiguration.bGenerateManifest) || (!ProjectFileGenerator.bGenerateProjectFiles && !UEBuildConfiguration.bGenerateManifest && !UEBuildConfiguration.bCleanProject))
-                    {
-                        // We don't currently support building game targets in rocket.
-                        if (UnrealBuildTool.RunningRocket() && Target.Rules != null && Target.Rules.Type == TargetRules.TargetType.Game)
-                        {
-                            throw new BuildException(
-                                "You currently can not build a game target in Rocket.\nTry again in a future release.\nFor now, build and run your editor project."
-                            );
-                        }
 
-                        // Generate an action graph if we were asked to do that.  The graph generation needs access to the include dependency cache, so
-                        // we generate it before saving and cleaning that up.
-                        if( GeneratingActionGraph )
-                        {
-                            // The graph generation feature currently only works with a single target at a time.  This is because of how we need access
-                            // to include dependencies for the target, but those aren't kept around as we process multiple targets
-                            if( TargetSettings.Count != 1 )
-                            {
-                                throw new BuildException( "ERROR: The '-graph' option only works with a single target at a time" );
-                            }
-                            LinkActionsAndItems();
-                            var ActionsToExecute = AllActions;
+                // Build action lists for all passed in targets.
+                var OutputItemsForAllTargets = new List<FileItem>();
+				var TargetNameToUObjectModules = new Dictionary<string, List<UHTModuleInfo>>( StringComparer.InvariantCultureIgnoreCase );
+                foreach (var Target in Targets)
+                {
+                    var TargetStartTime = DateTime.UtcNow;
 
-							var VisualizationType = ActionGraphVisualizationType.OnlyCPlusPlusFilesAndHeaders;
-							SaveActionGraphVisualization( Target, Path.Combine( BuildConfiguration.BaseIntermediatePath, Target.GetTargetName() + ".gexf" ), Target.GetTargetName(), VisualizationType, ActionsToExecute );
-						}
+					// When in 'assembler only' mode, we'll load this cache later on a worker thread.  It takes a long time to load!
+					if( !( !UnrealBuildTool.IsGatheringBuild && UnrealBuildTool.IsAssemblingBuild ) )
+					{	
+						// Load the direct include dependency cache.
+						CPPEnvironment.IncludeDependencyCache.Add( Target, DependencyCache.Create( DependencyCache.GetDependencyCachePathForTarget(Target) ) );
 					}
 
-                    var TargetBuildTime = (DateTime.UtcNow - TargetStartTime).TotalSeconds;
-
-					// Send out telemetry for this target
-					Telemetry.SendEvent("TargetBuildStats.2",
-						"AppName", Target.AppName,
-						"GameName", Target.GameName,
-						"Platform", Target.Platform.ToString(),
-						"Configuration", Target.Configuration.ToString(),
-						"CleanTarget", UEBuildConfiguration.bCleanProject.ToString(),
-						"Monolithic", Target.ShouldCompileMonolithic().ToString(),
-						"CreateDebugInfo", Target.IsCreatingDebugInfo().ToString(),
-						"TargetType", Target.Rules.Type.ToString(),
-						"TargetCreateTimeSec", TargetBuildTime.ToString("0.00")
-						);
+					// We don't need this dependency cache in 'gather only' mode
+					if( BuildConfiguration.bUseExperimentalFastDependencyScan &&
+						!( UnrealBuildTool.IsGatheringBuild && !UnrealBuildTool.IsAssemblingBuild ) )
+					{ 
+						// Load the cache that contains the list of flattened resolved includes for resolved source files
+						// @todo fastubt: Ideally load this asynchronously at startup and only block when it is first needed and not finished loading
+						CPPEnvironment.FlatCPPIncludeDependencyCache.Add( Target, FlatCPPIncludeDependencyCache.Create( Target ) );
+					}
+
+					if( UnrealBuildTool.IsGatheringBuild )
+					{ 
+	                    List<FileItem> TargetOutputItems;
+						List<UHTModuleInfo> TargetUObjectModules;
+						BuildResult = Target.Build(ToolChain, out TargetOutputItems, out TargetUObjectModules, out EULAViolationWarning);
+						if(BuildResult != ECompilationResult.Succeeded)
+						{
+							break;
+						}
+
+						OutputItemsForAllTargets.AddRange( TargetOutputItems );
+
+						// Update mapping of the target name to the list of UObject modules in this target
+						TargetNameToUObjectModules[ Target.GetTargetName() ] = TargetUObjectModules;
+
+						if ( (BuildConfiguration.bXGEExport && UEBuildConfiguration.bGenerateManifest) || (!ProjectFileGenerator.bGenerateProjectFiles && !UEBuildConfiguration.bGenerateManifest && !UEBuildConfiguration.bCleanProject))
+						{
+							// We don't currently support building game targets in rocket.
+							if (UnrealBuildTool.RunningRocket() && Target.Rules != null && Target.Rules.Type == TargetRules.TargetType.Game)
+							{
+								throw new BuildException(
+									"You currently can not build a game target in Rocket.\nTry again in a future release.\nFor now, build and run your editor project."
+								);
+							}
+
+							// Generate an action graph if we were asked to do that.  The graph generation needs access to the include dependency cache, so
+							// we generate it before saving and cleaning that up.
+							if( GeneratingActionGraph )
+							{
+								// The graph generation feature currently only works with a single target at a time.  This is because of how we need access
+								// to include dependencies for the target, but those aren't kept around as we process multiple targets
+								if( TargetSettings.Count != 1 )
+								{
+									throw new BuildException( "ERROR: The '-graph' option only works with a single target at a time" );
+								}
+								ActionGraph.FinalizeActionGraph();
+								var ActionsToExecute = ActionGraph.AllActions;
+
+								var VisualizationType = ActionGraph.ActionGraphVisualizationType.OnlyCPlusPlusFilesAndHeaders;
+								ActionGraph.SaveActionGraphVisualization( Target, Path.Combine( BuildConfiguration.BaseIntermediatePath, Target.GetTargetName() + ".gexf" ), Target.GetTargetName(), VisualizationType, ActionsToExecute );
+							}
+						}
+
+						var TargetBuildTime = (DateTime.UtcNow - TargetStartTime).TotalSeconds;
+
+						// Send out telemetry for this target
+						Telemetry.SendEvent("TargetBuildStats.2",
+							"AppName", Target.AppName,
+							"GameName", Target.GameName,
+							"Platform", Target.Platform.ToString(),
+							"Configuration", Target.Configuration.ToString(),
+							"CleanTarget", UEBuildConfiguration.bCleanProject.ToString(),
+							"Monolithic", Target.ShouldCompileMonolithic().ToString(),
+							"CreateDebugInfo", Target.IsCreatingDebugInfo().ToString(),
+							"TargetType", Target.Rules.Type.ToString(),
+							"TargetCreateTimeSec", TargetBuildTime.ToString("0.00")
+							);
+					}
 				}
 
                 if (BuildResult == ECompilationResult.Succeeded &&
@@ -1273,59 +1478,138 @@ namespace UnrealBuildTool
                         (!GeneratingActionGraph && !ProjectFileGenerator.bGenerateProjectFiles && !UEBuildConfiguration.bGenerateManifest && !UEBuildConfiguration.bCleanProject && !UEBuildConfiguration.bGenerateExternalFileList)
                     ))
                 {
-                    // Plan the actions to execute for the build.
-					Dictionary<UEBuildTarget,List<FileItem>> TargetToOutdatedPrerequisitesMap;
-                    List<Action> ActionsToExecute = GetActionsToExecute(OutputItemsForAllTargets, Targets, out TargetToOutdatedPrerequisitesMap);
+					if( UnrealBuildTool.IsGatheringBuild )
+					{
+						ActionGraph.FinalizeActionGraph();
 
-                    // Display some stats to the user.
-                    Log.TraceVerbose(
-                            "{0} actions, {1} outdated and requested actions",
-                            AllActions.Count,
-                            ActionsToExecute.Count
-                            );
+						UBTMakefile = new UBTMakefile();
+						UBTMakefile.AllActions = ActionGraph.AllActions;
 
-                    ToolChain.PreBuildSync();
+						// For now simply treat all object files as the root target.
+						{ 
+							var PrerequisiteActionsSet = new HashSet<Action>();
+							foreach (var OutputItem in OutputItemsForAllTargets)
+							{
+								ActionGraph.GatherPrerequisiteActions(OutputItem, ref PrerequisiteActionsSet);
+							}
+							UBTMakefile.PrerequisiteActions = PrerequisiteActionsSet.ToArray();
+						}			
 
-					
-					// Cache indirect includes for all outdated C++ files.  We kick this off as a background thread so that it can
-					// perform the scan while we're compilng.  It usually only takes up to a few seconds, but we don't want to hurt
-					// our best case UBT iteration times for this task which can easily be performed asynchronously
-					Thread CPPIncludesThread = null;
-					if( BuildConfiguration.bUseExperimentalFastDependencyScan )
-					{
-						CPPIncludesThread = CreateThreadForCachingCPPIncludes( TargetToOutdatedPrerequisitesMap );
-						CPPIncludesThread.Start();
-					}
+						foreach( System.Collections.DictionaryEntry EnvironmentVariable in Environment.GetEnvironmentVariables() )
+						{
+							UBTMakefile.EnvironmentVariables.Add( Tuple.Create( (string)EnvironmentVariable.Key, (string)EnvironmentVariable.Value ) );
+						}
+						UBTMakefile.TargetNameToUObjectModules = TargetNameToUObjectModules;
 
-                    // Execute the actions.
-                    bSuccess = ExecuteActions(ActionsToExecute, out ExecutorName);
+						// We've been told to prepare to build, so let's go ahead and save out our action graph so that we can use in a later invocation 
+						// to assemble the build.  Even if we are configured to assemble the build in this same invocation, we want to save out the
+						// Makefile so that it can be used on subsequent 'assemble only' runs, for the fastest possible iteration times
+						// @todo ubtmake: Optimization: We could make 'gather + assemble' mode slightly faster by saving this while busy compiling (on our worker thread)
+						SaveUBTMakefile( Targets, UBTMakefile );
+					}
 
-					// if the build succeeded, do any needed syncing
-					if (bSuccess)
+					if( UnrealBuildTool.IsAssemblingBuild )
 					{
-						foreach (UEBuildTarget Target in Targets)
+						// If we didn't build the graph in this session, then we'll need to load a cached one
+						if( !UnrealBuildTool.IsGatheringBuild )
 						{
-							ToolChain.PostBuildSync(Target);
+							ActionGraph.AllActions = UBTMakefile.AllActions;
+
+							foreach( var EnvironmentVariable in UBTMakefile.EnvironmentVariables )
+							{
+								// @todo ubtmake: There may be some variables we do NOT want to clobber.
+								Environment.SetEnvironmentVariable( EnvironmentVariable.Item1, EnvironmentVariable.Item2 );
+							}
+
+							// If any of the targets need UHT to be run, we'll go ahead and do that now
+							foreach( var Target in Targets )
+							{
+								List<UHTModuleInfo> TargetUObjectModules;
+								if( UBTMakefile.TargetNameToUObjectModules.TryGetValue( Target.GetTargetName(), out TargetUObjectModules ) )
+								{
+									if( TargetUObjectModules.Count > 0 )
+									{
+										// Execute the header tool
+										string ModuleInfoFileName = Path.GetFullPath( Path.Combine( Target.ProjectIntermediateDirectory, "UnrealHeaderTool.manifest" ) );
+										ECompilationResult UHTResult = ECompilationResult.OtherCompilationError;
+										if (!ExternalExecution.ExecuteHeaderToolIfNecessary(Target, GlobalCompileEnvironment:null, UObjectModules:TargetUObjectModules, ModuleInfoFileName:ModuleInfoFileName, UHTResult:ref UHTResult))
+										{
+											Log.TraceInformation("UnrealHeaderTool failed for target '" + Target.GetTargetName() + "' (platform: " + Target.GetTargetInfo().Platform.ToString() + ", module info: " + ModuleInfoFileName + ").");
+											BuildResult = UHTResult;
+											break;
+										}
+									}
+								}
+							}
 						}
-					}
-					else
-					{
-						BuildResult = ECompilationResult.OtherCompilationError;
-					}
 
-					if( CPPIncludesThread != null )
-					{ 
-						// Wait until our CPPIncludes dependency scanner thread has finished
-						CPPIncludesThread.Join();
+						if( BuildResult == ECompilationResult.Succeeded )
+						{ 
+							// Make sure any old DLL files from in-engine recompiles aren't lying around.  Must be called after the action graph is finalized.
+							// @todo ubtmake: Will we even use traditional hot reload DLLs anymore in the future?  
+							ActionGraph.DeleteStaleHotReloadDLLs();
+
+							// Plan the actions to execute for the build.
+							Dictionary<UEBuildTarget,List<FileItem>> TargetToOutdatedPrerequisitesMap;
+							List<Action> ActionsToExecute = ActionGraph.GetActionsToExecute(UBTMakefile.PrerequisiteActions, Targets, out TargetToOutdatedPrerequisitesMap);
+
+							// Display some stats to the user.
+							Log.TraceVerbose(
+									"{0} actions, {1} outdated and requested actions",
+									ActionGraph.AllActions.Count,
+									ActionsToExecute.Count
+									);
+
+							ToolChain.PreBuildSync();
+
+					
+							// Cache indirect includes for all outdated C++ files.  We kick this off as a background thread so that it can
+							// perform the scan while we're compiling.  It usually only takes up to a few seconds, but we don't want to hurt
+							// our best case UBT iteration times for this task which can easily be performed asynchronously
+							Thread CPPIncludesThread = null;
+							if( BuildConfiguration.bUseExperimentalFastDependencyScan && TargetToOutdatedPrerequisitesMap.Count > 0 )
+							{
+								CPPIncludesThread = CreateThreadForCachingCPPIncludes( TargetToOutdatedPrerequisitesMap );
+								CPPIncludesThread.Start();
+							}
+
+							// Execute the actions.
+							bSuccess = ActionGraph.ExecuteActions(ActionsToExecute, out ExecutorName);
+
+							// if the build succeeded, do any needed syncing
+							if (bSuccess)
+							{
+								foreach (UEBuildTarget Target in Targets)
+								{
+									ToolChain.PostBuildSync(Target);
+								}
+							}
+							else
+							{
+								BuildResult = ECompilationResult.OtherCompilationError;
+							}
+
+							if( CPPIncludesThread != null )
+							{ 
+								// @todo ubtmake urgent: It is critical that this thread completes and we save dependency cache files, even if an error happened, otherwise our
+								// comparison with last built target will pass the next time through, even though our cache is bogus.  Maybe we could wipe it when we think its out of date?
+								
+								// Wait until our CPPIncludes dependency scanner thread has finished
+								CPPIncludesThread.Join();
+							}
+						}
 					}
 				}
 
 				// Save the include dependency cache.
 				{ 
 					// @todo fastubt urgent: If we update the cache and actually compile files, but DON'T SAVE the stream, we will be in big trouble on the next run -- the files could be out of date due to an indirect include, but we won't know how to check because we didn't cache the includes off!
-					//			-> Link errors shouldn't prevent the cache from saving, but UBT exceptions could.  
-					//			-> In any case, we need to:  A) Wipe cache file on exception	B) force SLOW intermediate checking on restart if not fully clean and cache file is missing (or force a clean first)
-					//			
+					//			-> Compile/Link errors DO SEEM to prevent saving!  ACK!
+					//			-> In any case, we need to:  
+					//					A) Wipe cache file on exception	
+					//					B) force SLOW intermediate checking on restart if not fully clean and cache file is missing (or force a clean first)
+					//					C) Save cache files on compile/link errors!
+
 					foreach( var DependencyCache in CPPEnvironment.IncludeDependencyCache.Values )
 					{ 
 						DependencyCache.Save();
@@ -1572,6 +1856,18 @@ namespace UnrealBuildTool
 					// @todo fastubt: This thread will access data structures that are also used on the main UBT thread, but during this time UBT
 					// is only invoking the build executor, so should not be touching this stuff.  However, we need to at some guards to make sure.
 
+
+					// When in 'assembler only' mode, we need to load the direct include graph from a cache now.  We skipped loading this earlier because
+					// we didn't need it unless there was actually an outdated dependency.  And we can load this on our worker thread while we're compiling.
+					foreach( var Target in TargetToOutdatedPrerequisitesMap.Keys )
+					{
+						if (!CPPEnvironment.IncludeDependencyCache.ContainsKey(Target))
+						{
+							CPPEnvironment.IncludeDependencyCache.Add( Target, DependencyCache.Create( DependencyCache.GetDependencyCachePathForTarget(Target) ) );
+						}
+					}
+
+
 					foreach( var TargetAndOutdatedPrerequisites in TargetToOutdatedPrerequisitesMap )
 					{
 						var Target = TargetAndOutdatedPrerequisites.Key;
@@ -1581,11 +1877,245 @@ namespace UnrealBuildTool
 						{
 							// Invoke our deep include scanner to figure out whether any of the files included by this source file have
 							// changed since the build product was built
-							var FileBuildPlatform = UEBuildPlatform.GetBuildPlatformForCPPTargetPlatform( PrerequisiteItem.CachedCPPEnvironment.Config.Target.Platform );
-							CPPEnvironment.FindAndCacheAllIncludedFiles( Target, PrerequisiteItem, FileBuildPlatform, PrerequisiteItem.CachedCPPEnvironment.GetIncludesPathsToSearch( PrerequisiteItem ), PrerequisiteItem.CachedCPPEnvironment.IncludeFileSearchDictionary, bOnlyCachedDependencies:false );
+							var FileBuildPlatform = UEBuildPlatform.GetBuildPlatform( Target.GetTargetInfo().Platform );
+							CPPEnvironment.FindAndCacheAllIncludedFiles( Target, PrerequisiteItem, FileBuildPlatform, PrerequisiteItem.CachedCPPIncludeInfo, bOnlyCachedDependencies:false );
 						}
 					}
 				} ) );
 		}
-    }
+
+
+
+		/// <summary>
+		/// A special Makefile that UBT is able to create in "-gather" mode, then load in "-assemble" mode to accelerate iterative compiling and linking
+		/// </summary>
+		[Serializable]
+		class UBTMakefile
+		{
+			/** Every action in the action graph */
+			public List<Action> AllActions;
+
+			/** List of the actions that need to be run in order to build the targets' final output items */
+			public Action[] PrerequisiteActions;			
+
+			/** Environment variables that we'll need in order to invoke the platform's compiler and linker */
+			public readonly List<Tuple<string,string>> EnvironmentVariables = new List<Tuple<string,string>>();
+
+			/** Maps each target to a list of UObject module info structures */
+			public Dictionary<string,List<UHTModuleInfo>> TargetNameToUObjectModules;
+
+
+			/** @return	 True if this makefile's contents look valid.  Called after loading the file to make sure it is legit. */
+			public bool IsValidMakefile()
+			{
+				return 
+					AllActions != null && AllActions.Count > 0 &&
+					PrerequisiteActions != null && PrerequisiteActions.Length > 0 &&
+					EnvironmentVariables != null &&
+					TargetNameToUObjectModules != null && TargetNameToUObjectModules.Count > 0;
+			}
+		}
+
+
+		/// <summary>
+		/// Saves a UBTMakefile to disk
+		/// </summary>
+		/// <param name="Targets">List of targets.  Order is not important</param>
+		static void SaveUBTMakefile( List<UEBuildTarget> Targets, UBTMakefile UBTMakefile )
+		{
+			if( !UBTMakefile.IsValidMakefile() )
+			{
+				throw new BuildException( "Can't save a UBTMakefile that has invalid contents.  See UBTMakefile.IsValidMakefile()" );
+			}
+
+			var TimerStartTime = DateTime.UtcNow;
+
+			var UBTMakefileItem = FileItem.GetItemByFullPath( GetUBTMakefilePath( Targets ) );
+
+			// @todo ubtmake: Optimization: The UBTMakefile saved for game projects is upwards of 9 MB.  We should try to shrink its content if possible
+			// @todo ubtmake: Optimization: C# Serialization may be too slow for these big Makefiles.  Loading these files often shows up as the slower part of the assembling phase.
+
+			// Serialize the cache to disk.
+			try
+			{
+				Directory.CreateDirectory(Path.GetDirectoryName(UBTMakefileItem.AbsolutePath));
+				using (FileStream Stream = new FileStream(UBTMakefileItem.AbsolutePath, FileMode.Create, FileAccess.Write))
+				{
+					BinaryFormatter Formatter = new BinaryFormatter();
+					Formatter.Serialize(Stream, UBTMakefile);
+				}
+			}
+			catch (Exception Ex)
+			{
+				Console.Error.WriteLine("Failed to write UBTMakefile: {0}", Ex.Message);
+			}
+
+			if (BuildConfiguration.bPrintPerformanceInfo)
+			{
+				var TimerDuration = DateTime.UtcNow - TimerStartTime;
+				Log.TraceInformation("Saving UBTMakefile took " + TimerDuration.TotalSeconds + "s");
+			}
+		}
+
+
+		/// <summary>
+		/// Loads a UBTMakefile from disk
+		/// </summary>
+		/// <param name="Targets">List of targets.  Order is not important</param>
+		static UBTMakefile LoadUBTMakefile( List<UEBuildTarget> Targets )
+		{
+			var UBTMakefileItem = FileItem.GetItemByFullPath( GetUBTMakefilePath( Targets ) );
+
+			// Check the directory timestamp on the project files directory.  If the user has generated project files more
+			// recently than the UBTMakefile, then we need to consider the file to be out of date
+			bool bForceOutOfDate = false;
+			{
+				// @todo ubtmake: This will only work if the directory timestamp actually changes with every single GPF.  Force delete existing files before creating new ones?  Eh... really we probably just want to delete + create a file in that folder
+				//			-> UPDATE: Seems to work OK right now though on Windows platform, maybe due to GUID changes
+				// @todo ubtmake: Some platforms may not save any files into this folder.  We should delete + generate a "touch" file to force the directory timestamp to be updated (or just check the timestamp file itself.  We could put it ANYWHERE, actually)
+
+				var ProjectFilesLastUpdateTime = DateTime.MaxValue;
+				if( !UnrealBuildTool.RunningRocket() )
+				{
+					if( Directory.Exists( ProjectFileGenerator.IntermediateProjectFilesPath ) )
+					{
+						ProjectFilesLastUpdateTime = new DirectoryInfo( ProjectFileGenerator.IntermediateProjectFilesPath ).LastWriteTime;
+					}
+					else
+					{
+						ProjectFilesLastUpdateTime = DateTime.MaxValue;
+					}
+				}
+				else
+				{
+					// Rocket doesn't need to check engine projects for outdatedness
+					ProjectFilesLastUpdateTime = DateTime.MinValue;
+				}
+
+				// Check the game project directory too
+				if( UnrealBuildTool.HasUProjectFile() )
+				{ 
+					var MasterProjectRelativePath = UnrealBuildTool.GetUProjectPath();
+					var GameIntermediateProjectFilesPath = Path.Combine( MasterProjectRelativePath, "Intermediate", "ProjectFiles" );
+					if( Directory.Exists( GameIntermediateProjectFilesPath ) )
+					{
+						var GameProjectFilesLastUpdateTime = new DirectoryInfo( GameIntermediateProjectFilesPath ).LastWriteTime;
+						if( GameProjectFilesLastUpdateTime > ProjectFilesLastUpdateTime )
+						{
+							ProjectFilesLastUpdateTime = GameProjectFilesLastUpdateTime;
+						}
+					}	
+					else
+					{
+						// Game project files folder is missing, so force outdatedness
+						ProjectFilesLastUpdateTime = DateTime.MaxValue;
+
+					}
+				}
+	
+				if( UBTMakefileItem.LastWriteTime.CompareTo( ProjectFilesLastUpdateTime ) < 0 )
+				{
+					// Consider the file to be out of date
+					bForceOutOfDate = true;
+				}
+			} 
+
+			UBTMakefile LoadedUBTMakefile = null;
+			if( !bForceOutOfDate )
+			{ 
+				try
+				{
+					using (FileStream Stream = new FileStream(UBTMakefileItem.AbsolutePath, FileMode.Open, FileAccess.Read))
+					{	
+						BinaryFormatter Formatter = new BinaryFormatter();
+						LoadedUBTMakefile = Formatter.Deserialize(Stream) as UBTMakefile;
+					}
+				}
+				catch (Exception Ex)
+				{
+					Log.TraceWarning("Failed to read UBTMakefile: {0}", Ex.Message);
+				}
+
+				if( LoadedUBTMakefile != null && !LoadedUBTMakefile.IsValidMakefile() )
+				{
+					Log.TraceWarning("Loaded UBTMakefile appears to have invalid contents, ignoring it ({0})", UBTMakefileItem.AbsolutePath );
+					LoadedUBTMakefile = null;
+				}
+			}
+
+			return LoadedUBTMakefile;
+		}
+
+
+		/// <summary>
+		/// Gets the file path for a UBTMakefile
+		/// </summary>
+		/// <param name="Targets">List of targets.  Order is not important</param>
+		/// <returns>UBTMakefile path</returns>
+		public static string GetUBTMakefilePath( List<UEBuildTarget> Targets )
+		{
+			string UBTMakefilePath;
+
+			if( Targets.Count == 1 )
+			{
+				// If there's only one target, just save the UBT makefile in the target's build intermediate directory
+				// under a folder for that target (and platform/config combo.)
+				string PlatformIntermediatePath = BuildConfiguration.PlatformIntermediatePath;
+				if (UnrealBuildTool.HasUProjectFile())
+				{
+					PlatformIntermediatePath = Path.Combine(UnrealBuildTool.GetUProjectPath(), BuildConfiguration.PlatformIntermediateFolder);
+				}
+
+				UBTMakefilePath = Path.Combine(PlatformIntermediatePath, Targets[0].GetTargetName(), Targets[0].GetTargetInfo().Configuration.ToString(), "Makefile.ubt" );
+			}
+			else
+			{
+ 				// For Makefiles that contain multiple targets, we'll make up a file name that contains all of the targets, their
+ 				// configurations and platforms, and save it into the base intermediate folder
+				var TargetCollectionName = MakeTargetCollectionName( Targets );
+
+				string ProjectIntermediatePath = BuildConfiguration.BaseIntermediatePath;
+				if (UnrealBuildTool.HasUProjectFile())
+				{
+					ProjectIntermediatePath = Path.Combine(UnrealBuildTool.GetUProjectPath(), BuildConfiguration.BaseIntermediatePath);
+				}
+
+				// @todo ubtmake: The TargetCollectionName string could be really long if there is more than one target!  Hash it?
+				UBTMakefilePath = Path.Combine(ProjectIntermediatePath, TargetCollectionName + ".ubt" );
+			}
+
+			return UBTMakefilePath;
+		}
+
+
+		/// <summary>
+		/// Makes up a name for a set of targets that we can use for file or directory names
+		/// </summary>
+		/// <param name="Targets">List of targets.  Order is not important</param>
+		/// <returns>The name to use</returns>
+		private static string MakeTargetCollectionName( List<UEBuildTarget> Targets )
+		{
+			if( Targets.Count == 0 )
+			{
+				throw new BuildException( "Expecting at least one Target to be passed to MakeTargetCollectionName" );
+			}
+
+			var SortedTargets = new List<UEBuildTarget>();
+			SortedTargets.AddRange( Targets );
+			SortedTargets.Sort( ( x, y ) => { return x.GetTargetName().CompareTo( y.GetTargetName() ); } );
+
+			// Figure out what to call our action graph based on everything we're building
+			var TargetCollectionName = new StringBuilder();
+			foreach( var Target in SortedTargets )
+			{
+				if( TargetCollectionName.Length > 0 )
+				{
+					TargetCollectionName.Append( "_" );
+				}
+				TargetCollectionName.Append( Target.GetTargetName() + "-" + Target.GetTargetInfo().Platform.ToString() + "-" + Target.GetTargetInfo().Configuration.ToString() );
+			}
+
+			return TargetCollectionName.ToString();
+		}
+	}
 }
diff --git a/Engine/Source/Programs/UnrealBuildTool/System/Utils.cs b/Engine/Source/Programs/UnrealBuildTool/System/Utils.cs
--- a/Engine/Source/Programs/UnrealBuildTool/System/Utils.cs
+++ b/Engine/Source/Programs/UnrealBuildTool/System/Utils.cs
@@ -13,97 +13,6 @@ using System.Linq;
 
 namespace UnrealBuildTool
 {
-	public class PrecompileHeaderEnvironment
-	{
-		/** The name of the module this PCH header is a member of */
-		public readonly string ModuleName;
-
-		/** PCH header file name as it appears in an #include statement in source code (might include partial, or no relative path.)
-			This is needed by some compilers to use PCH features. */
-		public string PCHHeaderNameInCode;
-
-		/** The source header file that this precompiled header will be generated for */
-		public readonly FileItem PrecompiledHeaderIncludeFilename;
-
-		/** Whether this precompiled header will be built with CLR features enabled.  We won't mix and match CLR PCHs with non-CLR PCHs */
-		public readonly CPPCLRMode CLRMode;
-
-		/** Whether this precompiled header will be built with code optimization enabled. */
-		public readonly ModuleRules.CodeOptimization OptimizeCode;
-
-		/** The PCH file we're generating */
-		public FileItem PrecompiledHeaderFile = null;
-
-		/** Object files emitted from the compiler when generating this precompiled header.  These will be linked into modules that
-			include this PCH */
-		public readonly List<FileItem> OutputObjectFiles = new List<FileItem>();
-
-		public PrecompileHeaderEnvironment( string InitModuleName, string InitPCHHeaderNameInCode, FileItem InitPrecompiledHeaderIncludeFilename, CPPCLRMode InitCLRMode, ModuleRules.CodeOptimization InitOptimizeCode )
-		{
-			ModuleName = InitModuleName;
-			PCHHeaderNameInCode = InitPCHHeaderNameInCode;
-			PrecompiledHeaderIncludeFilename = InitPrecompiledHeaderIncludeFilename;
-			CLRMode = InitCLRMode;
-			OptimizeCode = InitOptimizeCode;
-		}
-
-		/// <summary>
-		/// Creates a precompiled header action to generate a new pch file 
-		/// </summary>
-		/// <param name="PCHHeaderNameInCode">The precompiled header name as it appeared in an #include statement</param>
-		/// <param name="PrecompiledHeaderIncludeFilename">Name of the header used for pch.</param>
-		/// <param name="ProjectCPPEnvironment">The environment the C/C++ files in the project are compiled with.</param>
-		/// <param name="OutputDirectory">The folder to save the generated PCH file to</param>
-		/// <param name="ModuleName">Name of the module this PCH is being generated for</param>
-		/// <param name="bAllowDLLExports">True if we should allow DLLEXPORT definitions for this PCH</param>
-		/// <returns>the compilation output result of the created pch.</returns>
-		public static CPPOutput GeneratePCHCreationAction(UEBuildTarget Target, string PCHHeaderNameInCode, FileItem PrecompiledHeaderIncludeFilename, CPPEnvironment ProjectCPPEnvironment, string OutputDirectory, string ModuleName, bool bAllowDLLExports )
-		{
-			// Find the header file to be precompiled. Don't skip external headers
-			if (PrecompiledHeaderIncludeFilename.bExists)
-			{
-				// Create a Dummy wrapper around the PCH to avoid problems with #pragma once on clang
-				var ToolChain = UEToolChain.GetPlatformToolChain(ProjectCPPEnvironment.Config.Target.Platform);
-				string PCHGuardDefine = Path.GetFileNameWithoutExtension(PrecompiledHeaderIncludeFilename.AbsolutePath).ToUpper();
-				string LocalPCHHeaderNameInCode = ToolChain.ConvertPath(PrecompiledHeaderIncludeFilename.AbsolutePath);
-				string TmpPCHHeaderContents = String.Format("#ifndef __AUTO_{0}_H__\n#define __AUTO_{0}_H__\n//Last Write: {2}\n#include \"{1}\"\n#endif//__AUTO_{0}_H__", PCHGuardDefine, LocalPCHHeaderNameInCode, PrecompiledHeaderIncludeFilename.LastWriteTime);
-				string DummyPath = Path.Combine(
-					ProjectCPPEnvironment.Config.OutputDirectory,
-					Path.GetFileName(PrecompiledHeaderIncludeFilename.AbsolutePath));
-				FileItem DummyPCH = FileItem.CreateIntermediateTextFile(DummyPath, TmpPCHHeaderContents);
-
-				// Create a new C++ environment that is used to create the PCH.
-				var ProjectPCHEnvironment = ProjectCPPEnvironment.DeepCopy();
-				ProjectPCHEnvironment.Config.PrecompiledHeaderAction = PrecompiledHeaderAction.Create;
-				ProjectPCHEnvironment.Config.PrecompiledHeaderIncludeFilename = PrecompiledHeaderIncludeFilename.AbsolutePath;
-				ProjectPCHEnvironment.Config.PCHHeaderNameInCode = PCHHeaderNameInCode;
-				ProjectPCHEnvironment.Config.OutputDirectory = OutputDirectory;
-
-				if( !bAllowDLLExports )
-				{
-					for( var CurDefinitionIndex = 0; CurDefinitionIndex < ProjectPCHEnvironment.Config.Definitions.Count; ++CurDefinitionIndex )
-					{
-						// We change DLLEXPORT to DLLIMPORT for "shared" PCH headers
-						var OldDefinition = ProjectPCHEnvironment.Config.Definitions[ CurDefinitionIndex ];
-						if( OldDefinition.EndsWith( "=DLLEXPORT" ) )
-						{
-							ProjectPCHEnvironment.Config.Definitions[ CurDefinitionIndex ] = OldDefinition.Replace( "DLLEXPORT", "DLLIMPORT" );
-						}
-					}
-				}
-
-				// Cache our CPP environment so that we can check for outdatedness quickly.  Only files that have includes need this.
-				DummyPCH.CachedCPPEnvironment = ProjectPCHEnvironment;
-
-				Log.TraceVerbose( "Found PCH file \"{0}\".", PrecompiledHeaderIncludeFilename );
-
-				// Create the action to compile the PCH file.
-				return ProjectPCHEnvironment.CompileFiles(Target, new List<FileItem>() { DummyPCH }, ModuleName);
-			}
-			throw new BuildException( "Couldn't find PCH file \"{0}\".", PrecompiledHeaderIncludeFilename );
-		}
-	}
-
 	public abstract class Utils
 	{
 		/// <summary>
diff --git a/Engine/Source/Programs/UnrealBuildTool/ToolChain/UEToolChain.cs b/Engine/Source/Programs/UnrealBuildTool/ToolChain/UEToolChain.cs
--- a/Engine/Source/Programs/UnrealBuildTool/ToolChain/UEToolChain.cs
+++ b/Engine/Source/Programs/UnrealBuildTool/ToolChain/UEToolChain.cs
@@ -159,5 +159,38 @@ namespace UnrealBuildTool
         {
 
         }
+
+
+		protected void AddPrerequisiteSourceFile( UEBuildTarget Target, IUEBuildPlatform BuildPlatform, CPPEnvironment CompileEnvironment, FileItem SourceFile, List<FileItem> PrerequisiteItems )
+		{
+			PrerequisiteItems.Add( SourceFile );
+
+			var RemoteThis = this as RemoteToolChain;
+			bool bAllowUploading = RemoteThis != null && BuildHostPlatform.Current.Platform != UnrealTargetPlatform.Mac;	// Don't use remote features when compiling from a Mac
+			if( bAllowUploading )
+			{
+				RemoteThis.QueueFileForBatchUpload(SourceFile);
+			}
+
+			if( !BuildConfiguration.bUseExperimentalFastBuildIteration )	// In fast build iteration mode, we'll gather includes later on
+			{
+				// @todo fastubt: What if one of the prerequisite files has become missing since it was updated in our cache? (usually, because a coder eliminated the source file)
+				//		-> Two CASES:
+				//				1) NOT WORKING: Non-unity file went away (SourceFile in this context).  That seems like an existing old use case.  Compile params or Response file should have changed?
+				//				2) WORKING: Indirect file went away (unity'd original source file or include).  This would return a file that no longer exists and adds to the prerequiteitems list
+				var IncludedFileList = CPPEnvironment.FindAndCacheAllIncludedFiles( Target, SourceFile, BuildPlatform, CompileEnvironment.Config.CPPIncludeInfo, bOnlyCachedDependencies:BuildConfiguration.bUseExperimentalFastDependencyScan );
+				foreach (FileItem IncludedFile in IncludedFileList)
+				{
+					PrerequisiteItems.Add( IncludedFile );
+
+					if( bAllowUploading &&
+						!BuildConfiguration.bUseExperimentalFastDependencyScan )	// With fast dependency scanning, we will not have an exhaustive list of dependencies here.  We rely on PostCodeGeneration() to upload these files.
+					{
+						RemoteThis.QueueFileForBatchUpload(IncludedFile);
+					}
+				}
+			}
+		}
+
 	};
 }
\ No newline at end of file
diff --git a/Engine/Source/Programs/UnrealBuildTool/WinRT/WinRTToolChain.cs b/Engine/Source/Programs/UnrealBuildTool/WinRT/WinRTToolChain.cs
--- a/Engine/Source/Programs/UnrealBuildTool/WinRT/WinRTToolChain.cs
+++ b/Engine/Source/Programs/UnrealBuildTool/WinRT/WinRTToolChain.cs
@@ -426,11 +426,11 @@ namespace UnrealBuildTool
 			string Arguments = GetCLArguments_Global(CompileEnvironment);
 
 			// Add include paths to the argument list.
-			foreach (string IncludePath in CompileEnvironment.Config.IncludePaths)
+			foreach (string IncludePath in CompileEnvironment.Config.CPPIncludeInfo.IncludePaths)
 			{
 				Arguments += string.Format(" /I \"{0}\"", IncludePath);
 			}
-			foreach (string IncludePath in CompileEnvironment.Config.SystemIncludePaths)
+			foreach (string IncludePath in CompileEnvironment.Config.CPPIncludeInfo.SystemIncludePaths)
 			{
 				Arguments += string.Format(" /I \"{0}\"", IncludePath);
 			}
@@ -492,11 +492,7 @@ namespace UnrealBuildTool
 				bool bIsPlainCFile = Path.GetExtension(SourceFile.AbsolutePath).ToUpperInvariant() == ".C";
 
 				// Add the C++ source file and its included files to the prerequisite item list.
-				CompileAction.PrerequisiteItems.Add(SourceFile);
-				{
-					var IncludedFileList = CPPEnvironment.FindAndCacheAllIncludedFiles( Target, SourceFile, BuildPlatform, CompileEnvironment.GetIncludesPathsToSearch( SourceFile ), CompileEnvironment.IncludeFileSearchDictionary, bOnlyCachedDependencies:BuildConfiguration.bUseExperimentalFastDependencyScan );
-					CompileAction.PrerequisiteItems.AddRange( IncludedFileList );
-				}
+				AddPrerequisiteSourceFile( Target, BuildPlatform, CompileEnvironment, SourceFile, CompileAction.PrerequisiteItems );
 
 				// If this is a CLR file then make sure our dependent assemblies are added as prerequisites
 				if (CompileEnvironment.Config.CLRMode == CPPCLRMode.CLREnabled)
@@ -652,7 +648,6 @@ namespace UnrealBuildTool
 				CompileAction.bIsVCCompiler = true;
 				CompileAction.CommandArguments = Arguments + FileArguments + CompileEnvironment.Config.AdditionalArguments;
 				CompileAction.StatusDescription = string.Format("{0}", Path.GetFileName(SourceFile.AbsolutePath));
-				CompileAction.StatusDetailedDescription = SourceFile.Description;
 
 				// Don't farm out creation of precomputed headers as it is the critical path task.
 				CompileAction.bCanExecuteRemotely = CompileEnvironment.Config.PrecompiledHeaderAction != PrecompiledHeaderAction.Create;
diff --git a/Engine/Source/Programs/UnrealBuildTool/Windows/VCToolChain.cs b/Engine/Source/Programs/UnrealBuildTool/Windows/VCToolChain.cs
--- a/Engine/Source/Programs/UnrealBuildTool/Windows/VCToolChain.cs
+++ b/Engine/Source/Programs/UnrealBuildTool/Windows/VCToolChain.cs
@@ -631,11 +631,11 @@ namespace UnrealBuildTool
 			AppendCLArguments_Global(CompileEnvironment, Arguments);
 
 			// Add include paths to the argument list.
-			foreach (string IncludePath in CompileEnvironment.Config.IncludePaths)
+			foreach (string IncludePath in CompileEnvironment.Config.CPPIncludeInfo.IncludePaths)
 			{
 				Arguments.AppendFormat(" /I \"{0}\"", IncludePath);
 			}
-			foreach (string IncludePath in CompileEnvironment.Config.SystemIncludePaths)
+			foreach (string IncludePath in CompileEnvironment.Config.CPPIncludeInfo.SystemIncludePaths)
 			{
 				if( WindowsPlatform.bCompileWithClang )
 				{
@@ -697,15 +697,7 @@ namespace UnrealBuildTool
 				bool bIsPlainCFile = Path.GetExtension(SourceFile.AbsolutePath).ToUpperInvariant() == ".C";
 
 				// Add the C++ source file and its included files to the prerequisite item list.
-				CompileAction.PrerequisiteItems.Add(SourceFile);
-				{
-					// @todo fastubt: What if one of the prerequisite files has become missing since it was updated in our cache? (usually, because a coder eliminated the source file)
-					//		-> Two CASES:
-					//				1) NOT WORKING: Non-unity file went away (SourceFile in this context).  That seems like an existing old use case.  Compile params or Response file should have changed?
-					//				2) WORKING: Indirect file went away (unity'd original source file or include).  This would return a file that no longer exists and adds to the prerequiteitems list
-					var IncludedFileList = CPPEnvironment.FindAndCacheAllIncludedFiles( Target, SourceFile, BuildPlatform, CompileEnvironment.GetIncludesPathsToSearch( SourceFile ), CompileEnvironment.IncludeFileSearchDictionary, bOnlyCachedDependencies:BuildConfiguration.bUseExperimentalFastDependencyScan );
-					CompileAction.PrerequisiteItems.AddRange( IncludedFileList );
-				}
+				AddPrerequisiteSourceFile( Target, BuildPlatform, CompileEnvironment, SourceFile, CompileAction.PrerequisiteItems );
 
 				// If this is a CLR file then make sure our dependent assemblies are added as prerequisites
 				if (CompileEnvironment.Config.CLRMode == CPPCLRMode.CLREnabled)
@@ -904,7 +896,6 @@ namespace UnrealBuildTool
 					CompileAction.bIsVCCompiler = true;
 				}
 				CompileAction.CommandArguments = Arguments.ToString() + FileArguments.ToString() + CompileEnvironment.Config.AdditionalArguments;
-				CompileAction.StatusDetailedDescription = SourceFile.Description;
 
 				if( CompileEnvironment.Config.PrecompiledHeaderAction == PrecompiledHeaderAction.Create )
 				{
@@ -984,13 +975,13 @@ namespace UnrealBuildTool
 				CompileAction.CommandArguments += " /l 0x409";
 
 				// Include paths.
-				foreach (string IncludePath in Environment.Config.IncludePaths)
+				foreach (string IncludePath in Environment.Config.CPPIncludeInfo.IncludePaths)
 				{
 					CompileAction.CommandArguments += string.Format( " /i \"{0}\"", IncludePath );
 				}
 
 				// System include paths.
-				foreach( var SystemIncludePath in Environment.Config.SystemIncludePaths )
+				foreach( var SystemIncludePath in Environment.Config.CPPIncludeInfo.SystemIncludePaths )
 				{
 					CompileAction.CommandArguments += string.Format( " /i \"{0}\"", SystemIncludePath );
 				}
@@ -1013,14 +1004,10 @@ namespace UnrealBuildTool
 				Result.ObjectFiles.Add(CompiledResourceFile);
 
 				// Add the RC file as a prerequisite of the action.
-				CompileAction.PrerequisiteItems.Add(RCFile);
 				CompileAction.CommandArguments += string.Format(" \"{0}\"", RCFile.AbsolutePath);
 
-				// Add the files included by the RC file as prerequisites of the action.
-				{
-					var IncludedFileList = CPPEnvironment.FindAndCacheAllIncludedFiles( Target, RCFile, BuildPlatform, Environment.GetIncludesPathsToSearch( RCFile ), Environment.IncludeFileSearchDictionary, bOnlyCachedDependencies:BuildConfiguration.bUseExperimentalFastDependencyScan );
-					CompileAction.PrerequisiteItems.AddRange( IncludedFileList );
-				}
+				// Add the C++ source file and its included files to the prerequisite item list.
+				AddPrerequisiteSourceFile( Target, BuildPlatform, Environment, RCFile, CompileAction.PrerequisiteItems );
 			}
 
 			return Result;
@@ -1285,9 +1272,6 @@ namespace UnrealBuildTool
 			// Add project
 			BuildProjectAction.CommandArguments += String.Format(" \"{0}\"", ProjectFileItem.AbsolutePath);
 
-			// We don't want to display all of the regular MSBuild output to the console
-			BuildProjectAction.bShouldBlockStandardOutput = false;
-
 			// Specify the output files.
 			string PDBFilePath = Path.Combine( Path.GetDirectoryName(DestinationFile), Path.GetFileNameWithoutExtension(DestinationFile) + ".pdb" );
 			FileItem PDBFile = FileItem.GetItemByPath( PDBFilePath );
