Make a separate Mutex for UBT operations that only deal with SDK validation

diff --git a/Engine/Source/Programs/UnrealBuildTool/System/UnrealBuildTool.cs b/Engine/Source/Programs/UnrealBuildTool/System/UnrealBuildTool.cs
--- a/Engine/Source/Programs/UnrealBuildTool/System/UnrealBuildTool.cs
+++ b/Engine/Source/Programs/UnrealBuildTool/System/UnrealBuildTool.cs
@@ -40,8 +40,8 @@ namespace UnrealBuildTool
         /** Total time spent on other actions. */
         static public double TotalOtherActionsTime = 0;
 
-		/// How much time was spent scanning for include dependencies for outdated C++ files
-		public static double TotalDeepIncludeScanTime = 0.0;
+        /// How much time was spent scanning for include dependencies for outdated C++ files
+        public static double TotalDeepIncludeScanTime = 0.0;
 
         /** The command line */
         static public List<string> CmdLine = new List<string>();
@@ -69,49 +69,49 @@ namespace UnrealBuildTool
         /** The Unreal project directory.  This should always be valid when compiling game projects.  Engine targets and program targets may not have a UProject set. */
         static string UProjectPath = "";
 
-		/** This is set to true during UBT startup, after it is safe to be accessing the IsGatheringBuild and IsAssemblingBuild properties.  Never access this directly. */
-		private static bool bIsSafeToCheckIfGatheringOrAssemblingBuild = false;
-
-		/** True if we should gather module dependencies for building in this run.  If this is false, then we'll expect to be loading information from disk about
-		    the target's modules before we'll be able to build anything.  One or both of IsGatheringBuild or IsAssemblingBuild must be true. */
-		static public bool IsGatheringBuild
-		{
-			get
-			{
-				if( bIsSafeToCheckIfGatheringOrAssemblingBuild )
-				{ 
-					return bIsGatheringBuild_Unsafe;
-				}
-				else
-				{
-					throw new BuildException( "Accessed UnrealBuildTool.IsGatheringBuild before it was fully initialized at startup (bIsSafeToCheckIfGatheringOrAssemblingBuild)" );
-				}
-			}
-		}
-		static private bool bIsGatheringBuild_Unsafe = true;
-
-		/** True if we should go ahead and actually compile and link in this run.  If this is false, no actions will be executed and we'll stop after gathering
-		    and saving information about dependencies.  One or both of IsGatheringBuild or IsAssemblingBuild must be true. */
-		static public bool IsAssemblingBuild
-		{
-			get
-			{
-				if( bIsSafeToCheckIfGatheringOrAssemblingBuild )
-				{ 
-					return bIsAssemblingBuild_Unsafe;
-				}
-				else
-				{
-					throw new BuildException( "Accessed UnrealBuildTool.IsGatheringBuild before it was fully initialized at startup (bIsSafeToCheckIfGatheringOrAssemblingBuild)" );
-				}
-			}
-		}
-		static private bool bIsAssemblingBuild_Unsafe = true;	// @todo ubtmake urgent: Doesn't seem to work with hotreload (tries to overwrite same DLL and not use a hotreload DLL name)
-
-		/** Used when BuildConfiguration.bUseExperimentalFastDependencyScan is enabled.  If true, it means that our cached includes may not longer be
-		    valid (or never were), and we need to recover by forcibly scanning included headers for all build prerequisites to make sure that our
-		    cached set of includes is actually correct, before determining which files are outdated. */
-		static public bool bNeedsFullCPPIncludeRescan = false;
+        /** This is set to true during UBT startup, after it is safe to be accessing the IsGatheringBuild and IsAssemblingBuild properties.  Never access this directly. */
+        private static bool bIsSafeToCheckIfGatheringOrAssemblingBuild = false;
+
+        /** True if we should gather module dependencies for building in this run.  If this is false, then we'll expect to be loading information from disk about
+            the target's modules before we'll be able to build anything.  One or both of IsGatheringBuild or IsAssemblingBuild must be true. */
+        static public bool IsGatheringBuild
+        {
+            get
+            {
+                if( bIsSafeToCheckIfGatheringOrAssemblingBuild )
+                { 
+                    return bIsGatheringBuild_Unsafe;
+                }
+                else
+                {
+                    throw new BuildException( "Accessed UnrealBuildTool.IsGatheringBuild before it was fully initialized at startup (bIsSafeToCheckIfGatheringOrAssemblingBuild)" );
+                }
+            }
+        }
+        static private bool bIsGatheringBuild_Unsafe = true;
+
+        /** True if we should go ahead and actually compile and link in this run.  If this is false, no actions will be executed and we'll stop after gathering
+            and saving information about dependencies.  One or both of IsGatheringBuild or IsAssemblingBuild must be true. */
+        static public bool IsAssemblingBuild
+        {
+            get
+            {
+                if( bIsSafeToCheckIfGatheringOrAssemblingBuild )
+                { 
+                    return bIsAssemblingBuild_Unsafe;
+                }
+                else
+                {
+                    throw new BuildException( "Accessed UnrealBuildTool.IsGatheringBuild before it was fully initialized at startup (bIsSafeToCheckIfGatheringOrAssemblingBuild)" );
+                }
+            }
+        }
+        static private bool bIsAssemblingBuild_Unsafe = true;	// @todo ubtmake urgent: Doesn't seem to work with hotreload (tries to overwrite same DLL and not use a hotreload DLL name)
+
+        /** Used when BuildConfiguration.bUseExperimentalFastDependencyScan is enabled.  If true, it means that our cached includes may not longer be
+            valid (or never were), and we need to recover by forcibly scanning included headers for all build prerequisites to make sure that our
+            cached set of includes is actually correct, before determining which files are outdated. */
+        static public bool bNeedsFullCPPIncludeRescan = false;
 
 
         /// <summary>
@@ -631,10 +631,10 @@ namespace UnrealBuildTool
                     EngineSourceDirectory += "Source";
                 }
             }
-			if (Directory.Exists(EngineSourceDirectory)) // only set the directory if it exists, this should only happen if we are launching the editor from an artist sync
-			{
-				Directory.SetCurrentDirectory(EngineSourceDirectory);
-			}
+            if (Directory.Exists(EngineSourceDirectory)) // only set the directory if it exists, this should only happen if we are launching the editor from an artist sync
+            {
+                Directory.SetCurrentDirectory(EngineSourceDirectory);
+            }
 
             // Build the list of game projects that we know about. When building from the editor (for hot-reload) or for Rocket projects, we require the 
             // project file to be passed in. Otherwise we scan for projects in directories named in UE4Games.uprojectdirs.
@@ -673,6 +673,23 @@ namespace UnrealBuildTool
                 }
             }
 
+            bool bAutoSDKOnly = false;
+            {
+                int AutoSDKOnlyArgumentIndex;
+                if (Utils.ParseCommandLineFlag(Arguments, "-autosdkonly", out AutoSDKOnlyArgumentIndex))
+                {
+                    bAutoSDKOnly = true;
+                }
+            }
+            bool bValidatePlatforms = false;
+            {
+                int ValidatePlatformsArgumentIndex;
+                if (Utils.ParseCommandLineFlag(Arguments, "-validateplatform", out ValidatePlatformsArgumentIndex))
+                {
+                    bValidatePlatforms = true;
+                }
+            }
+            
             // Don't allow simultaneous execution of Unreal Built Tool. Multi-selection in the UI e.g. causes this and you want serial
             // execution in this case to avoid contention issues with shared produced items.
             bool bCreatedMutex = false;
@@ -681,7 +698,19 @@ namespace UnrealBuildTool
                 if (bUseMutex)
                 {
                     int LocationHash = Utils.GetExecutingAssemblyLocation().GetHashCode();
-                    SingleInstanceMutex = new Mutex(true, "Global/UnrealBuildTool_Mutex_" + LocationHash.ToString(), out bCreatedMutex);
+
+                    String MutexName;
+                    if (bAutoSDKOnly || bValidatePlatforms)
+                    {
+                        // this mutex has to be truly global because AutoSDK installs may change global state (like a console copying files into VS directories,
+                        // or installing Target Management services
+                        MutexName = "Global/UnrealBuildTool_Mutex_AutoSDKS";                        
+                    }
+                    else
+                    {                        
+                        MutexName = "Global/UnrealBuildTool_Mutex_" + LocationHash.ToString();
+                    }
+                    SingleInstanceMutex = new Mutex(true, MutexName, out bCreatedMutex);
                     if (!bCreatedMutex)
                     {
                         // If this instance didn't create the mutex, wait for the existing mutex to be released by the mutex's creator.
@@ -699,9 +728,7 @@ namespace UnrealBuildTool
                     var bGenerateMakefiles = false;
                     var bValidPlatformsOnly = false;
                     var bSpecificModulesOnly = false;
-                    var bIgnoreJunk = false;
-                    var bAutoSDKOnly = false;
-					var bValidatePlatforms = false;
+                    var bIgnoreJunk = false;                    
 
                     // We need to be able to identify if one of the arguments is the platform...
                     // Leverage the existing parser function in UEBuildTarget to get this information.
@@ -718,12 +745,12 @@ namespace UnrealBuildTool
                         }
                     }
 
-					// If we were asked to enable fast build iteration, we want the 'gather' phase to default to off (unless it is overridden below
-					// using a command-line option.)
-					if( BuildConfiguration.bUseExperimentalFastBuildIteration )
-					{
-						bIsGatheringBuild_Unsafe = false;
-					}
+                    // If we were asked to enable fast build iteration, we want the 'gather' phase to default to off (unless it is overridden below
+                    // using a command-line option.)
+                    if( BuildConfiguration.bUseExperimentalFastBuildIteration )
+                    {
+                        bIsGatheringBuild_Unsafe = false;
+                    }
 
                     foreach (var Arg in Arguments)
                     {
@@ -846,32 +873,32 @@ namespace UnrealBuildTool
                             BuildConfiguration.bXGEExport = true;
                             BuildConfiguration.bAllowXGE = true;
                         }
-						else if (LowercaseArg == "-gather")
-						{
-							UnrealBuildTool.bIsGatheringBuild_Unsafe = true;
-						}
-						else if (LowercaseArg == "-nogather")
-						{
-							UnrealBuildTool.bIsGatheringBuild_Unsafe = false;
-						}
-						else if (LowercaseArg == "-gatheronly")
-						{
-							UnrealBuildTool.bIsGatheringBuild_Unsafe = true;
-							UnrealBuildTool.bIsAssemblingBuild_Unsafe = false;
-						}
-						else if (LowercaseArg == "-assemble")
-						{
-							UnrealBuildTool.bIsAssemblingBuild_Unsafe = true;
-						}
-						else if (LowercaseArg == "-noassemble")
-						{
-							UnrealBuildTool.bIsAssemblingBuild_Unsafe = false;
-						}
-						else if (LowercaseArg == "-assembleonly")
-						{
-							UnrealBuildTool.bIsGatheringBuild_Unsafe = false;
-							UnrealBuildTool.bIsAssemblingBuild_Unsafe = true;
-						}
+                        else if (LowercaseArg == "-gather")
+                        {
+                            UnrealBuildTool.bIsGatheringBuild_Unsafe = true;
+                        }
+                        else if (LowercaseArg == "-nogather")
+                        {
+                            UnrealBuildTool.bIsGatheringBuild_Unsafe = false;
+                        }
+                        else if (LowercaseArg == "-gatheronly")
+                        {
+                            UnrealBuildTool.bIsGatheringBuild_Unsafe = true;
+                            UnrealBuildTool.bIsAssemblingBuild_Unsafe = false;
+                        }
+                        else if (LowercaseArg == "-assemble")
+                        {
+                            UnrealBuildTool.bIsAssemblingBuild_Unsafe = true;
+                        }
+                        else if (LowercaseArg == "-noassemble")
+                        {
+                            UnrealBuildTool.bIsAssemblingBuild_Unsafe = false;
+                        }
+                        else if (LowercaseArg == "-assembleonly")
+                        {
+                            UnrealBuildTool.bIsGatheringBuild_Unsafe = false;
+                            UnrealBuildTool.bIsAssemblingBuild_Unsafe = true;
+                        }
                         else if (LowercaseArg == "-nosimplygon")
                         {
                             UEBuildConfiguration.bCompileSimplygon = false;
@@ -891,42 +918,34 @@ namespace UnrealBuildTool
                         else if (LowercaseArg == "-progress")
                         {
                             ProgressWriter.bWriteMarkup = true;
-                        }
-                        else if (LowercaseArg == "-autosdkonly")
-                        {
-                            bAutoSDKOnly = true;                            
-                        }
+                        }                        
                         else if (CheckPlatform.ToString().ToLowerInvariant() == LowercaseArg)
                         {
                             // It's the platform set...
                             //PlatformName = Arg;
+                        }						
+                        else
+                        {
+                            // This arg may be a game name. Check for the existence of a game folder with this name.
+                            // "Engine" is not a valid game name.
+                            if (LowercaseArg != "engine" && Arg.IndexOfAny(Path.GetInvalidPathChars()) == -1 &&
+                                Directory.Exists(Path.Combine(ProjectFileGenerator.RootRelativePath, Arg, "Config")))
+                            {
+                                GameName = Arg;
+                                Log.TraceVerbose("CommandLine: Found game name '{0}'", GameName);
+                            }
+                            else if (LowercaseArg == "rocket")
+                            {
+                                GameName = Arg;
+                                Log.TraceVerbose("CommandLine: Found game name '{0}'", GameName);
+                            }
                         }
-						else if (LowercaseArg == "-validateplatform")
-						{
-							bValidatePlatforms = true;
-						}
-						else
-						{
-							// This arg may be a game name. Check for the existence of a game folder with this name.
-							// "Engine" is not a valid game name.
-							if (LowercaseArg != "engine" && Arg.IndexOfAny(Path.GetInvalidPathChars()) == -1 &&
-								Directory.Exists(Path.Combine(ProjectFileGenerator.RootRelativePath, Arg, "Config")))
-							{
-								GameName = Arg;
-								Log.TraceVerbose("CommandLine: Found game name '{0}'", GameName);
-							}
-							else if (LowercaseArg == "rocket")
-							{
-								GameName = Arg;
-								Log.TraceVerbose("CommandLine: Found game name '{0}'", GameName);
-							}
-						}
                     }
 
-					if( !bIsGatheringBuild_Unsafe && !bIsAssemblingBuild_Unsafe )
-					{
-						throw new BuildException( "UnrealBuildTool: At least one of either IsGatheringBuild or IsAssemblingBuild must be true.  Did you pass '-NoGather' with '-NoAssemble'?" );
-					}
+                    if( !bIsGatheringBuild_Unsafe && !bIsAssemblingBuild_Unsafe )
+                    {
+                        throw new BuildException( "UnrealBuildTool: At least one of either IsGatheringBuild or IsAssemblingBuild must be true.  Did you pass '-NoGather' with '-NoAssemble'?" );
+                    }
 
                     // Send an event with basic usage dimensions
                     Telemetry.SendEvent("CommonAttributes.2",
@@ -945,8 +964,8 @@ namespace UnrealBuildTool
                                         : bGenerateMakefiles
                                         ? "GenerateMakefiles"
                                             : bValidatePlatforms 
-												? "ValidatePlatfomrs"
-												: "Build",
+                                                ? "ValidatePlatfomrs"
+                                                : "Build",
                         "Platform", CheckPlatform.ToString(),
                         "Configuration", CheckConfiguration.ToString(),
                         "IsRocket", bRunningRocket.ToString(),
@@ -1004,68 +1023,68 @@ namespace UnrealBuildTool
                         // RegisterAllUBTClasses has already done all the SDK setup.
                         Result = ECompilationResult.Succeeded;
                     }
-					else if (bValidatePlatforms)
-					{
-						ValidatePlatforms(Arguments);
-					}
-					else
-					{
-						// Check if any third party headers are included from public engine headers.
-						if (bCheckThirdPartyHeaders)
-						{
-							ThirdPartyHeaderFinder.FindThirdPartyIncludes(CheckPlatform, CheckConfiguration);
-						}
-
-						// Build our project
-						if (Result == ECompilationResult.Succeeded)
-						{
-							if (UEBuildConfiguration.bPrepForDeployment == false)
-							{
-								int PassCount = 1;	// @todo ubtmake: It's useful to set this to a high number when running UBT under a sampling profiler
-								for( int PassIndex = 0; PassIndex < PassCount; ++PassIndex )
-								{ 
-									// If we are only prepping for deployment, assume the build already occurred.
-									Result = RunUBT(Arguments);
-									if( Result != ECompilationResult.Succeeded )
-									{
-										break;
-									}
-								}
-							}
-							else
-							{
-								var BuildPlatform = UEBuildPlatform.GetBuildPlatform(CheckPlatform, true);
-								if (BuildPlatform != null)
-								{
-									// Setup environment wasn't called, so set the flag
-									BuildConfiguration.bDeployAfterCompile = BuildPlatform.RequiresDeployPrepAfterCompile();
-									BuildConfiguration.PlatformIntermediateFolder = Path.Combine(BuildConfiguration.BaseIntermediateFolder, CheckPlatform.ToString(), BuildPlatform.GetActiveArchitecture());
-								}
-							}
-
-							// If we build w/ bXGEExport true, we didn't REALLY build at this point, 
-							// so don't bother with doing the PrepTargetForDeployment call. 
-							if ((Result == ECompilationResult.Succeeded) && (BuildConfiguration.bDeployAfterCompile == true) && (BuildConfiguration.bXGEExport == false) &&
-								(UEBuildConfiguration.bGenerateManifest == false) && (UEBuildConfiguration.bGenerateExternalFileList == false) && (UEBuildConfiguration.bCleanProject == false))
-							{
-								var DeployHandler = UEBuildDeploy.GetBuildDeploy(CheckPlatform);
-								if (DeployHandler != null)
-								{
-									// We need to be able to identify the Target.Type we can derive it from the Arguments.
-									BuildConfiguration.bFlushBuildDirOnRemoteMac = false;
-									UEBuildTarget CheckTarget = UEBuildTarget.CreateTarget(Arguments);
-									CheckTarget.SetupGlobalEnvironment();
-									if ((CheckTarget.Rules.Type == TargetRules.TargetType.Game) ||
-										(CheckTarget.Rules.Type == TargetRules.TargetType.Server) ||
-										(CheckTarget.Rules.Type == TargetRules.TargetType.Client))
-									{
-										CheckTarget.AppName = CheckTarget.GameName;
-									}
-									DeployHandler.PrepTargetForDeployment(CheckTarget);
-								}
-							}
-						}
-					}
+                    else if (bValidatePlatforms)
+                    {
+                        ValidatePlatforms(Arguments);
+                    }
+                    else
+                    {
+                        // Check if any third party headers are included from public engine headers.
+                        if (bCheckThirdPartyHeaders)
+                        {
+                            ThirdPartyHeaderFinder.FindThirdPartyIncludes(CheckPlatform, CheckConfiguration);
+                        }
+
+                        // Build our project
+                        if (Result == ECompilationResult.Succeeded)
+                        {
+                            if (UEBuildConfiguration.bPrepForDeployment == false)
+                            {
+                                int PassCount = 1;	// @todo ubtmake: It's useful to set this to a high number when running UBT under a sampling profiler
+                                for( int PassIndex = 0; PassIndex < PassCount; ++PassIndex )
+                                { 
+                                    // If we are only prepping for deployment, assume the build already occurred.
+                                    Result = RunUBT(Arguments);
+                                    if( Result != ECompilationResult.Succeeded )
+                                    {
+                                        break;
+                                    }
+                                }
+                            }
+                            else
+                            {
+                                var BuildPlatform = UEBuildPlatform.GetBuildPlatform(CheckPlatform, true);
+                                if (BuildPlatform != null)
+                                {
+                                    // Setup environment wasn't called, so set the flag
+                                    BuildConfiguration.bDeployAfterCompile = BuildPlatform.RequiresDeployPrepAfterCompile();
+                                    BuildConfiguration.PlatformIntermediateFolder = Path.Combine(BuildConfiguration.BaseIntermediateFolder, CheckPlatform.ToString(), BuildPlatform.GetActiveArchitecture());
+                                }
+                            }
+
+                            // If we build w/ bXGEExport true, we didn't REALLY build at this point, 
+                            // so don't bother with doing the PrepTargetForDeployment call. 
+                            if ((Result == ECompilationResult.Succeeded) && (BuildConfiguration.bDeployAfterCompile == true) && (BuildConfiguration.bXGEExport == false) &&
+                                (UEBuildConfiguration.bGenerateManifest == false) && (UEBuildConfiguration.bGenerateExternalFileList == false) && (UEBuildConfiguration.bCleanProject == false))
+                            {
+                                var DeployHandler = UEBuildDeploy.GetBuildDeploy(CheckPlatform);
+                                if (DeployHandler != null)
+                                {
+                                    // We need to be able to identify the Target.Type we can derive it from the Arguments.
+                                    BuildConfiguration.bFlushBuildDirOnRemoteMac = false;
+                                    UEBuildTarget CheckTarget = UEBuildTarget.CreateTarget(Arguments);
+                                    CheckTarget.SetupGlobalEnvironment();
+                                    if ((CheckTarget.Rules.Type == TargetRules.TargetType.Game) ||
+                                        (CheckTarget.Rules.Type == TargetRules.TargetType.Server) ||
+                                        (CheckTarget.Rules.Type == TargetRules.TargetType.Client))
+                                    {
+                                        CheckTarget.AppName = CheckTarget.GameName;
+                                    }
+                                    DeployHandler.PrepTargetForDeployment(CheckTarget);
+                                }
+                            }
+                        }
+                    }
                     // Print some performance info
                     var BuildDuration = (DateTime.UtcNow - StartTime - MutexWaitTime).TotalSeconds;
                     if (BuildConfiguration.bPrintPerformanceInfo)
@@ -1137,90 +1156,90 @@ namespace UnrealBuildTool
 
 
 
-		/// <summary>
-		/// Validates the various platforms to determine if they are ready for building
-		/// </summary>
-		public static void ValidatePlatforms(string[] Arguments)
-		{
-			List<UnrealTargetPlatform> Platforms = new List<UnrealTargetPlatform>(); 
-			foreach (var CurArgument in Arguments)
-			{
-				if (CurArgument.StartsWith("-"))
-				{
-					if (CurArgument.StartsWith("-Platforms=", StringComparison.InvariantCultureIgnoreCase))
-					{
-						// Parse the list... will be in Foo+Bar+New format
-						string PlatformList = CurArgument.Substring(11);
-						while (PlatformList.Length > 0)
-						{
-							string PlatformString = PlatformList;
-							Int32 PlusIdx = PlatformList.IndexOf("+");
-							if (PlusIdx != -1)
-							{
-								PlatformString = PlatformList.Substring(0, PlusIdx);
-								PlatformList = PlatformList.Substring(PlusIdx + 1);
-							}
-							else
-							{
-								// We are on the last platform... clear the list to exit the loop
-								PlatformList = "";
-							}
-
-							// Is the string a valid platform? If so, add it to the list
-							UnrealTargetPlatform SpecifiedPlatform = UnrealTargetPlatform.Unknown;
-							foreach (UnrealTargetPlatform PlatformParam in Enum.GetValues(typeof(UnrealTargetPlatform)))
-							{
-								if (PlatformString.Equals(PlatformParam.ToString(), StringComparison.InvariantCultureIgnoreCase))
-								{
-									SpecifiedPlatform = PlatformParam;
-									break;
-								}
-							}
-
-							if (SpecifiedPlatform != UnrealTargetPlatform.Unknown)
-							{
-								if (Platforms.Contains(SpecifiedPlatform) == false)
-								{
-									Platforms.Add(SpecifiedPlatform);
-								}
-							}
-							else
-							{
-								Log.TraceWarning("ValidatePlatforms invalid platform specified: {0}", PlatformString);
-							}
-						}
-					}
-					else switch (CurArgument.ToUpperInvariant())
-					{
-						case "-ALLPLATFORMS":
-							foreach (UnrealTargetPlatform platform in Enum.GetValues(typeof(UnrealTargetPlatform)))
-							{
-								if (platform != UnrealTargetPlatform.Unknown)
-								{
-									if (Platforms.Contains(platform) == false)
-									{
-										Platforms.Add(platform);
-									}
-								}
-							}
-							break;
-					}
-				}
-			}
-
-			foreach (UnrealTargetPlatform platform in Platforms)
-			{
-				var BuildPlatform = UEBuildPlatform.GetBuildPlatform(platform, true);
-				if (BuildPlatform != null && BuildPlatform.HasRequiredSDKsInstalled() == SDKStatus.Valid)
-				{
-					Console.WriteLine("##PlatformValidate: {0} VALID", platform.ToString());
-				}
-				else
-				{
-					Console.WriteLine("##PlatformValidate: {0} INVALID", platform.ToString());
-				}
-			}
-		}
+        /// <summary>
+        /// Validates the various platforms to determine if they are ready for building
+        /// </summary>
+        public static void ValidatePlatforms(string[] Arguments)
+        {
+            List<UnrealTargetPlatform> Platforms = new List<UnrealTargetPlatform>(); 
+            foreach (var CurArgument in Arguments)
+            {
+                if (CurArgument.StartsWith("-"))
+                {
+                    if (CurArgument.StartsWith("-Platforms=", StringComparison.InvariantCultureIgnoreCase))
+                    {
+                        // Parse the list... will be in Foo+Bar+New format
+                        string PlatformList = CurArgument.Substring(11);
+                        while (PlatformList.Length > 0)
+                        {
+                            string PlatformString = PlatformList;
+                            Int32 PlusIdx = PlatformList.IndexOf("+");
+                            if (PlusIdx != -1)
+                            {
+                                PlatformString = PlatformList.Substring(0, PlusIdx);
+                                PlatformList = PlatformList.Substring(PlusIdx + 1);
+                            }
+                            else
+                            {
+                                // We are on the last platform... clear the list to exit the loop
+                                PlatformList = "";
+                            }
+
+                            // Is the string a valid platform? If so, add it to the list
+                            UnrealTargetPlatform SpecifiedPlatform = UnrealTargetPlatform.Unknown;
+                            foreach (UnrealTargetPlatform PlatformParam in Enum.GetValues(typeof(UnrealTargetPlatform)))
+                            {
+                                if (PlatformString.Equals(PlatformParam.ToString(), StringComparison.InvariantCultureIgnoreCase))
+                                {
+                                    SpecifiedPlatform = PlatformParam;
+                                    break;
+                                }
+                            }
+
+                            if (SpecifiedPlatform != UnrealTargetPlatform.Unknown)
+                            {
+                                if (Platforms.Contains(SpecifiedPlatform) == false)
+                                {
+                                    Platforms.Add(SpecifiedPlatform);
+                                }
+                            }
+                            else
+                            {
+                                Log.TraceWarning("ValidatePlatforms invalid platform specified: {0}", PlatformString);
+                            }
+                        }
+                    }
+                    else switch (CurArgument.ToUpperInvariant())
+                    {
+                        case "-ALLPLATFORMS":
+                            foreach (UnrealTargetPlatform platform in Enum.GetValues(typeof(UnrealTargetPlatform)))
+                            {
+                                if (platform != UnrealTargetPlatform.Unknown)
+                                {
+                                    if (Platforms.Contains(platform) == false)
+                                    {
+                                        Platforms.Add(platform);
+                                    }
+                                }
+                            }
+                            break;
+                    }
+                }
+            }
+
+            foreach (UnrealTargetPlatform platform in Platforms)
+            {
+                var BuildPlatform = UEBuildPlatform.GetBuildPlatform(platform, true);
+                if (BuildPlatform != null && BuildPlatform.HasRequiredSDKsInstalled() == SDKStatus.Valid)
+                {
+                    Console.WriteLine("##PlatformValidate: {0} VALID", platform.ToString());
+                }
+                else
+                {
+                    Console.WriteLine("##PlatformValidate: {0} INVALID", platform.ToString());
+                }
+            }
+        }
 
         // @todo: Ideally get rid of RunUBT() and all of the Clear/Reset stuff!
         public static ECompilationResult RunUBT(string[] Arguments)
@@ -1286,191 +1305,191 @@ namespace UnrealBuildTool
                         continue;
                     }
                     Targets.Add(Target);
-				}
-
-				UBTMakefile UBTMakefile = null;
-
-
-				// If UBT is being used to build something different than it did last time
-				if( !ProjectFileGenerator.bGenerateProjectFiles )
-				{
-					bool bIsBuildingSameTargetsAsLastTime = false;
-					{
-						// If we're building UHT without a Mutex, we'll need to assume that we're building the same targets and that no caches
-						// should be invalidated for this run.  This is important when UBT is invoked from within UBT in order to compile
-						// UHT.  In that (very common) case we definitely don't want to have to rebuild our cache from scratch.
-						bool bMustAssumeSameTargets = false;
-
-						if( Targets[0].GetTargetName().Equals( "UnrealHeaderTool", StringComparison.InvariantCultureIgnoreCase ) )
-						{
-							int NoMutexArgumentIndex;
-							if (Utils.ParseCommandLineFlag(Arguments, "-NoMutex", out NoMutexArgumentIndex))
-							{
-								bMustAssumeSameTargets = true;
-							}
-						}
-
-						if( bMustAssumeSameTargets )
-						{
-							bIsBuildingSameTargetsAsLastTime = true;
-						}
-						else
-						{
-							string TargetCollectionName = MakeTargetCollectionName( Targets );
-							string LastBuiltTargetsFilePath = Path.Combine( BuildConfiguration.BaseIntermediatePath, "LastBuiltTargets.txt" );
-							if( File.Exists( LastBuiltTargetsFilePath ) && Utils.ReadAllText( LastBuiltTargetsFilePath ) == TargetCollectionName )
-							{
-								bIsBuildingSameTargetsAsLastTime = true;
-							}
-
-							// Save out the name of the targets we're building
-							if( !bIsBuildingSameTargetsAsLastTime )
-							{
-								Directory.CreateDirectory( Path.GetDirectoryName( LastBuiltTargetsFilePath ) );
-								File.WriteAllText( LastBuiltTargetsFilePath, TargetCollectionName, Encoding.UTF8 );
-							}
-
-
-							if( !bIsBuildingSameTargetsAsLastTime )
-							{
-								// Can't use super fast include checking unless we're building the same set of targets as last time, because
-								// we might not know about all of the C++ include dependencies for already-up-to-date shared build products
-								// between the targets
-								bNeedsFullCPPIncludeRescan = true;
-							}
-						}
-					}
-				}
-			
-
-
-				{ 
-					// If we're generating project files, then go ahead and wipe out the existing UBTMakefile for every target, to make sure that
-					// it gets a full dependency scan next time.
-					// NOTE: This is just a safeguard and doesn't have to be perfect.  We also check for newer project file timestamps in LoadUBTMakefile()
-					if( ProjectFileGenerator.bGenerateProjectFiles )	// @todo ubtmake: This is only hit when generating IntelliSense for project files.  Probably should be done right inside ProjectFileGenerator.bat
-					{													// @todo ubtmake: Won't catch multi-target cases as GPF always builds one target at a time for Intellisense
-						// Delete the UBT makefile
-						var UBTMakefilePath = UnrealBuildTool.GetUBTMakefilePath( Targets );
-						if (File.Exists(UBTMakefilePath))
-						{
-							UEBuildTarget.CleanFile(UBTMakefilePath);
-						}
-					}
-
-					// Make sure the gather phase is executed if we're not actually building anything
-					if( ProjectFileGenerator.bGenerateProjectFiles || UEBuildConfiguration.bGenerateManifest || UEBuildConfiguration.bCleanProject )
-					{
-						UnrealBuildTool.bIsGatheringBuild_Unsafe = true;						
-					}
-
-					// Were we asked to run in 'assembler only' mode?  If so, let's check to see if that's even possible by seeing if
-					// we have a valid UBT Makefile already saved to disk, ready for us to load.
-					if( UnrealBuildTool.bIsAssemblingBuild_Unsafe && !UnrealBuildTool.bIsGatheringBuild_Unsafe )
-					{
-						// @todo ubtmake: Mildly terrified of BuildConfiguration/UEBuildConfiguration globals that were set during the Prepare phase but are not set now.  We may need to save/load all of these, otherwise
-						//		we'll need to call SetupGlobalEnvironment on all of the targets (maybe other stuff, too.  See PreBuildStep())
-							
-						UBTMakefile = LoadUBTMakefile( Targets );
-
-						if( UBTMakefile == null )
-						{ 
-							// If the Makefile couldn't be loaded, then we're not going to be able to continue in "assembler only" mode.  We'll do both
-							// a 'gather' and 'assemble' in the same run.  This will take a while longer, but subsequent runs will be fast!
-							UnrealBuildTool.bIsGatheringBuild_Unsafe = true;
-
-							Log.TraceInformation( "Forcing generation of new UBTMakefile for {0}{1}", Targets[0].GetTargetName(), Targets.Count > 1 ? ( " (and " + ( Targets.Count - 1 ).ToString() + " more)" ) : "" );
-						}
-					}
-
-					// OK, after this point it is safe to access the UnrealBuildTool.IsGatheringBuild and UnrealBuildTool.IsAssemblingBuild properties.
-					// These properties will not be changing again during this session/
-					bIsSafeToCheckIfGatheringOrAssemblingBuild = true;
-				}
+                }
+
+                UBTMakefile UBTMakefile = null;
+
+
+                // If UBT is being used to build something different than it did last time
+                if( !ProjectFileGenerator.bGenerateProjectFiles )
+                {
+                    bool bIsBuildingSameTargetsAsLastTime = false;
+                    {
+                        // If we're building UHT without a Mutex, we'll need to assume that we're building the same targets and that no caches
+                        // should be invalidated for this run.  This is important when UBT is invoked from within UBT in order to compile
+                        // UHT.  In that (very common) case we definitely don't want to have to rebuild our cache from scratch.
+                        bool bMustAssumeSameTargets = false;
+
+                        if( Targets[0].GetTargetName().Equals( "UnrealHeaderTool", StringComparison.InvariantCultureIgnoreCase ) )
+                        {
+                            int NoMutexArgumentIndex;
+                            if (Utils.ParseCommandLineFlag(Arguments, "-NoMutex", out NoMutexArgumentIndex))
+                            {
+                                bMustAssumeSameTargets = true;
+                            }
+                        }
+
+                        if( bMustAssumeSameTargets )
+                        {
+                            bIsBuildingSameTargetsAsLastTime = true;
+                        }
+                        else
+                        {
+                            string TargetCollectionName = MakeTargetCollectionName( Targets );
+                            string LastBuiltTargetsFilePath = Path.Combine( BuildConfiguration.BaseIntermediatePath, "LastBuiltTargets.txt" );
+                            if( File.Exists( LastBuiltTargetsFilePath ) && Utils.ReadAllText( LastBuiltTargetsFilePath ) == TargetCollectionName )
+                            {
+                                bIsBuildingSameTargetsAsLastTime = true;
+                            }
+
+                            // Save out the name of the targets we're building
+                            if( !bIsBuildingSameTargetsAsLastTime )
+                            {
+                                Directory.CreateDirectory( Path.GetDirectoryName( LastBuiltTargetsFilePath ) );
+                                File.WriteAllText( LastBuiltTargetsFilePath, TargetCollectionName, Encoding.UTF8 );
+                            }
+
+
+                            if( !bIsBuildingSameTargetsAsLastTime )
+                            {
+                                // Can't use super fast include checking unless we're building the same set of targets as last time, because
+                                // we might not know about all of the C++ include dependencies for already-up-to-date shared build products
+                                // between the targets
+                                bNeedsFullCPPIncludeRescan = true;
+                            }
+                        }
+                    }
+                }
+            
+
+
+                { 
+                    // If we're generating project files, then go ahead and wipe out the existing UBTMakefile for every target, to make sure that
+                    // it gets a full dependency scan next time.
+                    // NOTE: This is just a safeguard and doesn't have to be perfect.  We also check for newer project file timestamps in LoadUBTMakefile()
+                    if( ProjectFileGenerator.bGenerateProjectFiles )	// @todo ubtmake: This is only hit when generating IntelliSense for project files.  Probably should be done right inside ProjectFileGenerator.bat
+                    {													// @todo ubtmake: Won't catch multi-target cases as GPF always builds one target at a time for Intellisense
+                        // Delete the UBT makefile
+                        var UBTMakefilePath = UnrealBuildTool.GetUBTMakefilePath( Targets );
+                        if (File.Exists(UBTMakefilePath))
+                        {
+                            UEBuildTarget.CleanFile(UBTMakefilePath);
+                        }
+                    }
+
+                    // Make sure the gather phase is executed if we're not actually building anything
+                    if( ProjectFileGenerator.bGenerateProjectFiles || UEBuildConfiguration.bGenerateManifest || UEBuildConfiguration.bCleanProject )
+                    {
+                        UnrealBuildTool.bIsGatheringBuild_Unsafe = true;						
+                    }
+
+                    // Were we asked to run in 'assembler only' mode?  If so, let's check to see if that's even possible by seeing if
+                    // we have a valid UBT Makefile already saved to disk, ready for us to load.
+                    if( UnrealBuildTool.bIsAssemblingBuild_Unsafe && !UnrealBuildTool.bIsGatheringBuild_Unsafe )
+                    {
+                        // @todo ubtmake: Mildly terrified of BuildConfiguration/UEBuildConfiguration globals that were set during the Prepare phase but are not set now.  We may need to save/load all of these, otherwise
+                        //		we'll need to call SetupGlobalEnvironment on all of the targets (maybe other stuff, too.  See PreBuildStep())
+                            
+                        UBTMakefile = LoadUBTMakefile( Targets );
+
+                        if( UBTMakefile == null )
+                        { 
+                            // If the Makefile couldn't be loaded, then we're not going to be able to continue in "assembler only" mode.  We'll do both
+                            // a 'gather' and 'assemble' in the same run.  This will take a while longer, but subsequent runs will be fast!
+                            UnrealBuildTool.bIsGatheringBuild_Unsafe = true;
+
+                            Log.TraceInformation( "Forcing generation of new UBTMakefile for {0}{1}", Targets[0].GetTargetName(), Targets.Count > 1 ? ( " (and " + ( Targets.Count - 1 ).ToString() + " more)" ) : "" );
+                        }
+                    }
+
+                    // OK, after this point it is safe to access the UnrealBuildTool.IsGatheringBuild and UnrealBuildTool.IsAssemblingBuild properties.
+                    // These properties will not be changing again during this session/
+                    bIsSafeToCheckIfGatheringOrAssemblingBuild = true;
+                }
 
 
                 // Build action lists for all passed in targets.
                 var OutputItemsForAllTargets = new List<FileItem>();
-				var TargetNameToUObjectModules = new Dictionary<string, List<UHTModuleInfo>>( StringComparer.InvariantCultureIgnoreCase );
+                var TargetNameToUObjectModules = new Dictionary<string, List<UHTModuleInfo>>( StringComparer.InvariantCultureIgnoreCase );
                 foreach (var Target in Targets)
                 {
                     var TargetStartTime = DateTime.UtcNow;
 
-					// When in 'assembler only' mode, we'll load this cache later on a worker thread.  It takes a long time to load!
-					if( !( !UnrealBuildTool.IsGatheringBuild && UnrealBuildTool.IsAssemblingBuild ) )
-					{	
-						// Load the direct include dependency cache.
-						CPPEnvironment.IncludeDependencyCache.Add( Target, DependencyCache.Create( DependencyCache.GetDependencyCachePathForTarget(Target) ) );
-					}
-
-					// We don't need this dependency cache in 'gather only' mode
-					if( BuildConfiguration.bUseExperimentalFastDependencyScan &&
-						!( UnrealBuildTool.IsGatheringBuild && !UnrealBuildTool.IsAssemblingBuild ) )
-					{ 
-						// Load the cache that contains the list of flattened resolved includes for resolved source files
-						// @todo fastubt: Ideally load this asynchronously at startup and only block when it is first needed and not finished loading
-						CPPEnvironment.FlatCPPIncludeDependencyCache.Add( Target, FlatCPPIncludeDependencyCache.Create( Target ) );
-					}
-
-					if( UnrealBuildTool.IsGatheringBuild )
-					{ 
-	                    List<FileItem> TargetOutputItems;
-						List<UHTModuleInfo> TargetUObjectModules;
-						BuildResult = Target.Build(ToolChain, out TargetOutputItems, out TargetUObjectModules, out EULAViolationWarning);
-						if(BuildResult != ECompilationResult.Succeeded)
-						{
-							break;
-						}
-
-						OutputItemsForAllTargets.AddRange( TargetOutputItems );
-
-						// Update mapping of the target name to the list of UObject modules in this target
-						TargetNameToUObjectModules[ Target.GetTargetName() ] = TargetUObjectModules;
-
-						if ( (BuildConfiguration.bXGEExport && UEBuildConfiguration.bGenerateManifest) || (!ProjectFileGenerator.bGenerateProjectFiles && !UEBuildConfiguration.bGenerateManifest && !UEBuildConfiguration.bCleanProject))
-						{
-							// We don't currently support building game targets in rocket.
-							if (UnrealBuildTool.RunningRocket() && Target.Rules != null && Target.Rules.Type == TargetRules.TargetType.Game)
-							{
-								throw new BuildException(
-									"You currently can not build a game target in Rocket.\nTry again in a future release.\nFor now, build and run your editor project."
-								);
-							}
-
-							// Generate an action graph if we were asked to do that.  The graph generation needs access to the include dependency cache, so
-							// we generate it before saving and cleaning that up.
-							if( GeneratingActionGraph )
-							{
-								// The graph generation feature currently only works with a single target at a time.  This is because of how we need access
-								// to include dependencies for the target, but those aren't kept around as we process multiple targets
-								if( TargetSettings.Count != 1 )
-								{
-									throw new BuildException( "ERROR: The '-graph' option only works with a single target at a time" );
-								}
-								ActionGraph.FinalizeActionGraph();
-								var ActionsToExecute = ActionGraph.AllActions;
-
-								var VisualizationType = ActionGraph.ActionGraphVisualizationType.OnlyCPlusPlusFilesAndHeaders;
-								ActionGraph.SaveActionGraphVisualization( Target, Path.Combine( BuildConfiguration.BaseIntermediatePath, Target.GetTargetName() + ".gexf" ), Target.GetTargetName(), VisualizationType, ActionsToExecute );
-							}
-						}
-
-						var TargetBuildTime = (DateTime.UtcNow - TargetStartTime).TotalSeconds;
-
-						// Send out telemetry for this target
-						Telemetry.SendEvent("TargetBuildStats.2",
-							"AppName", Target.AppName,
-							"GameName", Target.GameName,
-							"Platform", Target.Platform.ToString(),
-							"Configuration", Target.Configuration.ToString(),
-							"CleanTarget", UEBuildConfiguration.bCleanProject.ToString(),
-							"Monolithic", Target.ShouldCompileMonolithic().ToString(),
-							"CreateDebugInfo", Target.IsCreatingDebugInfo().ToString(),
-							"TargetType", Target.Rules.Type.ToString(),
-							"TargetCreateTimeSec", TargetBuildTime.ToString("0.00")
-							);
-					}
-				}
+                    // When in 'assembler only' mode, we'll load this cache later on a worker thread.  It takes a long time to load!
+                    if( !( !UnrealBuildTool.IsGatheringBuild && UnrealBuildTool.IsAssemblingBuild ) )
+                    {	
+                        // Load the direct include dependency cache.
+                        CPPEnvironment.IncludeDependencyCache.Add( Target, DependencyCache.Create( DependencyCache.GetDependencyCachePathForTarget(Target) ) );
+                    }
+
+                    // We don't need this dependency cache in 'gather only' mode
+                    if( BuildConfiguration.bUseExperimentalFastDependencyScan &&
+                        !( UnrealBuildTool.IsGatheringBuild && !UnrealBuildTool.IsAssemblingBuild ) )
+                    { 
+                        // Load the cache that contains the list of flattened resolved includes for resolved source files
+                        // @todo fastubt: Ideally load this asynchronously at startup and only block when it is first needed and not finished loading
+                        CPPEnvironment.FlatCPPIncludeDependencyCache.Add( Target, FlatCPPIncludeDependencyCache.Create( Target ) );
+                    }
+
+                    if( UnrealBuildTool.IsGatheringBuild )
+                    { 
+                        List<FileItem> TargetOutputItems;
+                        List<UHTModuleInfo> TargetUObjectModules;
+                        BuildResult = Target.Build(ToolChain, out TargetOutputItems, out TargetUObjectModules, out EULAViolationWarning);
+                        if(BuildResult != ECompilationResult.Succeeded)
+                        {
+                            break;
+                        }
+
+                        OutputItemsForAllTargets.AddRange( TargetOutputItems );
+
+                        // Update mapping of the target name to the list of UObject modules in this target
+                        TargetNameToUObjectModules[ Target.GetTargetName() ] = TargetUObjectModules;
+
+                        if ( (BuildConfiguration.bXGEExport && UEBuildConfiguration.bGenerateManifest) || (!ProjectFileGenerator.bGenerateProjectFiles && !UEBuildConfiguration.bGenerateManifest && !UEBuildConfiguration.bCleanProject))
+                        {
+                            // We don't currently support building game targets in rocket.
+                            if (UnrealBuildTool.RunningRocket() && Target.Rules != null && Target.Rules.Type == TargetRules.TargetType.Game)
+                            {
+                                throw new BuildException(
+                                    "You currently can not build a game target in Rocket.\nTry again in a future release.\nFor now, build and run your editor project."
+                                );
+                            }
+
+                            // Generate an action graph if we were asked to do that.  The graph generation needs access to the include dependency cache, so
+                            // we generate it before saving and cleaning that up.
+                            if( GeneratingActionGraph )
+                            {
+                                // The graph generation feature currently only works with a single target at a time.  This is because of how we need access
+                                // to include dependencies for the target, but those aren't kept around as we process multiple targets
+                                if( TargetSettings.Count != 1 )
+                                {
+                                    throw new BuildException( "ERROR: The '-graph' option only works with a single target at a time" );
+                                }
+                                ActionGraph.FinalizeActionGraph();
+                                var ActionsToExecute = ActionGraph.AllActions;
+
+                                var VisualizationType = ActionGraph.ActionGraphVisualizationType.OnlyCPlusPlusFilesAndHeaders;
+                                ActionGraph.SaveActionGraphVisualization( Target, Path.Combine( BuildConfiguration.BaseIntermediatePath, Target.GetTargetName() + ".gexf" ), Target.GetTargetName(), VisualizationType, ActionsToExecute );
+                            }
+                        }
+
+                        var TargetBuildTime = (DateTime.UtcNow - TargetStartTime).TotalSeconds;
+
+                        // Send out telemetry for this target
+                        Telemetry.SendEvent("TargetBuildStats.2",
+                            "AppName", Target.AppName,
+                            "GameName", Target.GameName,
+                            "Platform", Target.Platform.ToString(),
+                            "Configuration", Target.Configuration.ToString(),
+                            "CleanTarget", UEBuildConfiguration.bCleanProject.ToString(),
+                            "Monolithic", Target.ShouldCompileMonolithic().ToString(),
+                            "CreateDebugInfo", Target.IsCreatingDebugInfo().ToString(),
+                            "TargetType", Target.Rules.Type.ToString(),
+                            "TargetCreateTimeSec", TargetBuildTime.ToString("0.00")
+                            );
+                    }
+                }
 
                 if (BuildResult == ECompilationResult.Succeeded &&
                     (
@@ -1478,162 +1497,162 @@ namespace UnrealBuildTool
                         (!GeneratingActionGraph && !ProjectFileGenerator.bGenerateProjectFiles && !UEBuildConfiguration.bGenerateManifest && !UEBuildConfiguration.bCleanProject && !UEBuildConfiguration.bGenerateExternalFileList)
                     ))
                 {
-					if( UnrealBuildTool.IsGatheringBuild )
-					{
-						ActionGraph.FinalizeActionGraph();
-
-						UBTMakefile = new UBTMakefile();
-						UBTMakefile.AllActions = ActionGraph.AllActions;
-
-						// For now simply treat all object files as the root target.
-						{ 
-							var PrerequisiteActionsSet = new HashSet<Action>();
-							foreach (var OutputItem in OutputItemsForAllTargets)
-							{
-								ActionGraph.GatherPrerequisiteActions(OutputItem, ref PrerequisiteActionsSet);
-							}
-							UBTMakefile.PrerequisiteActions = PrerequisiteActionsSet.ToArray();
-						}			
-
-						foreach( System.Collections.DictionaryEntry EnvironmentVariable in Environment.GetEnvironmentVariables() )
-						{
-							UBTMakefile.EnvironmentVariables.Add( Tuple.Create( (string)EnvironmentVariable.Key, (string)EnvironmentVariable.Value ) );
-						}
-						UBTMakefile.TargetNameToUObjectModules = TargetNameToUObjectModules;
-
-						// We've been told to prepare to build, so let's go ahead and save out our action graph so that we can use in a later invocation 
-						// to assemble the build.  Even if we are configured to assemble the build in this same invocation, we want to save out the
-						// Makefile so that it can be used on subsequent 'assemble only' runs, for the fastest possible iteration times
-						// @todo ubtmake: Optimization: We could make 'gather + assemble' mode slightly faster by saving this while busy compiling (on our worker thread)
-						SaveUBTMakefile( Targets, UBTMakefile );
-					}
-
-					if( UnrealBuildTool.IsAssemblingBuild )
-					{
-						// If we didn't build the graph in this session, then we'll need to load a cached one
-						if( !UnrealBuildTool.IsGatheringBuild )
-						{
-							ActionGraph.AllActions = UBTMakefile.AllActions;
-
-							foreach( var EnvironmentVariable in UBTMakefile.EnvironmentVariables )
-							{
-								// @todo ubtmake: There may be some variables we do NOT want to clobber.
-								Environment.SetEnvironmentVariable( EnvironmentVariable.Item1, EnvironmentVariable.Item2 );
-							}
-
-							// If any of the targets need UHT to be run, we'll go ahead and do that now
-							foreach( var Target in Targets )
-							{
-								List<UHTModuleInfo> TargetUObjectModules;
-								if( UBTMakefile.TargetNameToUObjectModules.TryGetValue( Target.GetTargetName(), out TargetUObjectModules ) )
-								{
-									if( TargetUObjectModules.Count > 0 )
-									{
-										// Execute the header tool
-										string ModuleInfoFileName = Path.GetFullPath( Path.Combine( Target.ProjectIntermediateDirectory, "UnrealHeaderTool.manifest" ) );
-										ECompilationResult UHTResult = ECompilationResult.OtherCompilationError;
-										if (!ExternalExecution.ExecuteHeaderToolIfNecessary(Target, GlobalCompileEnvironment:null, UObjectModules:TargetUObjectModules, ModuleInfoFileName:ModuleInfoFileName, UHTResult:ref UHTResult))
-										{
-											Log.TraceInformation("UnrealHeaderTool failed for target '" + Target.GetTargetName() + "' (platform: " + Target.GetTargetInfo().Platform.ToString() + ", module info: " + ModuleInfoFileName + ").");
-											BuildResult = UHTResult;
-											break;
-										}
-									}
-								}
-							}
-						}
-
-						if( BuildResult == ECompilationResult.Succeeded )
-						{ 
-							// Make sure any old DLL files from in-engine recompiles aren't lying around.  Must be called after the action graph is finalized.
-							// @todo ubtmake: Will we even use traditional hot reload DLLs anymore in the future?  
-							ActionGraph.DeleteStaleHotReloadDLLs();
-
-							// Plan the actions to execute for the build.
-							Dictionary<UEBuildTarget,List<FileItem>> TargetToOutdatedPrerequisitesMap;
-							List<Action> ActionsToExecute = ActionGraph.GetActionsToExecute(UBTMakefile.PrerequisiteActions, Targets, out TargetToOutdatedPrerequisitesMap);
-
-							// Display some stats to the user.
-							Log.TraceVerbose(
-									"{0} actions, {1} outdated and requested actions",
-									ActionGraph.AllActions.Count,
-									ActionsToExecute.Count
-									);
-
-							ToolChain.PreBuildSync();
-
-					
-							// Cache indirect includes for all outdated C++ files.  We kick this off as a background thread so that it can
-							// perform the scan while we're compiling.  It usually only takes up to a few seconds, but we don't want to hurt
-							// our best case UBT iteration times for this task which can easily be performed asynchronously
-							Thread CPPIncludesThread = null;
-							if( BuildConfiguration.bUseExperimentalFastDependencyScan && TargetToOutdatedPrerequisitesMap.Count > 0 )
-							{
-								CPPIncludesThread = CreateThreadForCachingCPPIncludes( TargetToOutdatedPrerequisitesMap );
-								CPPIncludesThread.Start();
-							}
-
-							// Execute the actions.
-							bSuccess = ActionGraph.ExecuteActions(ActionsToExecute, out ExecutorName);
-
-							// if the build succeeded, do any needed syncing
-							if (bSuccess)
-							{
-								foreach (UEBuildTarget Target in Targets)
-								{
-									ToolChain.PostBuildSync(Target);
-								}
-							}
-							else
-							{
-								BuildResult = ECompilationResult.OtherCompilationError;
-							}
-
-							if( CPPIncludesThread != null )
-							{ 
-								// @todo ubtmake urgent: It is critical that this thread completes and we save dependency cache files, even if an error happened, otherwise our
-								// comparison with last built target will pass the next time through, even though our cache is bogus.  Maybe we could wipe it when we think its out of date?
-								
-								// Wait until our CPPIncludes dependency scanner thread has finished
-								CPPIncludesThread.Join();
-							}
-						}
-					}
-				}
-
-				// Save the include dependency cache.
-				{ 
-					// @todo fastubt urgent: If we update the cache and actually compile files, but DON'T SAVE the stream, we will be in big trouble on the next run -- the files could be out of date due to an indirect include, but we won't know how to check because we didn't cache the includes off!
-					//			-> Compile/Link errors DO SEEM to prevent saving!  ACK!
-					//			-> In any case, we need to:  
-					//					A) Wipe cache file on exception	
-					//					B) force SLOW intermediate checking on restart if not fully clean and cache file is missing (or force a clean first)
-					//					C) Save cache files on compile/link errors!
-
-					foreach( var DependencyCache in CPPEnvironment.IncludeDependencyCache.Values )
-					{ 
-						DependencyCache.Save();
-					}
-					CPPEnvironment.IncludeDependencyCache.Clear();
-
-					foreach( var FlatCPPIncludeDependencyCache in CPPEnvironment.FlatCPPIncludeDependencyCache.Values )
-					{ 
-						FlatCPPIncludeDependencyCache.Save();
-					}
-					CPPEnvironment.FlatCPPIncludeDependencyCache.Clear();
-				}
-			}
-			catch (BuildException Exception)
-			{
-				// Output the message only, without the call stack
-				Log.TraceInformation(Exception.Message);
-				BuildResult = ECompilationResult.OtherCompilationError;
-			}
-			catch (Exception Exception)
-			{
-				Log.TraceInformation("ERROR: {0}", Exception);
-				BuildResult = ECompilationResult.OtherCompilationError;
-			}
+                    if( UnrealBuildTool.IsGatheringBuild )
+                    {
+                        ActionGraph.FinalizeActionGraph();
+
+                        UBTMakefile = new UBTMakefile();
+                        UBTMakefile.AllActions = ActionGraph.AllActions;
+
+                        // For now simply treat all object files as the root target.
+                        { 
+                            var PrerequisiteActionsSet = new HashSet<Action>();
+                            foreach (var OutputItem in OutputItemsForAllTargets)
+                            {
+                                ActionGraph.GatherPrerequisiteActions(OutputItem, ref PrerequisiteActionsSet);
+                            }
+                            UBTMakefile.PrerequisiteActions = PrerequisiteActionsSet.ToArray();
+                        }			
+
+                        foreach( System.Collections.DictionaryEntry EnvironmentVariable in Environment.GetEnvironmentVariables() )
+                        {
+                            UBTMakefile.EnvironmentVariables.Add( Tuple.Create( (string)EnvironmentVariable.Key, (string)EnvironmentVariable.Value ) );
+                        }
+                        UBTMakefile.TargetNameToUObjectModules = TargetNameToUObjectModules;
+
+                        // We've been told to prepare to build, so let's go ahead and save out our action graph so that we can use in a later invocation 
+                        // to assemble the build.  Even if we are configured to assemble the build in this same invocation, we want to save out the
+                        // Makefile so that it can be used on subsequent 'assemble only' runs, for the fastest possible iteration times
+                        // @todo ubtmake: Optimization: We could make 'gather + assemble' mode slightly faster by saving this while busy compiling (on our worker thread)
+                        SaveUBTMakefile( Targets, UBTMakefile );
+                    }
+
+                    if( UnrealBuildTool.IsAssemblingBuild )
+                    {
+                        // If we didn't build the graph in this session, then we'll need to load a cached one
+                        if( !UnrealBuildTool.IsGatheringBuild )
+                        {
+                            ActionGraph.AllActions = UBTMakefile.AllActions;
+
+                            foreach( var EnvironmentVariable in UBTMakefile.EnvironmentVariables )
+                            {
+                                // @todo ubtmake: There may be some variables we do NOT want to clobber.
+                                Environment.SetEnvironmentVariable( EnvironmentVariable.Item1, EnvironmentVariable.Item2 );
+                            }
+
+                            // If any of the targets need UHT to be run, we'll go ahead and do that now
+                            foreach( var Target in Targets )
+                            {
+                                List<UHTModuleInfo> TargetUObjectModules;
+                                if( UBTMakefile.TargetNameToUObjectModules.TryGetValue( Target.GetTargetName(), out TargetUObjectModules ) )
+                                {
+                                    if( TargetUObjectModules.Count > 0 )
+                                    {
+                                        // Execute the header tool
+                                        string ModuleInfoFileName = Path.GetFullPath( Path.Combine( Target.ProjectIntermediateDirectory, "UnrealHeaderTool.manifest" ) );
+                                        ECompilationResult UHTResult = ECompilationResult.OtherCompilationError;
+                                        if (!ExternalExecution.ExecuteHeaderToolIfNecessary(Target, GlobalCompileEnvironment:null, UObjectModules:TargetUObjectModules, ModuleInfoFileName:ModuleInfoFileName, UHTResult:ref UHTResult))
+                                        {
+                                            Log.TraceInformation("UnrealHeaderTool failed for target '" + Target.GetTargetName() + "' (platform: " + Target.GetTargetInfo().Platform.ToString() + ", module info: " + ModuleInfoFileName + ").");
+                                            BuildResult = UHTResult;
+                                            break;
+                                        }
+                                    }
+                                }
+                            }
+                        }
+
+                        if( BuildResult == ECompilationResult.Succeeded )
+                        { 
+                            // Make sure any old DLL files from in-engine recompiles aren't lying around.  Must be called after the action graph is finalized.
+                            // @todo ubtmake: Will we even use traditional hot reload DLLs anymore in the future?  
+                            ActionGraph.DeleteStaleHotReloadDLLs();
+
+                            // Plan the actions to execute for the build.
+                            Dictionary<UEBuildTarget,List<FileItem>> TargetToOutdatedPrerequisitesMap;
+                            List<Action> ActionsToExecute = ActionGraph.GetActionsToExecute(UBTMakefile.PrerequisiteActions, Targets, out TargetToOutdatedPrerequisitesMap);
+
+                            // Display some stats to the user.
+                            Log.TraceVerbose(
+                                    "{0} actions, {1} outdated and requested actions",
+                                    ActionGraph.AllActions.Count,
+                                    ActionsToExecute.Count
+                                    );
+
+                            ToolChain.PreBuildSync();
+
+                    
+                            // Cache indirect includes for all outdated C++ files.  We kick this off as a background thread so that it can
+                            // perform the scan while we're compiling.  It usually only takes up to a few seconds, but we don't want to hurt
+                            // our best case UBT iteration times for this task which can easily be performed asynchronously
+                            Thread CPPIncludesThread = null;
+                            if( BuildConfiguration.bUseExperimentalFastDependencyScan && TargetToOutdatedPrerequisitesMap.Count > 0 )
+                            {
+                                CPPIncludesThread = CreateThreadForCachingCPPIncludes( TargetToOutdatedPrerequisitesMap );
+                                CPPIncludesThread.Start();
+                            }
+
+                            // Execute the actions.
+                            bSuccess = ActionGraph.ExecuteActions(ActionsToExecute, out ExecutorName);
+
+                            // if the build succeeded, do any needed syncing
+                            if (bSuccess)
+                            {
+                                foreach (UEBuildTarget Target in Targets)
+                                {
+                                    ToolChain.PostBuildSync(Target);
+                                }
+                            }
+                            else
+                            {
+                                BuildResult = ECompilationResult.OtherCompilationError;
+                            }
+
+                            if( CPPIncludesThread != null )
+                            { 
+                                // @todo ubtmake urgent: It is critical that this thread completes and we save dependency cache files, even if an error happened, otherwise our
+                                // comparison with last built target will pass the next time through, even though our cache is bogus.  Maybe we could wipe it when we think its out of date?
+                                
+                                // Wait until our CPPIncludes dependency scanner thread has finished
+                                CPPIncludesThread.Join();
+                            }
+                        }
+                    }
+                }
+
+                // Save the include dependency cache.
+                { 
+                    // @todo fastubt urgent: If we update the cache and actually compile files, but DON'T SAVE the stream, we will be in big trouble on the next run -- the files could be out of date due to an indirect include, but we won't know how to check because we didn't cache the includes off!
+                    //			-> Compile/Link errors DO SEEM to prevent saving!  ACK!
+                    //			-> In any case, we need to:  
+                    //					A) Wipe cache file on exception	
+                    //					B) force SLOW intermediate checking on restart if not fully clean and cache file is missing (or force a clean first)
+                    //					C) Save cache files on compile/link errors!
+
+                    foreach( var DependencyCache in CPPEnvironment.IncludeDependencyCache.Values )
+                    { 
+                        DependencyCache.Save();
+                    }
+                    CPPEnvironment.IncludeDependencyCache.Clear();
+
+                    foreach( var FlatCPPIncludeDependencyCache in CPPEnvironment.FlatCPPIncludeDependencyCache.Values )
+                    { 
+                        FlatCPPIncludeDependencyCache.Save();
+                    }
+                    CPPEnvironment.FlatCPPIncludeDependencyCache.Clear();
+                }
+            }
+            catch (BuildException Exception)
+            {
+                // Output the message only, without the call stack
+                Log.TraceInformation(Exception.Message);
+                BuildResult = ECompilationResult.OtherCompilationError;
+            }
+            catch (Exception Exception)
+            {
+                Log.TraceInformation("ERROR: {0}", Exception);
+                BuildResult = ECompilationResult.OtherCompilationError;
+            }
 
             if(EULAViolationWarning != null)
             {
@@ -1642,7 +1661,7 @@ namespace UnrealBuildTool
 
             // Figure out how long we took to execute.
             double BuildDuration = (DateTime.UtcNow - StartTime - MutexWaitTime).TotalSeconds;
-			if (ExecutorName == "Local" || ExecutorName == "Distcc")
+            if (ExecutorName == "Local" || ExecutorName == "Distcc")
             {
                 Log.TraceInformation("Cumulative action seconds ({0} processors): {1:0.00} building projects, {2:0.00} compiling, {3:0.00} creating app bundles, {4:0.00} generating debug info, {5:0.00} linking, {6:0.00} other",
                     System.Environment.ProcessorCount,
@@ -1844,278 +1863,278 @@ namespace UnrealBuildTool
         }
 
 
-		/// <summary>
-		/// Returns a Thread object that can be kicked off to update C++ include dependency cache
-		/// </summary>
-		/// <param name="TargetToOutdatedPrerequisitesMap">Maps each target to a list of outdated C++ files that need indirect dependencies cached</param>
-		/// <returns>The thread object</returns>
-		private static Thread CreateThreadForCachingCPPIncludes( Dictionary<UEBuildTarget, List<FileItem>> TargetToOutdatedPrerequisitesMap)
-		{
-			return new Thread( new ThreadStart( () =>
-				{
-					// @todo fastubt: This thread will access data structures that are also used on the main UBT thread, but during this time UBT
-					// is only invoking the build executor, so should not be touching this stuff.  However, we need to at some guards to make sure.
-
-
-					// When in 'assembler only' mode, we need to load the direct include graph from a cache now.  We skipped loading this earlier because
-					// we didn't need it unless there was actually an outdated dependency.  And we can load this on our worker thread while we're compiling.
-					foreach( var Target in TargetToOutdatedPrerequisitesMap.Keys )
-					{
-						if (!CPPEnvironment.IncludeDependencyCache.ContainsKey(Target))
-						{
-							CPPEnvironment.IncludeDependencyCache.Add( Target, DependencyCache.Create( DependencyCache.GetDependencyCachePathForTarget(Target) ) );
-						}
-					}
-
-
-					foreach( var TargetAndOutdatedPrerequisites in TargetToOutdatedPrerequisitesMap )
-					{
-						var Target = TargetAndOutdatedPrerequisites.Key;
-						var OutdatedPrerequisites = TargetAndOutdatedPrerequisites.Value;
-
-						foreach( var PrerequisiteItem in OutdatedPrerequisites )
-						{
-							// Invoke our deep include scanner to figure out whether any of the files included by this source file have
-							// changed since the build product was built
-							var FileBuildPlatform = UEBuildPlatform.GetBuildPlatform( Target.GetTargetInfo().Platform );
-							CPPEnvironment.FindAndCacheAllIncludedFiles( Target, PrerequisiteItem, FileBuildPlatform, PrerequisiteItem.CachedCPPIncludeInfo, bOnlyCachedDependencies:false );
-						}
-					}
-				} ) );
-		}
-
-
-
-		/// <summary>
-		/// A special Makefile that UBT is able to create in "-gather" mode, then load in "-assemble" mode to accelerate iterative compiling and linking
-		/// </summary>
-		[Serializable]
-		class UBTMakefile
-		{
-			/** Every action in the action graph */
-			public List<Action> AllActions;
-
-			/** List of the actions that need to be run in order to build the targets' final output items */
-			public Action[] PrerequisiteActions;			
-
-			/** Environment variables that we'll need in order to invoke the platform's compiler and linker */
-			public readonly List<Tuple<string,string>> EnvironmentVariables = new List<Tuple<string,string>>();
-
-			/** Maps each target to a list of UObject module info structures */
-			public Dictionary<string,List<UHTModuleInfo>> TargetNameToUObjectModules;
-
-
-			/** @return	 True if this makefile's contents look valid.  Called after loading the file to make sure it is legit. */
-			public bool IsValidMakefile()
-			{
-				return 
-					AllActions != null && AllActions.Count > 0 &&
-					PrerequisiteActions != null && PrerequisiteActions.Length > 0 &&
-					EnvironmentVariables != null &&
-					TargetNameToUObjectModules != null && TargetNameToUObjectModules.Count > 0;
-			}
-		}
-
-
-		/// <summary>
-		/// Saves a UBTMakefile to disk
-		/// </summary>
-		/// <param name="Targets">List of targets.  Order is not important</param>
-		static void SaveUBTMakefile( List<UEBuildTarget> Targets, UBTMakefile UBTMakefile )
-		{
-			if( !UBTMakefile.IsValidMakefile() )
-			{
-				throw new BuildException( "Can't save a UBTMakefile that has invalid contents.  See UBTMakefile.IsValidMakefile()" );
-			}
-
-			var TimerStartTime = DateTime.UtcNow;
-
-			var UBTMakefileItem = FileItem.GetItemByFullPath( GetUBTMakefilePath( Targets ) );
-
-			// @todo ubtmake: Optimization: The UBTMakefile saved for game projects is upwards of 9 MB.  We should try to shrink its content if possible
-			// @todo ubtmake: Optimization: C# Serialization may be too slow for these big Makefiles.  Loading these files often shows up as the slower part of the assembling phase.
-
-			// Serialize the cache to disk.
-			try
-			{
-				Directory.CreateDirectory(Path.GetDirectoryName(UBTMakefileItem.AbsolutePath));
-				using (FileStream Stream = new FileStream(UBTMakefileItem.AbsolutePath, FileMode.Create, FileAccess.Write))
-				{
-					BinaryFormatter Formatter = new BinaryFormatter();
-					Formatter.Serialize(Stream, UBTMakefile);
-				}
-			}
-			catch (Exception Ex)
-			{
-				Console.Error.WriteLine("Failed to write UBTMakefile: {0}", Ex.Message);
-			}
-
-			if (BuildConfiguration.bPrintPerformanceInfo)
-			{
-				var TimerDuration = DateTime.UtcNow - TimerStartTime;
-				Log.TraceInformation("Saving UBTMakefile took " + TimerDuration.TotalSeconds + "s");
-			}
-		}
-
-
-		/// <summary>
-		/// Loads a UBTMakefile from disk
-		/// </summary>
-		/// <param name="Targets">List of targets.  Order is not important</param>
-		static UBTMakefile LoadUBTMakefile( List<UEBuildTarget> Targets )
-		{
-			var UBTMakefileItem = FileItem.GetItemByFullPath( GetUBTMakefilePath( Targets ) );
-
-			// Check the directory timestamp on the project files directory.  If the user has generated project files more
-			// recently than the UBTMakefile, then we need to consider the file to be out of date
-			bool bForceOutOfDate = false;
-			{
-				// @todo ubtmake: This will only work if the directory timestamp actually changes with every single GPF.  Force delete existing files before creating new ones?  Eh... really we probably just want to delete + create a file in that folder
-				//			-> UPDATE: Seems to work OK right now though on Windows platform, maybe due to GUID changes
-				// @todo ubtmake: Some platforms may not save any files into this folder.  We should delete + generate a "touch" file to force the directory timestamp to be updated (or just check the timestamp file itself.  We could put it ANYWHERE, actually)
-
-				var ProjectFilesLastUpdateTime = DateTime.MaxValue;
-				if( !UnrealBuildTool.RunningRocket() )
-				{
-					if( Directory.Exists( ProjectFileGenerator.IntermediateProjectFilesPath ) )
-					{
-						ProjectFilesLastUpdateTime = new DirectoryInfo( ProjectFileGenerator.IntermediateProjectFilesPath ).LastWriteTime;
-					}
-					else
-					{
-						ProjectFilesLastUpdateTime = DateTime.MaxValue;
-					}
-				}
-				else
-				{
-					// Rocket doesn't need to check engine projects for outdatedness
-					ProjectFilesLastUpdateTime = DateTime.MinValue;
-				}
-
-				// Check the game project directory too
-				if( UnrealBuildTool.HasUProjectFile() )
-				{ 
-					var MasterProjectRelativePath = UnrealBuildTool.GetUProjectPath();
-					var GameIntermediateProjectFilesPath = Path.Combine( MasterProjectRelativePath, "Intermediate", "ProjectFiles" );
-					if( Directory.Exists( GameIntermediateProjectFilesPath ) )
-					{
-						var GameProjectFilesLastUpdateTime = new DirectoryInfo( GameIntermediateProjectFilesPath ).LastWriteTime;
-						if( GameProjectFilesLastUpdateTime > ProjectFilesLastUpdateTime )
-						{
-							ProjectFilesLastUpdateTime = GameProjectFilesLastUpdateTime;
-						}
-					}	
-					else
-					{
-						// Game project files folder is missing, so force outdatedness
-						ProjectFilesLastUpdateTime = DateTime.MaxValue;
-
-					}
-				}
-	
-				if( UBTMakefileItem.LastWriteTime.CompareTo( ProjectFilesLastUpdateTime ) < 0 )
-				{
-					// Consider the file to be out of date
-					bForceOutOfDate = true;
-				}
-			} 
-
-			UBTMakefile LoadedUBTMakefile = null;
-			if( !bForceOutOfDate )
-			{ 
-				try
-				{
-					using (FileStream Stream = new FileStream(UBTMakefileItem.AbsolutePath, FileMode.Open, FileAccess.Read))
-					{	
-						BinaryFormatter Formatter = new BinaryFormatter();
-						LoadedUBTMakefile = Formatter.Deserialize(Stream) as UBTMakefile;
-					}
-				}
-				catch (Exception Ex)
-				{
-					Log.TraceWarning("Failed to read UBTMakefile: {0}", Ex.Message);
-				}
-
-				if( LoadedUBTMakefile != null && !LoadedUBTMakefile.IsValidMakefile() )
-				{
-					Log.TraceWarning("Loaded UBTMakefile appears to have invalid contents, ignoring it ({0})", UBTMakefileItem.AbsolutePath );
-					LoadedUBTMakefile = null;
-				}
-			}
-
-			return LoadedUBTMakefile;
-		}
-
-
-		/// <summary>
-		/// Gets the file path for a UBTMakefile
-		/// </summary>
-		/// <param name="Targets">List of targets.  Order is not important</param>
-		/// <returns>UBTMakefile path</returns>
-		public static string GetUBTMakefilePath( List<UEBuildTarget> Targets )
-		{
-			string UBTMakefilePath;
-
-			if( Targets.Count == 1 )
-			{
-				// If there's only one target, just save the UBT makefile in the target's build intermediate directory
-				// under a folder for that target (and platform/config combo.)
-				string PlatformIntermediatePath = BuildConfiguration.PlatformIntermediatePath;
-				if (UnrealBuildTool.HasUProjectFile())
-				{
-					PlatformIntermediatePath = Path.Combine(UnrealBuildTool.GetUProjectPath(), BuildConfiguration.PlatformIntermediateFolder);
-				}
-
-				UBTMakefilePath = Path.Combine(PlatformIntermediatePath, Targets[0].GetTargetName(), Targets[0].GetTargetInfo().Configuration.ToString(), "Makefile.ubt" );
-			}
-			else
-			{
- 				// For Makefiles that contain multiple targets, we'll make up a file name that contains all of the targets, their
- 				// configurations and platforms, and save it into the base intermediate folder
-				var TargetCollectionName = MakeTargetCollectionName( Targets );
-
-				string ProjectIntermediatePath = BuildConfiguration.BaseIntermediatePath;
-				if (UnrealBuildTool.HasUProjectFile())
-				{
-					ProjectIntermediatePath = Path.Combine(UnrealBuildTool.GetUProjectPath(), BuildConfiguration.BaseIntermediatePath);
-				}
-
-				// @todo ubtmake: The TargetCollectionName string could be really long if there is more than one target!  Hash it?
-				UBTMakefilePath = Path.Combine(ProjectIntermediatePath, TargetCollectionName + ".ubt" );
-			}
-
-			return UBTMakefilePath;
-		}
-
-
-		/// <summary>
-		/// Makes up a name for a set of targets that we can use for file or directory names
-		/// </summary>
-		/// <param name="Targets">List of targets.  Order is not important</param>
-		/// <returns>The name to use</returns>
-		private static string MakeTargetCollectionName( List<UEBuildTarget> Targets )
-		{
-			if( Targets.Count == 0 )
-			{
-				throw new BuildException( "Expecting at least one Target to be passed to MakeTargetCollectionName" );
-			}
-
-			var SortedTargets = new List<UEBuildTarget>();
-			SortedTargets.AddRange( Targets );
-			SortedTargets.Sort( ( x, y ) => { return x.GetTargetName().CompareTo( y.GetTargetName() ); } );
-
-			// Figure out what to call our action graph based on everything we're building
-			var TargetCollectionName = new StringBuilder();
-			foreach( var Target in SortedTargets )
-			{
-				if( TargetCollectionName.Length > 0 )
-				{
-					TargetCollectionName.Append( "_" );
-				}
-				TargetCollectionName.Append( Target.GetTargetName() + "-" + Target.GetTargetInfo().Platform.ToString() + "-" + Target.GetTargetInfo().Configuration.ToString() );
-			}
-
-			return TargetCollectionName.ToString();
-		}
-	}
+        /// <summary>
+        /// Returns a Thread object that can be kicked off to update C++ include dependency cache
+        /// </summary>
+        /// <param name="TargetToOutdatedPrerequisitesMap">Maps each target to a list of outdated C++ files that need indirect dependencies cached</param>
+        /// <returns>The thread object</returns>
+        private static Thread CreateThreadForCachingCPPIncludes( Dictionary<UEBuildTarget, List<FileItem>> TargetToOutdatedPrerequisitesMap)
+        {
+            return new Thread( new ThreadStart( () =>
+                {
+                    // @todo fastubt: This thread will access data structures that are also used on the main UBT thread, but during this time UBT
+                    // is only invoking the build executor, so should not be touching this stuff.  However, we need to at some guards to make sure.
+
+
+                    // When in 'assembler only' mode, we need to load the direct include graph from a cache now.  We skipped loading this earlier because
+                    // we didn't need it unless there was actually an outdated dependency.  And we can load this on our worker thread while we're compiling.
+                    foreach( var Target in TargetToOutdatedPrerequisitesMap.Keys )
+                    {
+                        if (!CPPEnvironment.IncludeDependencyCache.ContainsKey(Target))
+                        {
+                            CPPEnvironment.IncludeDependencyCache.Add( Target, DependencyCache.Create( DependencyCache.GetDependencyCachePathForTarget(Target) ) );
+                        }
+                    }
+
+
+                    foreach( var TargetAndOutdatedPrerequisites in TargetToOutdatedPrerequisitesMap )
+                    {
+                        var Target = TargetAndOutdatedPrerequisites.Key;
+                        var OutdatedPrerequisites = TargetAndOutdatedPrerequisites.Value;
+
+                        foreach( var PrerequisiteItem in OutdatedPrerequisites )
+                        {
+                            // Invoke our deep include scanner to figure out whether any of the files included by this source file have
+                            // changed since the build product was built
+                            var FileBuildPlatform = UEBuildPlatform.GetBuildPlatform( Target.GetTargetInfo().Platform );
+                            CPPEnvironment.FindAndCacheAllIncludedFiles( Target, PrerequisiteItem, FileBuildPlatform, PrerequisiteItem.CachedCPPIncludeInfo, bOnlyCachedDependencies:false );
+                        }
+                    }
+                } ) );
+        }
+
+
+
+        /// <summary>
+        /// A special Makefile that UBT is able to create in "-gather" mode, then load in "-assemble" mode to accelerate iterative compiling and linking
+        /// </summary>
+        [Serializable]
+        class UBTMakefile
+        {
+            /** Every action in the action graph */
+            public List<Action> AllActions;
+
+            /** List of the actions that need to be run in order to build the targets' final output items */
+            public Action[] PrerequisiteActions;			
+
+            /** Environment variables that we'll need in order to invoke the platform's compiler and linker */
+            public readonly List<Tuple<string,string>> EnvironmentVariables = new List<Tuple<string,string>>();
+
+            /** Maps each target to a list of UObject module info structures */
+            public Dictionary<string,List<UHTModuleInfo>> TargetNameToUObjectModules;
+
+
+            /** @return	 True if this makefile's contents look valid.  Called after loading the file to make sure it is legit. */
+            public bool IsValidMakefile()
+            {
+                return 
+                    AllActions != null && AllActions.Count > 0 &&
+                    PrerequisiteActions != null && PrerequisiteActions.Length > 0 &&
+                    EnvironmentVariables != null &&
+                    TargetNameToUObjectModules != null && TargetNameToUObjectModules.Count > 0;
+            }
+        }
+
+
+        /// <summary>
+        /// Saves a UBTMakefile to disk
+        /// </summary>
+        /// <param name="Targets">List of targets.  Order is not important</param>
+        static void SaveUBTMakefile( List<UEBuildTarget> Targets, UBTMakefile UBTMakefile )
+        {
+            if( !UBTMakefile.IsValidMakefile() )
+            {
+                throw new BuildException( "Can't save a UBTMakefile that has invalid contents.  See UBTMakefile.IsValidMakefile()" );
+            }
+
+            var TimerStartTime = DateTime.UtcNow;
+
+            var UBTMakefileItem = FileItem.GetItemByFullPath( GetUBTMakefilePath( Targets ) );
+
+            // @todo ubtmake: Optimization: The UBTMakefile saved for game projects is upwards of 9 MB.  We should try to shrink its content if possible
+            // @todo ubtmake: Optimization: C# Serialization may be too slow for these big Makefiles.  Loading these files often shows up as the slower part of the assembling phase.
+
+            // Serialize the cache to disk.
+            try
+            {
+                Directory.CreateDirectory(Path.GetDirectoryName(UBTMakefileItem.AbsolutePath));
+                using (FileStream Stream = new FileStream(UBTMakefileItem.AbsolutePath, FileMode.Create, FileAccess.Write))
+                {
+                    BinaryFormatter Formatter = new BinaryFormatter();
+                    Formatter.Serialize(Stream, UBTMakefile);
+                }
+            }
+            catch (Exception Ex)
+            {
+                Console.Error.WriteLine("Failed to write UBTMakefile: {0}", Ex.Message);
+            }
+
+            if (BuildConfiguration.bPrintPerformanceInfo)
+            {
+                var TimerDuration = DateTime.UtcNow - TimerStartTime;
+                Log.TraceInformation("Saving UBTMakefile took " + TimerDuration.TotalSeconds + "s");
+            }
+        }
+
+
+        /// <summary>
+        /// Loads a UBTMakefile from disk
+        /// </summary>
+        /// <param name="Targets">List of targets.  Order is not important</param>
+        static UBTMakefile LoadUBTMakefile( List<UEBuildTarget> Targets )
+        {
+            var UBTMakefileItem = FileItem.GetItemByFullPath( GetUBTMakefilePath( Targets ) );
+
+            // Check the directory timestamp on the project files directory.  If the user has generated project files more
+            // recently than the UBTMakefile, then we need to consider the file to be out of date
+            bool bForceOutOfDate = false;
+            {
+                // @todo ubtmake: This will only work if the directory timestamp actually changes with every single GPF.  Force delete existing files before creating new ones?  Eh... really we probably just want to delete + create a file in that folder
+                //			-> UPDATE: Seems to work OK right now though on Windows platform, maybe due to GUID changes
+                // @todo ubtmake: Some platforms may not save any files into this folder.  We should delete + generate a "touch" file to force the directory timestamp to be updated (or just check the timestamp file itself.  We could put it ANYWHERE, actually)
+
+                var ProjectFilesLastUpdateTime = DateTime.MaxValue;
+                if( !UnrealBuildTool.RunningRocket() )
+                {
+                    if( Directory.Exists( ProjectFileGenerator.IntermediateProjectFilesPath ) )
+                    {
+                        ProjectFilesLastUpdateTime = new DirectoryInfo( ProjectFileGenerator.IntermediateProjectFilesPath ).LastWriteTime;
+                    }
+                    else
+                    {
+                        ProjectFilesLastUpdateTime = DateTime.MaxValue;
+                    }
+                }
+                else
+                {
+                    // Rocket doesn't need to check engine projects for outdatedness
+                    ProjectFilesLastUpdateTime = DateTime.MinValue;
+                }
+
+                // Check the game project directory too
+                if( UnrealBuildTool.HasUProjectFile() )
+                { 
+                    var MasterProjectRelativePath = UnrealBuildTool.GetUProjectPath();
+                    var GameIntermediateProjectFilesPath = Path.Combine( MasterProjectRelativePath, "Intermediate", "ProjectFiles" );
+                    if( Directory.Exists( GameIntermediateProjectFilesPath ) )
+                    {
+                        var GameProjectFilesLastUpdateTime = new DirectoryInfo( GameIntermediateProjectFilesPath ).LastWriteTime;
+                        if( GameProjectFilesLastUpdateTime > ProjectFilesLastUpdateTime )
+                        {
+                            ProjectFilesLastUpdateTime = GameProjectFilesLastUpdateTime;
+                        }
+                    }	
+                    else
+                    {
+                        // Game project files folder is missing, so force outdatedness
+                        ProjectFilesLastUpdateTime = DateTime.MaxValue;
+
+                    }
+                }
+    
+                if( UBTMakefileItem.LastWriteTime.CompareTo( ProjectFilesLastUpdateTime ) < 0 )
+                {
+                    // Consider the file to be out of date
+                    bForceOutOfDate = true;
+                }
+            } 
+
+            UBTMakefile LoadedUBTMakefile = null;
+            if( !bForceOutOfDate )
+            { 
+                try
+                {
+                    using (FileStream Stream = new FileStream(UBTMakefileItem.AbsolutePath, FileMode.Open, FileAccess.Read))
+                    {	
+                        BinaryFormatter Formatter = new BinaryFormatter();
+                        LoadedUBTMakefile = Formatter.Deserialize(Stream) as UBTMakefile;
+                    }
+                }
+                catch (Exception Ex)
+                {
+                    Log.TraceWarning("Failed to read UBTMakefile: {0}", Ex.Message);
+                }
+
+                if( LoadedUBTMakefile != null && !LoadedUBTMakefile.IsValidMakefile() )
+                {
+                    Log.TraceWarning("Loaded UBTMakefile appears to have invalid contents, ignoring it ({0})", UBTMakefileItem.AbsolutePath );
+                    LoadedUBTMakefile = null;
+                }
+            }
+
+            return LoadedUBTMakefile;
+        }
+
+
+        /// <summary>
+        /// Gets the file path for a UBTMakefile
+        /// </summary>
+        /// <param name="Targets">List of targets.  Order is not important</param>
+        /// <returns>UBTMakefile path</returns>
+        public static string GetUBTMakefilePath( List<UEBuildTarget> Targets )
+        {
+            string UBTMakefilePath;
+
+            if( Targets.Count == 1 )
+            {
+                // If there's only one target, just save the UBT makefile in the target's build intermediate directory
+                // under a folder for that target (and platform/config combo.)
+                string PlatformIntermediatePath = BuildConfiguration.PlatformIntermediatePath;
+                if (UnrealBuildTool.HasUProjectFile())
+                {
+                    PlatformIntermediatePath = Path.Combine(UnrealBuildTool.GetUProjectPath(), BuildConfiguration.PlatformIntermediateFolder);
+                }
+
+                UBTMakefilePath = Path.Combine(PlatformIntermediatePath, Targets[0].GetTargetName(), Targets[0].GetTargetInfo().Configuration.ToString(), "Makefile.ubt" );
+            }
+            else
+            {
+                // For Makefiles that contain multiple targets, we'll make up a file name that contains all of the targets, their
+                // configurations and platforms, and save it into the base intermediate folder
+                var TargetCollectionName = MakeTargetCollectionName( Targets );
+
+                string ProjectIntermediatePath = BuildConfiguration.BaseIntermediatePath;
+                if (UnrealBuildTool.HasUProjectFile())
+                {
+                    ProjectIntermediatePath = Path.Combine(UnrealBuildTool.GetUProjectPath(), BuildConfiguration.BaseIntermediatePath);
+                }
+
+                // @todo ubtmake: The TargetCollectionName string could be really long if there is more than one target!  Hash it?
+                UBTMakefilePath = Path.Combine(ProjectIntermediatePath, TargetCollectionName + ".ubt" );
+            }
+
+            return UBTMakefilePath;
+        }
+
+
+        /// <summary>
+        /// Makes up a name for a set of targets that we can use for file or directory names
+        /// </summary>
+        /// <param name="Targets">List of targets.  Order is not important</param>
+        /// <returns>The name to use</returns>
+        private static string MakeTargetCollectionName( List<UEBuildTarget> Targets )
+        {
+            if( Targets.Count == 0 )
+            {
+                throw new BuildException( "Expecting at least one Target to be passed to MakeTargetCollectionName" );
+            }
+
+            var SortedTargets = new List<UEBuildTarget>();
+            SortedTargets.AddRange( Targets );
+            SortedTargets.Sort( ( x, y ) => { return x.GetTargetName().CompareTo( y.GetTargetName() ); } );
+
+            // Figure out what to call our action graph based on everything we're building
+            var TargetCollectionName = new StringBuilder();
+            foreach( var Target in SortedTargets )
+            {
+                if( TargetCollectionName.Length > 0 )
+                {
+                    TargetCollectionName.Append( "_" );
+                }
+                TargetCollectionName.Append( Target.GetTargetName() + "-" + Target.GetTargetInfo().Platform.ToString() + "-" + Target.GetTargetInfo().Configuration.ToString() );
+            }
+
+            return TargetCollectionName.ToString();
+        }
+    }
 }
