Include shaders as part of engine source in standalone game projects.

diff --git a/Engine/Source/Programs/UnrealBuildTool/System/ProjectFileGenerator.cs b/Engine/Source/Programs/UnrealBuildTool/System/ProjectFileGenerator.cs
--- a/Engine/Source/Programs/UnrealBuildTool/System/ProjectFileGenerator.cs
+++ b/Engine/Source/Programs/UnrealBuildTool/System/ProjectFileGenerator.cs
@@ -283,7 +283,7 @@ namespace UnrealBuildTool
 
 				IncludeDocumentation = false;
 				IncludeBuildSystemFiles = false;
-				IncludeShaderSource = false;
+				IncludeShaderSource = true;
 				IncludeTemplateFiles = false;
 				IncludeConfigFiles = true;
 			}
