Add libgnustl_shared.so to APK during packaging (not on all Android devices)
#Android

diff --git a/Engine/Source/Programs/UnrealBuildTool/Android/UEDeployAndroid.cs b/Engine/Source/Programs/UnrealBuildTool/Android/UEDeployAndroid.cs
--- a/Engine/Source/Programs/UnrealBuildTool/Android/UEDeployAndroid.cs
+++ b/Engine/Source/Programs/UnrealBuildTool/Android/UEDeployAndroid.cs
@@ -485,7 +485,18 @@ namespace UnrealBuildTool.Android
 				File.Copy(SourceSOName, FinalSOName);
 
 			}
-			
+
+			// copy libgnustl_shared.so to library (use 4.8 if possible, otherwise 4.6)
+            string GccVersion = "4.8";
+            if (!Directory.Exists(Environment.ExpandEnvironmentVariables("%NDKROOT%/sources/cxx-stl/gnu-libstdc++/4.8")))
+            {
+                GccVersion = "4.6";
+            }
+            string SourceSTLSOName = Environment.ExpandEnvironmentVariables("%NDKROOT%/sources/cxx-stl/gnu-libstdc++/") + GccVersion + "/libs/armeabi-v7a/libgnustl_shared.so";
+			string FinalSTLSOName = UE4BuildPath + "/libs/armeabi-v7a/libgnustl_shared.so";
+            Directory.CreateDirectory(Path.GetDirectoryName(FinalSTLSOName));
+            File.Copy(SourceSTLSOName, FinalSTLSOName);
+
 			// remove any read only flags
 			FileInfo DestFileInfo = new FileInfo(FinalSOName);
 			DestFileInfo.Attributes = DestFileInfo.Attributes & ~FileAttributes.ReadOnly;
