Set the active project to the game when generating project files for the first time in GitHub.

diff --git a/Engine/Source/Programs/UnrealBuildTool/System/ProjectFileGenerator.cs b/Engine/Source/Programs/UnrealBuildTool/System/ProjectFileGenerator.cs
--- a/Engine/Source/Programs/UnrealBuildTool/System/ProjectFileGenerator.cs
+++ b/Engine/Source/Programs/UnrealBuildTool/System/ProjectFileGenerator.cs
@@ -384,6 +384,12 @@ namespace UnrealBuildTool
 				// Add all game projects and game config files
 				AddAllGameProjects(GameProjects, SupportedPlatformNames, RootFolder);
 
+				// Set the game to be the default project
+				if(bGeneratingGameProjectFiles && GameProjects.Count > 0)
+				{
+					DefaultProject = GameProjects.Values.First();
+				}
+
 				// Place projects into root level solution folders
 				if( ( !bGeneratingGameProjectFiles && !bGeneratingRocketProjectFiles ) || bAlwaysIncludeEngineModules )
 				{
@@ -2101,6 +2107,8 @@ namespace UnrealBuildTool
 			OtherProjectFiles.Add( InProject );
 		}
 
+		/// The default project to be built for the solution.
+		protected ProjectFile DefaultProject;
 
 		/// The project for UnrealBuildTool.  Note that when generating Rocket project files, we won't have
 		/// an UnrealBuildTool project at all.
diff --git a/Engine/Source/Programs/UnrealBuildTool/System/VCProjectFileGenerator.cs b/Engine/Source/Programs/UnrealBuildTool/System/VCProjectFileGenerator.cs
--- a/Engine/Source/Programs/UnrealBuildTool/System/VCProjectFileGenerator.cs
+++ b/Engine/Source/Programs/UnrealBuildTool/System/VCProjectFileGenerator.cs
@@ -771,6 +771,10 @@ namespace UnrealBuildTool
 					{
 						VCSolutionOptions Options = new VCSolutionOptions();
 						Options.SolutionConfiguration.Add(new VCBinarySetting("ActiveCfg", DefaultConfig.VCSolutionConfigAndPlatformName));
+						if(DefaultProject != null)
+						{
+							Options.SolutionConfiguration.Add(new VCBinarySetting("StartupProject", ((MSBuildProjectFile)DefaultProject).ProjectGUID.ToString("B")));
+						}
 						Options.Write(SolutionOptionsFileName);
 					}
 				}
