Create a VS2012 SUO file setting the default project and configuration, if we're generating VS 2012 project files. Seems to have a compatible format as the 2012 version, so just needs renaming.

diff --git a/Engine/Source/Programs/UnrealBuildTool/System/VCProjectFileGenerator.cs b/Engine/Source/Programs/UnrealBuildTool/System/VCProjectFileGenerator.cs
--- a/Engine/Source/Programs/UnrealBuildTool/System/VCProjectFileGenerator.cs
+++ b/Engine/Source/Programs/UnrealBuildTool/System/VCProjectFileGenerator.cs
@@ -761,12 +761,16 @@ namespace UnrealBuildTool
 
 
 			// Save a solution config file which selects the development editor configuration by default.
-			if (bSuccess && ProjectFileFormat == VCProjectFileFormat.VisualStudio2013)
+			if (bSuccess)
 			{
 				VCSolutionConfigCombination DefaultConfig = SolutionConfigCombinations.Find(x => x.Configuration == UnrealTargetConfiguration.Development && x.Platform == UnrealTargetPlatform.Win64 && x.TargetConfigurationName == "Editor");
 				if (DefaultConfig != null)
 				{
-					string SolutionOptionsFileName = Path.Combine(MasterProjectRelativePath, Path.ChangeExtension(SolutionFileName, "v12.suo"));
+					// Figure out the filename for the SUO file. VS2013 will automatically import the VS2012 options if necessary.
+					string SolutionOptionsExtension = (ProjectFileFormat == VCProjectFileFormat.VisualStudio2012)? "v11.suo" : "v12.suo";
+
+					// Check it doesn't exist before overwriting it. Since these files store the user's preferences, it'd be bad form to overwrite them.
+					string SolutionOptionsFileName = Path.Combine(MasterProjectRelativePath, Path.ChangeExtension(SolutionFileName, SolutionOptionsExtension));
 					if(!File.Exists(SolutionOptionsFileName))
 					{
 						VCSolutionOptions Options = new VCSolutionOptions();
