#UE4   Fix root cause of why blacklisted plugins were being loaded on HTML5,  revert CL 2255415.

diff --git a/Engine/Source/Programs/UnrealBuildTool/HTML5/UEBuildHTML5.cs b/Engine/Source/Programs/UnrealBuildTool/HTML5/UEBuildHTML5.cs
--- a/Engine/Source/Programs/UnrealBuildTool/HTML5/UEBuildHTML5.cs
+++ b/Engine/Source/Programs/UnrealBuildTool/HTML5/UEBuildHTML5.cs
@@ -406,7 +406,6 @@ namespace UnrealBuildTool
                 InBuildTarget.GlobalLinkEnvironment.Config.ExcludedLibraries.Add("LIBCMT");
 
             InBuildTarget.GlobalCompileEnvironment.Config.Definitions.Add("PLATFORM_HTML5=1");
-            InBuildTarget.GlobalCompileEnvironment.Config.Definitions.Add("HTML5=1");
             if (InBuildTarget.GlobalCompileEnvironment.Config.Target.Architecture == "-win32")
             {
                 InBuildTarget.GlobalCompileEnvironment.Config.Definitions.Add("PLATFORM_HTML5_WIN32=1");
diff --git a/Engine/Source/Runtime/Projects/Private/PluginDescriptor.cpp b/Engine/Source/Runtime/Projects/Private/PluginDescriptor.cpp
--- a/Engine/Source/Runtime/Projects/Private/PluginDescriptor.cpp
+++ b/Engine/Source/Runtime/Projects/Private/PluginDescriptor.cpp
@@ -88,12 +88,6 @@ bool FPluginDescriptor::Read(const FJsonObject& Object, FText& OutFailReason)
 	}
 
 	Object.TryGetBoolField(TEXT("EnabledByDefault"), bEnabledByDefault);
-
-	for ( FModuleDescriptor Descriptor : Modules )
-	{
-		bEnabledByDefault &= (!Descriptor.BlacklistPlatforms.Contains(FString(FPlatformProperties::PlatformName())));
-	}
-
 	Object.TryGetBoolField(TEXT("CanContainContent"), bCanContainContent);
 	Object.TryGetBoolField(TEXT("IsBetaVersion"), bIsBetaVersion);
 	return true;
