#HTML5

CIS fix. Need to fix up physX first before this can go in.

diff --git a/Engine/Source/Programs/UnrealBuildTool/HTML5/UEBuildHTML5.cs b/Engine/Source/Programs/UnrealBuildTool/HTML5/UEBuildHTML5.cs
--- a/Engine/Source/Programs/UnrealBuildTool/HTML5/UEBuildHTML5.cs
+++ b/Engine/Source/Programs/UnrealBuildTool/HTML5/UEBuildHTML5.cs
@@ -406,6 +406,7 @@ namespace UnrealBuildTool
                 InBuildTarget.GlobalLinkEnvironment.Config.ExcludedLibraries.Add("LIBCMT");
 
             InBuildTarget.GlobalCompileEnvironment.Config.Definitions.Add("PLATFORM_HTML5=1");
+            InBuildTarget.GlobalCompileEnvironment.Config.Definitions.Add("HTML5=1");
             if (InBuildTarget.GlobalCompileEnvironment.Config.Target.Architecture == "-win32")
             {
                 InBuildTarget.GlobalCompileEnvironment.Config.Definitions.Add("PLATFORM_HTML5_WIN32=1");
