Actually exposing UEBuildConfiguration.bAllowHotReloadFromIDE to BuildConfiguration.xml

diff --git a/Engine/Source/Programs/UnrealBuildTool/Configuration/UEBuildConfiguration.cs b/Engine/Source/Programs/UnrealBuildTool/Configuration/UEBuildConfiguration.cs
--- a/Engine/Source/Programs/UnrealBuildTool/Configuration/UEBuildConfiguration.cs
+++ b/Engine/Source/Programs/UnrealBuildTool/Configuration/UEBuildConfiguration.cs
@@ -173,6 +173,7 @@ namespace UnrealBuildTool
 		public static bool bCompileForSize;
 
 		/** True if hot-reload from IDE is allowed */
+		[XmlConfig]
 		public static bool bAllowHotReloadFromIDE;
 
 		/** True if performing hot-reload from IDE */
