UnrealBuildTool performance: Save dependency caches even if build fails

This is critical when using the new fast dependency scanning feature, because a buidl can fail *after* some targets compiled successfully, but we would not have cached C++ includes for those files, so subsequent runs would have incomplete data when determining outdatedness

This also makes iterating on builds faster when you have compile errors every time.  If your cache was missing or incomplete, it would reset with every single build attempt

diff --git a/Engine/Source/Programs/UnrealBuildTool/System/Unity.cs b/Engine/Source/Programs/UnrealBuildTool/System/Unity.cs
--- a/Engine/Source/Programs/UnrealBuildTool/System/Unity.cs
+++ b/Engine/Source/Programs/UnrealBuildTool/System/Unity.cs
@@ -229,6 +229,7 @@ namespace UnrealBuildTool
                 UnityCPPFiles.Add(UnityCPPFile);
 
 				// Cache information about the unity .cpp dependencies
+				// @todo fastubt urgent: Fails when building remotely for Mac because unity .cpp has an include for a PCH on the REMOTE machine
 				UEBuildModuleCPP.CachePCHUsageForModuleSourceFile( Target, CompileEnvironment, UnityCPPFile );
 			}
 
diff --git a/Engine/Source/Programs/UnrealBuildTool/System/UnrealBuildTool.cs b/Engine/Source/Programs/UnrealBuildTool/System/UnrealBuildTool.cs
--- a/Engine/Source/Programs/UnrealBuildTool/System/UnrealBuildTool.cs
+++ b/Engine/Source/Programs/UnrealBuildTool/System/UnrealBuildTool.cs
@@ -1357,6 +1357,7 @@ namespace UnrealBuildTool
                                 // we might not know about all of the C++ include dependencies for already-up-to-date shared build products
                                 // between the targets
                                 bNeedsFullCPPIncludeRescan = true;
+								Log.TraceVerbose( "Forcing full C++ include scan because we're building a different target" );
                             }
                         }
                     }
@@ -1618,37 +1619,12 @@ namespace UnrealBuildTool
 
                             if( CPPIncludesThread != null )
                             { 
-                                // @todo ubtmake urgent: It is critical that this thread completes and we save dependency cache files, even if an error happened, otherwise our
-                                // comparison with last built target will pass the next time through, even though our cache is bogus.  Maybe we could wipe it when we think its out of date?
-                                
                                 // Wait until our CPPIncludes dependency scanner thread has finished
                                 CPPIncludesThread.Join();
                             }
                         }
                     }
                 }
-
-                // Save the include dependency cache.
-                { 
-                    // @todo fastubt urgent: If we update the cache and actually compile files, but DON'T SAVE the stream, we will be in big trouble on the next run -- the files could be out of date due to an indirect include, but we won't know how to check because we didn't cache the includes off!
-                    //			-> Compile/Link errors DO SEEM to prevent saving!  ACK!
-                    //			-> In any case, we need to:  
-                    //					A) Wipe cache file on exception	
-                    //					B) force SLOW intermediate checking on restart if not fully clean and cache file is missing (or force a clean first)
-                    //					C) Save cache files on compile/link errors!
-
-                    foreach( var DependencyCache in CPPEnvironment.IncludeDependencyCache.Values )
-                    { 
-                        DependencyCache.Save();
-                    }
-                    CPPEnvironment.IncludeDependencyCache.Clear();
-
-                    foreach( var FlatCPPIncludeDependencyCache in CPPEnvironment.FlatCPPIncludeDependencyCache.Values )
-                    { 
-                        FlatCPPIncludeDependencyCache.Save();
-                    }
-                    CPPEnvironment.FlatCPPIncludeDependencyCache.Clear();
-                }
             }
             catch (BuildException Exception)
             {
@@ -1662,6 +1638,26 @@ namespace UnrealBuildTool
                 BuildResult = ECompilationResult.OtherCompilationError;
             }
 
+            // Save the include dependency cache.
+            { 
+				// NOTE: It's very important that we save the include cache, even if a build exception was thrown (compile error, etc), because we need to make sure that
+				//    any C++ include dependencies that we computed for out of date source files are saved.  Remember, the build may fail *after* some build products
+				//    are successfully built.  If we didn't save our dependency cache after build failures, source files for those build products that were successsfully
+				//    built before the failure would not be considered out of date on the next run, so this is our only chance to cache C++ includes for those files!
+
+                foreach( var DependencyCache in CPPEnvironment.IncludeDependencyCache.Values )
+                { 
+                    DependencyCache.Save();
+                }
+                CPPEnvironment.IncludeDependencyCache.Clear();
+
+                foreach( var FlatCPPIncludeDependencyCache in CPPEnvironment.FlatCPPIncludeDependencyCache.Values )
+                { 
+                    FlatCPPIncludeDependencyCache.Save();
+                }
+                CPPEnvironment.FlatCPPIncludeDependencyCache.Clear();
+            }
+
             if(EULAViolationWarning != null)
             {
                 Log.TraceWarning("WARNING: {0}", EULAViolationWarning);
@@ -1947,6 +1943,7 @@ namespace UnrealBuildTool
             public Action[] PrerequisiteActions;			
 
             /** Environment variables that we'll need in order to invoke the platform's compiler and linker */
+			// @todo ubtmake: Really we want to allow a different set of environment variables for every Action.  This would allow for targets on multiple platforms to be built in a single assembling phase.  Really, we'd only have unique variables for each platform that has actions, so we'd want to make sure we only store the minimum set.
             public readonly List<Tuple<string,string>> EnvironmentVariables = new List<Tuple<string,string>>();
 
             /** Maps each target to a list of UObject module info structures */
