UnrealBuildTool: Don't bother saving UBTMakefile when not using fast build iteration mode

diff --git a/Engine/Source/Programs/UnrealBuildTool/System/UnrealBuildTool.cs b/Engine/Source/Programs/UnrealBuildTool/System/UnrealBuildTool.cs
--- a/Engine/Source/Programs/UnrealBuildTool/System/UnrealBuildTool.cs
+++ b/Engine/Source/Programs/UnrealBuildTool/System/UnrealBuildTool.cs
@@ -1529,11 +1529,14 @@ namespace UnrealBuildTool
                         }
                         UBTMakefile.TargetNameToUObjectModules = TargetNameToUObjectModules;
 
-                        // We've been told to prepare to build, so let's go ahead and save out our action graph so that we can use in a later invocation 
-                        // to assemble the build.  Even if we are configured to assemble the build in this same invocation, we want to save out the
-                        // Makefile so that it can be used on subsequent 'assemble only' runs, for the fastest possible iteration times
-                        // @todo ubtmake: Optimization: We could make 'gather + assemble' mode slightly faster by saving this while busy compiling (on our worker thread)
-                        SaveUBTMakefile( Targets, UBTMakefile );
+						if( BuildConfiguration.bUseExperimentalFastBuildIteration )
+						{ 
+							// We've been told to prepare to build, so let's go ahead and save out our action graph so that we can use in a later invocation 
+							// to assemble the build.  Even if we are configured to assemble the build in this same invocation, we want to save out the
+							// Makefile so that it can be used on subsequent 'assemble only' runs, for the fastest possible iteration times
+							// @todo ubtmake: Optimization: We could make 'gather + assemble' mode slightly faster by saving this while busy compiling (on our worker thread)
+							SaveUBTMakefile( Targets, UBTMakefile );
+						}
                     }
 
                     if( UnrealBuildTool.IsAssemblingBuild )
