UnrealBuildTool: Catch invalid path characters in include statements

diff --git a/Engine/Source/Programs/UnrealBuildTool/System/CPPHeaders.cs b/Engine/Source/Programs/UnrealBuildTool/System/CPPHeaders.cs
--- a/Engine/Source/Programs/UnrealBuildTool/System/CPPHeaders.cs
+++ b/Engine/Source/Programs/UnrealBuildTool/System/CPPHeaders.cs
@@ -460,6 +460,11 @@ namespace UnrealBuildTool
 			{
 				string HeaderValue = C.Value;
 
+				if( HeaderValue.IndexOfAny( Path.GetInvalidPathChars() ) != -1 )
+				{
+					throw new BuildException( "In {0}: An #include statement contains invalid characters.  You might be missing a double-quote character. (\"{1}\")", FileToRead, C.Value );
+				}
+
 				//@TODO: The intermediate exclusion is to work around autogenerated absolute paths in Module.SomeGame.cpp style files
 				bool bIsIntermediateOrThirdParty = FileToRead.Contains("Intermediate") || FileToRead.Contains("ThirdParty");
 				bool bCheckForBackwardSlashes = FileToRead.StartsWith(InstalledFolder);
