New /Zo option supported for debugging optimized builds.

UpdateVersion read from registry.
New flag used based on version number read from registry.
Some refactoring.

diff --git a/Engine/Source/Programs/UnrealBuildTool/Windows/VCToolChain.cs b/Engine/Source/Programs/UnrealBuildTool/Windows/VCToolChain.cs
--- a/Engine/Source/Programs/UnrealBuildTool/Windows/VCToolChain.cs
+++ b/Engine/Source/Programs/UnrealBuildTool/Windows/VCToolChain.cs
@@ -194,7 +194,15 @@ namespace UnrealBuildTool
 
 					// Allow optimized code to be debugged more easily.  This makes PDBs a bit larger, but doesn't noticeably affect
 					// compile times.  The executable code is not affected at all by this switch, only the debugging information.
-					Arguments.Append(" /d2Zi+");
+					if (WindowsPlatform.Compiler == WindowsCompiler.VisualStudio2013 && EnvVars.UpdateVersion >= new Version("12.0.30723"))
+					{
+						// VC2013 Update 3 has a new flag for doing this
+						Arguments.Append(" /Zo");
+					}
+					else
+					{
+						Arguments.Append(" /d2Zi+");
+					}
 				}
 				
 				// Favor code speed.
@@ -627,6 +635,8 @@ namespace UnrealBuildTool
 
 		public override CPPOutput CompileCPPFiles(UEBuildTarget Target, CPPEnvironment CompileEnvironment, List<FileItem> SourceFiles, string ModuleName)
 		{
+			InitializeEnvironmentVariables(CompileEnvironment.Config.Target.Platform);
+
 			StringBuilder Arguments = new StringBuilder();
 			AppendCLArguments_Global(CompileEnvironment, Arguments);
 
@@ -1282,22 +1292,21 @@ namespace UnrealBuildTool
 		/** Gets the default include paths for the given platform. */
 		public static string GetVCIncludePaths(CPPTargetPlatform Platform)
 		{
-			string IncludePaths = "";
-			if (Platform == CPPTargetPlatform.Win32 || Platform == CPPTargetPlatform.Win64)
-			{
-				// Make sure we've got the environment variables set up for this target
-				VCToolChain.InitializeEnvironmentVariables(Platform);
+			Debug.Assert(Platform == CPPTargetPlatform.Win32 || Platform == CPPTargetPlatform.Win64);
 
-				// Also add any include paths from the INCLUDE environment variable.  MSVC is not necessarily running with an environment that
-				// matches what UBT extracted from the vcvars*.bat using SetEnvironmentVariablesFromBatchFile().  We'll use the variables we
-				// extracted to populate the project file's list of include paths
-				// @todo projectfiles: Should we only do this for VC++ platforms?
-				IncludePaths = Environment.GetEnvironmentVariable("INCLUDE");
-				if (!String.IsNullOrEmpty(IncludePaths) && !IncludePaths.EndsWith(";"))
-				{
-					IncludePaths += ";";
-				}
+			// Make sure we've got the environment variables set up for this target
+			VCToolChain.InitializeEnvironmentVariables(Platform);
+
+			// Also add any include paths from the INCLUDE environment variable.  MSVC is not necessarily running with an environment that
+			// matches what UBT extracted from the vcvars*.bat using SetEnvironmentVariablesFromBatchFile().  We'll use the variables we
+			// extracted to populate the project file's list of include paths
+			// @todo projectfiles: Should we only do this for VC++ platforms?
+			var IncludePaths = Environment.GetEnvironmentVariable("INCLUDE");
+			if (!String.IsNullOrEmpty(IncludePaths) && !IncludePaths.EndsWith(";"))
+			{
+				IncludePaths += ";";
 			}
+
 			return IncludePaths;
 		}
 
@@ -1316,19 +1325,15 @@ namespace UnrealBuildTool
 				// 64 bit -- we can use the 32 bit version to target 64 bit on 32 bit OS.
 				if (Platform == CPPTargetPlatform.Win64)
 				{
-					VCToolPath = Path.Combine(WindowsSDKDir, "bin/x64/rc.exe");
+					VCToolPath = Path.Combine(EnvVars.WindowsSDKDir, "bin/x64/rc.exe");
+				}
+				else if( !WindowsPlatform.SupportWindowsXP )	// Windows XP requires use to force Windows SDK 7.1 even on the newer compiler, so we need the old path RC.exe
+				{
+					VCToolPath = Path.Combine(EnvVars.WindowsSDKDir, "bin/x86/rc.exe");
 				}
-				// 32 bit
 				else
 				{
-					if( !WindowsPlatform.SupportWindowsXP )	// Windows XP requires use to force Windows SDK 7.1 even on the newer compiler, so we need the old path RC.exe
-					{
-						VCToolPath = Path.Combine( WindowsSDKDir, "bin/x86/rc.exe" );
-					}
-					else
-					{
-						VCToolPath = Path.Combine( WindowsSDKDir, "bin/rc.exe" );
-					}
+					VCToolPath = Path.Combine(EnvVars.WindowsSDKDir, "bin/rc.exe");
 				}
 			}
 			// cl.exe and link.exe are found in the toolchain specific folders (32 vs. 64 bit)
@@ -1348,8 +1353,6 @@ namespace UnrealBuildTool
 				}
 				else
 				{
-					string BaseVSToolPath = FindBaseVSToolPath();
-
 					// Both target and build machines are 64 bit
 					bool bIs64Bit = (Platform == CPPTargetPlatform.Win64);
 					// Regardless of the target, if we're linking on a 64 bit machine, we want to use the 64 bit linker (it's faster than the 32 bit linker)
@@ -1360,17 +1363,17 @@ namespace UnrealBuildTool
 					if (bIs64Bit || bUse64BitLinker)
 					{
 						// Use the native 64-bit compiler if present, otherwise use the amd64-on-x86 compiler. VS2012 Express only includes the latter.
-						string PlatformToolPath = Path.Combine(BaseVSToolPath, "../../VC/bin/amd64/");
+						string PlatformToolPath = Path.Combine(EnvVars.BaseVSToolPath, "../../VC/bin/amd64/");
 						if(!Directory.Exists(PlatformToolPath))
 						{
-							PlatformToolPath = Path.Combine(BaseVSToolPath, "../../VC/bin/x86_amd64/");
+							PlatformToolPath = Path.Combine(EnvVars.BaseVSToolPath, "../../VC/bin/x86_amd64/");
 						}
 						VCToolPath = PlatformToolPath + ToolName + ".exe";
 					}
 					else
 					{
 						// Use 32 bit for cl.exe and other tools, or for link.exe if 64 bit path doesn't exist and we're targeting 32 bit.
-						VCToolPath = Path.Combine(BaseVSToolPath, "../../VC/bin/" + ToolName + ".exe");
+						VCToolPath = Path.Combine(EnvVars.BaseVSToolPath, "../../VC/bin/" + ToolName + ".exe");
 					}
 				}
 			}
@@ -1395,119 +1398,116 @@ namespace UnrealBuildTool
 			return ToolPath;
 		}
 
-		/** Helper to only initialize environment variables once. */
-		static bool bAreEnvironmentVariablesAlreadyInitialized = false;
+		class EnvVarsType
+		{
+			public EnvVarsType(CPPTargetPlatform InPlatform, string InWindowsSDKDir, string InBaseVSToolPath, Version InUpdateVersion)
+			{
+				Platform       = InPlatform;
+				WindowsSDKDir  = InWindowsSDKDir;
+				BaseVSToolPath = InBaseVSToolPath;
+				UpdateVersion  = InUpdateVersion;
+			}
 
-		/** Helper to make sure environment variables have been initialized for the right platform. */
-		static CPPTargetPlatform PlatformEnvironmentVariablesAreInitializedFor = CPPTargetPlatform.Win32;
+			public readonly CPPTargetPlatform Platform;       // The platform the envvars have been initialized for
+			public readonly string            BaseVSToolPath; // The path to Visual Studio's /Common7/Tools directory.
+			public readonly string            WindowsSDKDir;  // Installation folder of the Windows SDK, e.g. C:\Program Files\Microsoft SDKs\Windows\v6.0A\
+			public readonly Version           UpdateVersion;  // The update version of Visual C++
+		}
 
-		/** Installation folder of the Windows SDK, e.g. C:\Program Files\Microsoft SDKs\Windows\v6.0A\ */
-		static string WindowsSDKDir = "";
+		static EnvVarsType EnvVars = null;
 
 		/**
 		 * Initializes environment variables required by toolchain. Different for 32 and 64 bit.
 		 */
-		static void InitializeEnvironmentVariables( CPPTargetPlatform Platform )
+		static void InitializeEnvironmentVariables(CPPTargetPlatform Platform)
 		{
-			if (!bAreEnvironmentVariablesAlreadyInitialized || Platform != PlatformEnvironmentVariablesAreInitializedFor)
+			if (EnvVars != null && EnvVars.Platform == Platform)
 			{
-				string BaseVSToolPath = FindBaseVSToolPath();
-				
-				string VCVarsBatchFile = "";
+				return;
+			}
 
-				// 64 bit tool chain.
-				if( Platform == CPPTargetPlatform.Win64 )
-				{
-					VCVarsBatchFile = Path.Combine(BaseVSToolPath, "../../VC/bin/x86_amd64/vcvarsx86_amd64.bat");
-				}
-				// The 32 bit vars batch file in the binary folder simply points to the one in the common tools folder.
-				else
-				{
-					VCVarsBatchFile = Path.Combine(BaseVSToolPath, "vsvars32.bat");
-				}
-				Utils.SetEnvironmentVariablesFromBatchFile(VCVarsBatchFile);
+			// If Visual Studio is not installed, the Windows SDK path will be used, which also happens to be the same
+			// directory. (It installs the toolchain into the folder where Visual Studio would have installed it to).
+			var BaseVSToolPath = WindowsPlatform.GetVSComnToolsPath();
+			if (string.IsNullOrEmpty(BaseVSToolPath))
+			{
+				throw new BuildException("Visual Studio 2012 or Visual Studio 2013 must be installed in order to build this target.");
+			}
+
+			// When targeting Windows XP on Visual Studio 2012+, we need to point at the older Windows SDK 7.1A that comes
+			// installed with Visual Studio 2012 Update 1. (http://blogs.msdn.com/b/vcblog/archive/2012/10/08/10357555.aspx)
+			var WindowsSDKDir = "";
+			if (WindowsPlatform.SupportWindowsXP)
+			{
+				WindowsSDKDir = FindWindowsSDKInstallationFolder("v7.1A");
+			}
+			else if (WindowsPlatform.Compiler == WindowsCompiler.VisualStudio2013)
+			{
+				WindowsSDKDir = FindWindowsSDKInstallationFolder("v8.1");
+			}
+			else if (WindowsPlatform.Compiler == WindowsCompiler.VisualStudio2012)
+			{
+				WindowsSDKDir = FindWindowsSDKInstallationFolder("v8.0");
+			}
+
+			var UpdateVersion = FindVCUpdateVersion();
+
+			var VCVarsBatchFile = Path.Combine(BaseVSToolPath, (Platform == CPPTargetPlatform.Win64) ? "../../VC/bin/x86_amd64/vcvarsx86_amd64.bat" : "vsvars32.bat");
+			Utils.SetEnvironmentVariablesFromBatchFile(VCVarsBatchFile);
 
+			// When targeting Windows XP on Visual Studio 2012+, we need to override the Windows SDK include and lib path set
+			// by the batch file environment (http://blogs.msdn.com/b/vcblog/archive/2012/10/08/10357555.aspx)
+			if (WindowsPlatform.SupportWindowsXP)
+			{
 				// Lib and bin folders have a x64 subfolder for 64 bit development.
-				string ConfigSuffix = "";
-				if( Platform == CPPTargetPlatform.Win64 )
-				{
-					ConfigSuffix = "\\x64";
-				}
+				var ConfigSuffix = (Platform == CPPTargetPlatform.Win64) ? "\\x64" : "";
 
-				// When targeting Windows XP on Visual Studio 2012+, we need to override the Windows SDK include and lib path set
-				// by the batch file environment (http://blogs.msdn.com/b/vcblog/archive/2012/10/08/10357555.aspx)
-				if( WindowsPlatform.SupportWindowsXP )
-				{
-					Environment.SetEnvironmentVariable("PATH", Utils.ResolveEnvironmentVariable(WindowsSDKDir + "bin" + ConfigSuffix + ";%PATH%"));
-					Environment.SetEnvironmentVariable("LIB", Utils.ResolveEnvironmentVariable(WindowsSDKDir + "lib" + ConfigSuffix + ";%LIB%"));
-					Environment.SetEnvironmentVariable( "INCLUDE", Utils.ResolveEnvironmentVariable(WindowsSDKDir + "include;%INCLUDE%"));
-				}
+				Environment.SetEnvironmentVariable("PATH",    Utils.ResolveEnvironmentVariable(WindowsSDKDir + "bin" + ConfigSuffix + ";%PATH%"));
+				Environment.SetEnvironmentVariable("LIB",     Utils.ResolveEnvironmentVariable(WindowsSDKDir + "lib" + ConfigSuffix + ";%LIB%"));
+				Environment.SetEnvironmentVariable("INCLUDE", Utils.ResolveEnvironmentVariable(WindowsSDKDir + "include;%INCLUDE%"));
+			}
 
-				bAreEnvironmentVariablesAlreadyInitialized = true;
-				PlatformEnvironmentVariablesAreInitializedFor = Platform;
-			}			
+			EnvVars = new EnvVarsType(
+				InPlatform      : Platform,
+				InBaseVSToolPath: BaseVSToolPath,
+				InWindowsSDKDir : WindowsSDKDir,
+				InUpdateVersion : UpdateVersion
+			);
 		}
 
 		/// <returns>The path to Windows SDK directory for the specified version.</returns>
-		private static string FindWindowsSDKInstallationFolder( string Version )
+		private static string FindWindowsSDKInstallationFolder(string Version)
 		{
 			// Based on VCVarsQueryRegistry
-			string WinSDKPath = (string)Microsoft.Win32.Registry.GetValue( @"HKEY_CURRENT_USER\SOFTWARE\Microsoft\Microsoft SDKs\Windows\" + Version, "InstallationFolder", null );
-			if( WinSDKPath != null )
-			{
-				return WinSDKPath;
-			}
-			WinSDKPath = (string)Microsoft.Win32.Registry.GetValue( @"HKEY_LOCAL_MACHINE\SOFTWARE\Wow6432Node\Microsoft\Microsoft SDKs\Windows\" + Version, "InstallationFolder", null );
-			if( WinSDKPath != null )
-			{
-				return WinSDKPath;
-			}
-			WinSDKPath = (string)Microsoft.Win32.Registry.GetValue( @"HKEY_CURRENT_USER\SOFTWARE\Wow6432Node\Microsoft\Microsoft SDKs\Windows\" + Version, "InstallationFolder", null );
-			if( WinSDKPath != null )
+			var WinSDKPath =
+				   Microsoft.Win32.Registry.GetValue(@"HKEY_CURRENT_USER\SOFTWARE\Microsoft\Microsoft SDKs\Windows\"              + Version, "InstallationFolder", null)
+				?? Microsoft.Win32.Registry.GetValue(@"HKEY_LOCAL_MACHINE\SOFTWARE\Wow6432Node\Microsoft\Microsoft SDKs\Windows\" + Version, "InstallationFolder", null)
+				?? Microsoft.Win32.Registry.GetValue(@"HKEY_CURRENT_USER\SOFTWARE\Wow6432Node\Microsoft\Microsoft SDKs\Windows\"  + Version, "InstallationFolder", null);
+
+			if (WinSDKPath == null)
 			{
-				return WinSDKPath;
+				throw new BuildException("Windows SDK {0} must be installed in order to build this target.", Version);
 			}
 
-			throw new BuildException( "Windows SDK {0} must be installed in order to build this target.", Version );
+			return (string)WinSDKPath;
 		}
 
-		/// <summary>
-		/// Figures out the path to Visual Studio's /Common7/Tools directory.  Note that if Visual Studio is not
-		/// installed, the Windows SDK path will be used, which also happens to be the same directory. (It installs
-		/// the toolchain into the folder where Visual Studio would have installed it to.
-		/// </summary>
-		/// <returns>The path to Visual Studio's /Common7/Tools directory</returns>
-		private static string FindBaseVSToolPath()
+		/// <returns>The installed version of VC.</returns>
+		private static Version FindVCUpdateVersion()
 		{
-			string BaseVSToolPath = "";
-			
-			// When targeting Windows XP on Visual Studio 2012+, we need to point at the older Windows SDK 7.1A that comes
-			// installed with Visual Studio 2012 Update 1. (http://blogs.msdn.com/b/vcblog/archive/2012/10/08/10357555.aspx)
-			if( WindowsPlatform.SupportWindowsXP )
-			{
-				WindowsSDKDir = FindWindowsSDKInstallationFolder( "v7.1A" );
-			}
-			else
-			{
-				if (WindowsPlatform.Compiler == WindowsCompiler.VisualStudio2013)
-				{
-					WindowsSDKDir = FindWindowsSDKInstallationFolder( "v8.1" );
-				}
-				else if (WindowsPlatform.Compiler == WindowsCompiler.VisualStudio2012)
-				{
-					WindowsSDKDir = FindWindowsSDKInstallationFolder( "v8.0" );
-				}
-			}
-
-			// Grab path to Visual Studio binaries from the system environment
-			BaseVSToolPath = WindowsPlatform.GetVSComnToolsPath();
+			// Based on VCVarsQueryRegistry
+			var UpdateVersionString =
+				   Microsoft.Win32.Registry.GetValue(@"HKEY_LOCAL_MACHINE\SOFTWARE\Microsoft\DevDiv\vc\Servicing\12.0\CompilerCore",             "UpdateVersion", null)
+				?? Microsoft.Win32.Registry.GetValue(@"HKEY_LOCAL_MACHINE\SOFTWARE\Wow6432Node\Microsoft\DevDiv\vc\Servicing\12.0\CompilerCore", "UpdateVersion", null)
+				?? Microsoft.Win32.Registry.GetValue(@"HKEY_CURRENT_USER\SOFTWARE\Wow6432Node\Microsoft\DevDiv\vc\Servicing\12.0\CompilerCore",  "UpdateVersion", null);
 
-			if (string.IsNullOrEmpty(BaseVSToolPath))
+			Version Result;
+			if (UpdateVersionString == null || !Version.TryParse((string)UpdateVersionString, out Result))
 			{
-				throw new BuildException("Visual Studio 2012 or Visual Studio 2013 must be installed in order to build this target.");
+				throw new BuildException("Unable to determine Visual Studio version from registry");
 			}
 
-			return BaseVSToolPath;
+			return Result;
 		}
 
         public override void AddFilesToManifest(ref FileManifest Manifest, UEBuildBinary Binary)
