UnrealBuildTool: Clang 64-bit support on Windows; LLD support
- Enables support for 64-bit compiling under Clang on Windows platform
- Adds initial support for LLVM linker on Windows (WindowsPlatform.bAllowClangLinker) (NOT working yet.)
- Miscellaneous Clang compiling fixes

diff --git a/Engine/Source/Programs/UnrealBuildTool/Configuration/BuildConfiguration.cs b/Engine/Source/Programs/UnrealBuildTool/Configuration/BuildConfiguration.cs
--- a/Engine/Source/Programs/UnrealBuildTool/Configuration/BuildConfiguration.cs
+++ b/Engine/Source/Programs/UnrealBuildTool/Configuration/BuildConfiguration.cs
@@ -295,8 +295,7 @@ namespace UnrealBuildTool
 		/// <summary>
 		/// Enables code analysis mode.  Currently, this has specific requirements.  It only works on Windows
 		/// platform with the MSVC compiler.  Also, it requires a version of the compiler that supports the
-		/// /analyze option.  The compiler included with the Windows SDK supports this, as well as Team Edition
-		/// versions of Visual Studio.
+		/// /analyze option, such as Visual Studio 2013
 		/// </summary>
 		[XmlConfig]
 		public static bool bEnableCodeAnalysis;
diff --git a/Engine/Source/Programs/UnrealBuildTool/Windows/UEBuildWindows.cs b/Engine/Source/Programs/UnrealBuildTool/Windows/UEBuildWindows.cs
--- a/Engine/Source/Programs/UnrealBuildTool/Windows/UEBuildWindows.cs
+++ b/Engine/Source/Programs/UnrealBuildTool/Windows/UEBuildWindows.cs
@@ -61,7 +61,9 @@ namespace UnrealBuildTool
 
         /// True if we are using "clang-cl" to compile instead of MSVC on Windows platform
         public static readonly bool bCompileWithClang = false;
-            // !UnrealBuildTool.CommandLineContains( "UnrealHeaderTool" );	// @todo clang: Avoid using Clang to compile UHT because 64-bit toolchain is not ready yet
+
+		/// True if we should use the Clang linker (LLD) when bCompileWithClang is enabled, otherwise we use the MSVC linker
+		public static readonly bool bAllowClangLinker = false;
 
         /// True if we're targeting Windows XP as a minimum spec.  In Visual Studio 2012 and higher, this may change how
         /// we compile and link the application (http://blogs.msdn.com/b/vcblog/archive/2012/10/08/10357555.aspx)
diff --git a/Engine/Source/Programs/UnrealBuildTool/Windows/VCToolChain.cs b/Engine/Source/Programs/UnrealBuildTool/Windows/VCToolChain.cs
--- a/Engine/Source/Programs/UnrealBuildTool/Windows/VCToolChain.cs
+++ b/Engine/Source/Programs/UnrealBuildTool/Windows/VCToolChain.cs
@@ -28,7 +28,7 @@ namespace UnrealBuildTool
 
 			// @todo fastubt: We can use '/showIncludes' to accelerate outdatedness checking, as the compiler will discover all indirect includes itself.  But, the spew is pretty noisy!
 			//		-> If no files in source file chain have changed, even if the build product is outdated, we can skip using /showIncludes (relies on cache surviving)
-			// Result += " /showIncludes";
+			// Arguments.Append( " /showIncludes" );
 
 			if( WindowsPlatform.bCompileWithClang )
 			{ 
@@ -78,22 +78,35 @@ namespace UnrealBuildTool
 			// Enable for static code analysis (where supported). Not treating analysis warnings as errors.
 			// Arguments.Append(" /analyze:WX-");
 			
-			if( !WindowsPlatform.bCompileWithClang )	// @todo clang: These options are not supported with clang-cl yet
+			if( WindowsPlatform.bCompileWithClang )
 			{
-				if (CompileEnvironment.Config.Target.Platform == CPPTargetPlatform.Win64)
+				// Tell the Clang compiler whether we want to generate 32-bit code or 64-bit code
+				if( CompileEnvironment.Config.Target.Platform == CPPTargetPlatform.Win64 )
 				{
-					// Pack struct members on 8-byte boundaries.
-					Arguments.Append(" /Zp8");
+					Arguments.Append( " --target=x86_64-pc-win32" );
 				}
 				else
 				{
-					// Pack struct members on 4-byte boundaries.
-					Arguments.Append(" /Zp4");
+					Arguments.Append( " --target=x86-pc-win32" );
 				}
+			}
 
-				// Separate functions for linker.
-				Arguments.Append(" /Gy");
+			if (CompileEnvironment.Config.Target.Platform == CPPTargetPlatform.Win64)
+			{
+				// Pack struct members on 8-byte boundaries.
+				Arguments.Append(" /Zp8");
+			}
+			else
+			{
+				// Pack struct members on 4-byte boundaries.
+				Arguments.Append(" /Zp4");
+			}
+
+			// Separate functions for linker.
+			Arguments.Append(" /Gy");
 
+			if( !WindowsPlatform.bCompileWithClang )	// @todo clang: These options are not supported with clang-cl yet
+			{
 				// Relaxes floating point precision semantics to allow more optimization.
 				Arguments.Append(" /fp:fast");
 			}
@@ -192,16 +205,19 @@ namespace UnrealBuildTool
 				{
 					Arguments.Append(" /Ox");
 
-					// Allow optimized code to be debugged more easily.  This makes PDBs a bit larger, but doesn't noticeably affect
-					// compile times.  The executable code is not affected at all by this switch, only the debugging information.
-					if (WindowsPlatform.Compiler == WindowsCompiler.VisualStudio2013 && EnvVars.UpdateVersion >= new Version("12.0.30723"))
-					{
-						// VC2013 Update 3 has a new flag for doing this
-						Arguments.Append(" /Zo");
-					}
-					else
+					if( !WindowsPlatform.bCompileWithClang )
 					{
-						Arguments.Append(" /d2Zi+");
+						// Allow optimized code to be debugged more easily.  This makes PDBs a bit larger, but doesn't noticeably affect
+						// compile times.  The executable code is not affected at all by this switch, only the debugging information.
+					  	if (WindowsPlatform.Compiler == WindowsCompiler.VisualStudio2013 && EnvVars.UpdateVersion >= new Version("12.0.30723"))
+					  	{
+							// VC2013 Update 3 has a new flag for doing this
+							Arguments.Append(" /Zo");
+					 	}	
+						else
+						{
+							Arguments.Append(" /d2Zi+");
+						}
 					}
 				}
 				
@@ -243,12 +259,8 @@ namespace UnrealBuildTool
 				if (CompileEnvironment.Config.CLRMode == CPPCLRMode.CLRDisabled &&
 					CompileEnvironment.Config.Target.Platform != CPPTargetPlatform.Win64)
 				{
-					// @todo Clang: Doesn't make sense to the Clang compiler
-					if( !WindowsPlatform.bCompileWithClang )
-					{
-						// Allow the compiler to generate SSE2 instructions.
-						Arguments.Append(" /arch:SSE2");
-					}
+					// Allow the compiler to generate SSE2 instructions.
+					Arguments.Append(" /arch:SSE2");
 				}
 
 				// Prompt the user before reporting internal errors to Microsoft.
@@ -369,9 +381,14 @@ namespace UnrealBuildTool
 				// Disable specific warnings that cause problems with Clang
 				// NOTE: These must appear after we set the MSVC warning level
 
+				// @todo clang: Ideally we want as few warnings disabled as possible
+				// 
 				// Allow Microsoft-specific syntax to slide, even though it may be non-standard.  Needed for Windows headers.
 				Arguments.Append(" -Wno-microsoft");								
 
+				// Don't complain about use of 'const' on reference types when it will have no effect
+				// Arguments.Append( " -Wno-ignored-qualifiers" );	// @todo clang: Needed?
+
 				// @todo clang: Kind of a shame to turn these off.  We'd like to catch unused variables, but it is tricky with how our assertion macros work.
 				Arguments.Append(" -Wno-unused-variable");
 				Arguments.Append(" -Wno-unused-function");
@@ -380,6 +397,9 @@ namespace UnrealBuildTool
 
 				Arguments.Append(" -Wno-inline-new-delete");	// @todo clang: We declare operator new as inline.  Clang doesn't seem to like that.
 
+				// Sometimes we compare 'this' pointers against nullptr, which Clang warns about by default
+				Arguments.Append(" -Wno-undefined-bool-conversion" );
+
 				// @todo clang: Disabled warnings were copied from MacToolChain for the most part
 				Arguments.Append(" -Wno-deprecated-declarations");
 				Arguments.Append(" -Wno-deprecated-writable-strings");
@@ -408,6 +428,12 @@ namespace UnrealBuildTool
 
 		static void AppendLinkArguments(LinkEnvironment LinkEnvironment, StringBuilder Arguments)
 		{
+			if( WindowsPlatform.bCompileWithClang && WindowsPlatform.bAllowClangLinker )
+			{
+				// This tells LLD to run in "Windows emulation" mode, meaning that it will accept MSVC Link arguments
+				Arguments.Append( " -flavor link" );
+			}
+
 			// Don't create a side-by-side manifest file for the executable.
 			Arguments.Append(" /MANIFEST:NO");
 
@@ -718,6 +744,12 @@ namespace UnrealBuildTool
 					}
 				}
 
+				if( WindowsPlatform.bCompileWithClang )
+				{ 
+					CompileAction.OutputEventHandler = new DataReceivedEventHandler( ClangCompilerOutputFormatter );
+				}
+
+
 				bool bEmitsObjectFile = true;
 				if (CompileEnvironment.Config.PrecompiledHeaderAction == PrecompiledHeaderAction.Create)
 				{
@@ -957,6 +989,11 @@ namespace UnrealBuildTool
 				CompileAction.CommandPath = GetVCToolPath(Environment.Config.Target.Platform, Environment.Config.Target.Configuration, "rc");
 				CompileAction.StatusDescription = Path.GetFileName(RCFile.AbsolutePath);
 
+				if( WindowsPlatform.bCompileWithClang )
+				{ 
+					CompileAction.OutputEventHandler = new DataReceivedEventHandler( ClangCompilerOutputFormatter );
+				}
+
 				// Suppress header spew
 				CompileAction.CommandArguments += " /nologo";
 
@@ -1220,6 +1257,11 @@ namespace UnrealBuildTool
 			LinkAction.PrerequisiteItems.AddRange(PrerequisiteItems);
 			LinkAction.StatusDescription = Path.GetFileName(OutputFile.AbsolutePath);
 
+			if( WindowsPlatform.bCompileWithClang )
+			{ 
+				LinkAction.OutputEventHandler = new DataReceivedEventHandler( ClangCompilerOutputFormatter );
+			}
+
 			// Tell the action that we're building an import library here and it should conditionally be
 			// ignored as a prerequisite for other actions
 			LinkAction.bProducesImportLibrary = bBuildImportLibraryOnly || LinkEnvironment.Config.bIsBuildingDLL;
@@ -1343,9 +1385,17 @@ namespace UnrealBuildTool
 				bool bIsRequestingLibTool = ToolName.Equals( "lib", StringComparison.InvariantCultureIgnoreCase );
 
 				// If we were asked to use Clang, then we'll redirect the path to the compiler to the LLVM installation directory
-				if( WindowsPlatform.bCompileWithClang && !bIsRequestingLinkTool && !bIsRequestingLibTool )
+				if( WindowsPlatform.bCompileWithClang && !bIsRequestingLibTool && !( bIsRequestingLinkTool && !WindowsPlatform.bAllowClangLinker ) )
 				{
-					VCToolPath = Path.Combine( Environment.GetFolderPath( Environment.SpecialFolder.ProgramFilesX86 ), "LLVM", "msbuild-bin", ToolName + ".exe" );
+					if( bIsRequestingLinkTool )
+					{
+						ToolName = Path.Combine( "bin", "lld" );
+					}
+					else
+					{
+						ToolName = Path.Combine( "msbuild-bin", "cl" );
+					}
+					VCToolPath = Path.Combine( Environment.GetFolderPath( Environment.SpecialFolder.ProgramFilesX86 ), "LLVM", ToolName + ".exe" );
 					if( !File.Exists( VCToolPath ) )
 					{
 						throw new BuildException( "Clang was selected as the Windows compiler, but LLVM/Clang does not appear to be installed.  Could not find: " + VCToolPath );
@@ -1522,5 +1572,19 @@ namespace UnrealBuildTool
                 Manifest.AddBinaryNames(Path.Combine(Binary.Config.IntermediateDirectory, Path.GetFileNameWithoutExtension(Binary.Config.OutputFilePath) + ".lib"), "");
             }
         }
+
+
+		/** Formats compiler output from Clang so that it is clickable in Visual Studio */
+		protected static void ClangCompilerOutputFormatter(object sender, DataReceivedEventArgs e)
+		{
+			var Output = e.Data;
+			if (Output == null)
+			{
+				return;
+			}
+
+			// @todo clang: Convert relative includes to absolute files so they'll be clickable
+			Log.TraceInformation(Output);
+		}
 	};
 }
