UnrealBuildTool: Fixed wrong NMake parameters when using platform-specific project files
- Thanks to [REDACTED_PERSON] for contributing this fix!

diff --git a/Engine/Source/Programs/UnrealBuildTool/System/VCProject.cs b/Engine/Source/Programs/UnrealBuildTool/System/VCProject.cs
--- a/Engine/Source/Programs/UnrealBuildTool/System/VCProject.cs
+++ b/Engine/Source/Programs/UnrealBuildTool/System/VCProject.cs
@@ -897,7 +897,7 @@ namespace UnrealBuildTool
 			string UProjectPath = "";
 			if (IsForeignProject)
 			{
-				UProjectPath = "\"$(SolutionDir)$(SolutionName).uproject\"";
+				UProjectPath = "\"$(SolutionDir)$(ProjectName).uproject\"";
 			}
 
 			string ProjectConfigName = StubProjectConfigurationName;
