UnrealBuildTool: bUseExperimentalFastBuildIteration: Fixed timestamp exception
- Fixed a timestamp exception when no engine project files existed
- More UBTMakefile diagnostics are now printed to the console window when using fast build iteration

diff --git a/Engine/Source/Programs/UnrealBuildTool/System/UnrealBuildTool.cs b/Engine/Source/Programs/UnrealBuildTool/System/UnrealBuildTool.cs
--- a/Engine/Source/Programs/UnrealBuildTool/System/UnrealBuildTool.cs
+++ b/Engine/Source/Programs/UnrealBuildTool/System/UnrealBuildTool.cs
@@ -1357,7 +1357,7 @@ namespace UnrealBuildTool
                                 // we might not know about all of the C++ include dependencies for already-up-to-date shared build products
                                 // between the targets
                                 bNeedsFullCPPIncludeRescan = true;
-								Log.TraceVerbose( "Forcing full C++ include scan because we're building a different target" );
+								Log.TraceInformation( "Forcing full C++ include scan because we're building a different target" );
                             }
                         }
                     }
@@ -2017,27 +2017,28 @@ namespace UnrealBuildTool
             // Check the directory timestamp on the project files directory.  If the user has generated project files more
             // recently than the UBTMakefile, then we need to consider the file to be out of date
             bool bForceOutOfDate = false;
+            if( UBTMakefileItem.bExists )
             {
                 // @todo ubtmake: This will only work if the directory timestamp actually changes with every single GPF.  Force delete existing files before creating new ones?  Eh... really we probably just want to delete + create a file in that folder
                 //			-> UPDATE: Seems to work OK right now though on Windows platform, maybe due to GUID changes
                 // @todo ubtmake: Some platforms may not save any files into this folder.  We should delete + generate a "touch" file to force the directory timestamp to be updated (or just check the timestamp file itself.  We could put it ANYWHERE, actually)
 
-                var ProjectFilesLastUpdateTime = DateTime.MaxValue;
-                if( !UnrealBuildTool.RunningRocket() )
+				if( !UnrealBuildTool.RunningRocket() )
                 {
                     if( Directory.Exists( ProjectFileGenerator.IntermediateProjectFilesPath ) )
                     {
-                        ProjectFilesLastUpdateTime = new DirectoryInfo( ProjectFileGenerator.IntermediateProjectFilesPath ).LastWriteTime;
-                    }
-                    else
-                    {
-                        ProjectFilesLastUpdateTime = DateTime.MaxValue;
+                        var EngineProjectFilesLastUpdateTime = new DirectoryInfo( ProjectFileGenerator.IntermediateProjectFilesPath ).LastWriteTime;
+						if( UBTMakefileItem.LastWriteTime.CompareTo( EngineProjectFilesLastUpdateTime ) < 0 )
+						{
+							// Engine project files are newer than UBT make file
+							bForceOutOfDate = true;
+							Log.TraceInformation("Existing UBTMakefile is older than generated engine project files, ignoring it" );
+						}
                     }
                 }
                 else
                 {
                     // Rocket doesn't need to check engine projects for outdatedness
-                    ProjectFilesLastUpdateTime = DateTime.MinValue;
                 }
 
                 // Check the game project directory too
@@ -2048,25 +2049,27 @@ namespace UnrealBuildTool
                     if( Directory.Exists( GameIntermediateProjectFilesPath ) )
                     {
                         var GameProjectFilesLastUpdateTime = new DirectoryInfo( GameIntermediateProjectFilesPath ).LastWriteTime;
-                        if( GameProjectFilesLastUpdateTime > ProjectFilesLastUpdateTime )
-                        {
-                            ProjectFilesLastUpdateTime = GameProjectFilesLastUpdateTime;
-                        }
+						if( UBTMakefileItem.LastWriteTime.CompareTo( GameProjectFilesLastUpdateTime ) < 0 )
+						{
+							// Game project files are newer than UBT make file
+							bForceOutOfDate = true;
+							Log.TraceInformation("UBTMakefile is older than generated game project files, ignoring it" );
+						}
                     }	
                     else
                     {
                         // Game project files folder is missing, so force outdatedness
-                        ProjectFilesLastUpdateTime = DateTime.MaxValue;
-
+						bForceOutOfDate = true;
                     }
                 }
-    
-                if( UBTMakefileItem.LastWriteTime.CompareTo( ProjectFilesLastUpdateTime ) < 0 )
-                {
-                    // Consider the file to be out of date
-                    bForceOutOfDate = true;
-                }
-            } 
+			}
+			else
+			{
+				// UBTMakefile doesn't even exist, so we won't bother loading it
+				bForceOutOfDate = true;
+			}
+
+
 
             UBTMakefile LoadedUBTMakefile = null;
             if( !bForceOutOfDate )
