UBT Hot-Reload detection fix.

diff --git a/Engine/Source/Programs/UnrealBuildTool/System/UnrealBuildTool.cs b/Engine/Source/Programs/UnrealBuildTool/System/UnrealBuildTool.cs
--- a/Engine/Source/Programs/UnrealBuildTool/System/UnrealBuildTool.cs
+++ b/Engine/Source/Programs/UnrealBuildTool/System/UnrealBuildTool.cs
@@ -1731,7 +1731,26 @@ namespace UnrealBuildTool
 				var EditorProcessName = Path.GetFileNameWithoutExtension(EditorProcessFilename);
 				var EditorProcesses = Process.GetProcessesByName(EditorProcessName);
 				var BinariesPath = Path.GetFullPath(Path.GetDirectoryName(EditorProcessFilename));
-				bIsRunning = EditorProcesses.FirstOrDefault(p => Path.GetFullPath(p.MainModule.FileName).StartsWith(BinariesPath, StringComparison.InvariantCultureIgnoreCase)) != default(Process);
+				var PerProjectBinariesPath = Path.Combine(Target.ProjectDirectory, BinariesPath.Substring(BinariesPath.LastIndexOf("Binaries")));
+				bIsRunning = EditorProcesses.FirstOrDefault(p =>
+					{
+						if(!Path.GetFullPath(p.MainModule.FileName).StartsWith(BinariesPath, StringComparison.InvariantCultureIgnoreCase))
+						{
+							return false;
+						}
+
+						if(PerProjectBinariesPath.Equals(BinariesPath, StringComparison.InvariantCultureIgnoreCase))
+						{
+							return false;
+						}
+
+						if(!p.Modules.Cast<System.Diagnostics.ProcessModule>().Any(Module => Path.GetFullPath(Module.FileName).StartsWith(PerProjectBinariesPath)))
+						{
+							return false;
+						}
+
+						return true;
+					}) != default(Process);
 			}
 			return bIsRunning;
 		}
