Removed various hacks for Clang on Windows that are no longer needed
- PDB files are now emitted from Clang (currently do not work with MSVC breakpoints)
- /EHsc option is now passed to the compiler (but Clang doesn't actually supported SEH exceptions yet)
- ICU is now enabled with Clang compiling

diff --git a/Engine/Source/Programs/UnrealBuildTool/Windows/UEBuildWindows.cs b/Engine/Source/Programs/UnrealBuildTool/Windows/UEBuildWindows.cs
--- a/Engine/Source/Programs/UnrealBuildTool/Windows/UEBuildWindows.cs
+++ b/Engine/Source/Programs/UnrealBuildTool/Windows/UEBuildWindows.cs
@@ -326,7 +326,7 @@ namespace UnrealBuildTool
          */
         public override void ValidateUEBuildConfiguration()
         {
-            UEBuildConfiguration.bCompileICU &= !bCompileWithClang;	// @todo clang: ICU causes STL link errors when using Clang on Windows.  Needs debugging.
+            UEBuildConfiguration.bCompileICU = true;
         }
 
         public override void ModifyNewlyLoadedModule(UEBuildModule InModule, TargetInfo Target)
diff --git a/Engine/Source/Programs/UnrealBuildTool/Windows/VCToolChain.cs b/Engine/Source/Programs/UnrealBuildTool/Windows/VCToolChain.cs
--- a/Engine/Source/Programs/UnrealBuildTool/Windows/VCToolChain.cs
+++ b/Engine/Source/Programs/UnrealBuildTool/Windows/VCToolChain.cs
@@ -270,8 +270,7 @@ namespace UnrealBuildTool
 				if (CompileEnvironment.Config.CLRMode == CPPCLRMode.CLRDisabled)
 				{
 					// Enable C++ exceptions when building with the editor or when building UHT.
-					if ((CompileEnvironment.Config.bEnableExceptions || UEBuildConfiguration.bBuildEditor || UEBuildConfiguration.bForceEnableExceptions)
-						&& !WindowsPlatform.bCompileWithClang)	// @todo clang: Clang compiler frontend doesn't understand /EHsc yet
+					if (CompileEnvironment.Config.bEnableExceptions || UEBuildConfiguration.bBuildEditor || UEBuildConfiguration.bForceEnableExceptions)
 					{
 						// Enable C++ exception handling, but not C exceptions.
 						Arguments.Append(" /EHsc");
@@ -297,41 +296,33 @@ namespace UnrealBuildTool
             // If enabled, create debug information.
 			if (CompileEnvironment.Config.bCreateDebugInfo)
 			{
-				if( WindowsPlatform.bCompileWithClang )	// @todo clang: /Zx options are not supported in clang-cl yet
+				// Store debug info in .pdb files.
+				// @todo clang: PDB files are emited from Clang but do not fully work with Visual Studio yet (breakpoints won't hit due to "symbol read error")
+				if( BuildConfiguration.bUsePDBFiles )
 				{
-					// Manually enable full debug info.  This is basically DWARF3 debug data.  MSVC won't know what to do with this and
-					// will just strip it out of executables linked using the MSVC linker.
-					Arguments.Append(" -Xclang -g");
-				}
-				else
-				{
-					// Store debug info in .pdb files.
-					if( BuildConfiguration.bUsePDBFiles )
+					// Create debug info suitable for E&C if wanted.
+					if( BuildConfiguration.bSupportEditAndContinue &&
+						// We only need to do this in debug as that's the only configuration that supports E&C.
+						CompileEnvironment.Config.Target.Configuration == CPPTargetConfiguration.Debug)
 					{
-						// Create debug info suitable for E&C if wanted.
-						if( BuildConfiguration.bSupportEditAndContinue &&
-							// We only need to do this in debug as that's the only configuration that supports E&C.
-							CompileEnvironment.Config.Target.Configuration == CPPTargetConfiguration.Debug)
-						{
-							Arguments.Append(" /ZI");
-						}
-						// Regular PDB debug information.
-						else
-						{
-							Arguments.Append(" /Zi");
-						}
-                        // We need to add this so VS won't lock the PDB file and prevent synchronous updates. This forces serialization through MSPDBSRV.exe.
-                        // See http://msdn.microsoft.com/en-us/library/dn502518.aspx for deeper discussion of /FS switch.
-                        if (BuildConfiguration.bUseIncrementalLinking)
-                        {
-                            Arguments.Append(" /FS");
-                        }
+						Arguments.Append(" /ZI");
 					}
-					// Store C7-format debug info in the .obj files, which is faster.
+					// Regular PDB debug information.
 					else
 					{
-						Arguments.Append(" /Z7");
+						Arguments.Append(" /Zi");
 					}
+                    // We need to add this so VS won't lock the PDB file and prevent synchronous updates. This forces serialization through MSPDBSRV.exe.
+                    // See http://msdn.microsoft.com/en-us/library/dn502518.aspx for deeper discussion of /FS switch.
+                    if (BuildConfiguration.bUseIncrementalLinking)
+                    {
+                        Arguments.Append(" /FS");
+                    }
+				}
+				// Store C7-format debug info in the .obj files, which is faster.
+				else
+				{
+					Arguments.Append(" /Z7");
 				}
 			}
 
@@ -897,8 +888,6 @@ namespace UnrealBuildTool
 						PDBFileName = "MiscPlainC";
 					}
 
-					// @todo Clang: Clang doesn't emit PDB files even when debugging is enabled
-					if( !WindowsPlatform.bCompileWithClang )
 					{ 
 						// Specify the PDB file that the compiler should write to.
 						FileItem PDBFile = FileItem.GetItemByPath(
@@ -1220,8 +1209,6 @@ namespace UnrealBuildTool
 				if (LinkEnvironment.Config.bCreateDebugInfo)
 				{
 					// Write the PDB file to the output directory.			
-					// @todo Clang: Clang doesn't emit PDB files even when debugging is enabled
-					if( !WindowsPlatform.bCompileWithClang )
 					{
 						string PDBFilePath = Path.Combine(LinkEnvironment.Config.OutputDirectory, Path.GetFileNameWithoutExtension(OutputFile.AbsolutePath) + ".pdb");
 						FileItem PDBFile = FileItem.GetItemByPath(PDBFilePath);
diff --git a/Engine/Source/Runtime/Core/Public/Windows/WIndowsPlatform.h b/Engine/Source/Runtime/Core/Public/Windows/WIndowsPlatform.h
--- a/Engine/Source/Runtime/Core/Public/Windows/WIndowsPlatform.h
+++ b/Engine/Source/Runtime/Core/Public/Windows/WIndowsPlatform.h
@@ -36,13 +36,10 @@ typedef FWindowsPlatformTypes FPlatformTypes;
 
 #define PLATFORM_LITTLE_ENDIAN								1
 #if defined(__clang__)
-	// @todo Clang Clang emits a compile error on Windows when __try/__except blocks are used (missing functionality in Clang.)
-	//             After Clang adds proper support for SEH exceptions, we should delete this definition and all use cases of it!
+	// @todo clang: Clang compiler on Windows doesn't support SEH exception handling yet (__try/__except)
 	#define PLATFORM_SEH_EXCEPTIONS_DISABLED				1
 
-	// @todo clang: This is needed because of a bug in Clang with RTTI and exceptions (even when RTTI is disabled!)  This manifests as a 
-	// compile error whenever try/catch/throw is used.  See http://llvm.org/bugs/show_bug.cgi?id=17403.  This code should be removed 
-	// after the bug is fixed in Clang!
+	// @todo clang: Clang compiler on Windows doesn't support C++ exception handling yet (try/throw/catch)
 	#define PLATFORM_EXCEPTIONS_DISABLED					1
 #endif
 
