UnrealBuildTool: Don't invalidate UBTMakefiles when game-specific project files don't exist
- This case is currently a legit workflow, when you have engine projects that contain game targets

diff --git a/Engine/Source/Programs/UnrealBuildTool/System/UnrealBuildTool.cs b/Engine/Source/Programs/UnrealBuildTool/System/UnrealBuildTool.cs
--- a/Engine/Source/Programs/UnrealBuildTool/System/UnrealBuildTool.cs
+++ b/Engine/Source/Programs/UnrealBuildTool/System/UnrealBuildTool.cs
@@ -2075,11 +2075,6 @@ namespace UnrealBuildTool
 							Log.TraceInformation("UBTMakefile is older than generated game project files, ignoring it" );
 						}
                     }	
-                    else
-                    {
-                        // Game project files folder is missing, so force outdatedness
-						bForceOutOfDate = true;
-                    }
                 }
 			}
 			else
