UnrealBuildTool: Minor fixes for Clang on Windows
- Actually turned off exceptions as Clang can't handle them yet
- Switched LLVM target to "msvc" mode instead of vanilla win32
- Added comments about third party libraries that cannot link using LLD yet (after doing tests)

diff --git a/Engine/Source/Programs/UnrealBuildTool/Windows/VCToolChain.cs b/Engine/Source/Programs/UnrealBuildTool/Windows/VCToolChain.cs
--- a/Engine/Source/Programs/UnrealBuildTool/Windows/VCToolChain.cs
+++ b/Engine/Source/Programs/UnrealBuildTool/Windows/VCToolChain.cs
@@ -32,7 +32,7 @@ namespace UnrealBuildTool
 
 			if( WindowsPlatform.bCompileWithClang )
 			{ 
-				// Result += " -###";	// @todo clang: Print Clang command-lines (instead of outputting compile results!)
+				// Arguments.Append( " -###" );	// @todo clang: Print Clang command-lines (instead of outputting compile results!)
 
 				// @todo clang: We're impersonating the Visual C++ compiler by setting MSC_VER and _MSC_FULL_VER to values that MSVC would set
 				string VersionString;
@@ -84,11 +84,11 @@ namespace UnrealBuildTool
 				// Tell the Clang compiler whether we want to generate 32-bit code or 64-bit code
 				if( CompileEnvironment.Config.Target.Platform == CPPTargetPlatform.Win64 )
 				{
-					Arguments.Append( " --target=x86_64-pc-win32" );
+					Arguments.Append( " --target=x86_64-pc-windows-msvc" );
 				}
 				else
 				{
-					Arguments.Append( " --target=x86-pc-win32" );
+					Arguments.Append( " --target=x86-pc-windows-msvc" );
 				}
 			}
 
@@ -270,7 +270,8 @@ namespace UnrealBuildTool
 				if (CompileEnvironment.Config.CLRMode == CPPCLRMode.CLRDisabled)
 				{
 					// Enable C++ exceptions when building with the editor or when building UHT.
-					if (CompileEnvironment.Config.bEnableExceptions || UEBuildConfiguration.bBuildEditor || UEBuildConfiguration.bForceEnableExceptions)
+					if (!WindowsPlatform.bCompileWithClang &&	// @todo clang: C++ exceptions are not supported with Clang on Windows yet
+						(CompileEnvironment.Config.bEnableExceptions || UEBuildConfiguration.bBuildEditor || UEBuildConfiguration.bForceEnableExceptions))
 					{
 						// Enable C++ exception handling, but not C exceptions.
 						Arguments.Append(" /EHsc");
@@ -424,6 +425,10 @@ namespace UnrealBuildTool
 			{
 				// This tells LLD to run in "Windows emulation" mode, meaning that it will accept MSVC Link arguments
 				Arguments.Append( " -flavor link" );
+
+				// @todo clang: The following static libraries aren't linking correctly with Clang:
+				//		tbbmalloc.lib, zlib_64.lib, libpng_64.lib, freetype2412MT.lib, IlmImf.lib
+				//		LLD: Assertion failed: result.size() == 1, file ..\tools\lld\lib\ReaderWriter\FileArchive.cpp, line 71
 			}
 
 			// Don't create a side-by-side manifest file for the executable.
@@ -773,7 +778,7 @@ namespace UnrealBuildTool
 					if( WindowsPlatform.bCompileWithClang )
 					{
 						// Tell Clang to generate a PCH header
-						FileArguments.Append(" -Xclang -x -Xclang c++-header");	// @todo clang: Doesn't work do to Clang-cl overriding us at the end of the command-line (see -### option output)
+						FileArguments.Append(" -Xclang -x -Xclang c++-header");	// @todo clang: Doesn't work due to Clang-cl overriding us at the end of the command-line (see -### option output)
 						FileArguments.AppendFormat(" /Fo\"{0}\"", PrecompiledHeaderFile.AbsolutePath);
 
 						// Clang PCH generation doesn't create an .obj file to link in, unlike MSVC
