UnrealBuildTool: Fast iteration mode fixes
- Fixed confusing exception when a generated code source file was missing
- No longer try compute PCH for a source file when in assembler only mode (PCH is loaded from UBTMakefile)

diff --git a/Engine/Source/Programs/UnrealBuildTool/Configuration/UEBuildModule.cs b/Engine/Source/Programs/UnrealBuildTool/Configuration/UEBuildModule.cs
--- a/Engine/Source/Programs/UnrealBuildTool/Configuration/UEBuildModule.cs
+++ b/Engine/Source/Programs/UnrealBuildTool/Configuration/UEBuildModule.cs
@@ -1536,6 +1536,11 @@ namespace UnrealBuildTool
 
 		public static FileItem CachePCHUsageForModuleSourceFile(UEBuildTarget Target, CPPEnvironment ModuleCompileEnvironment, FileItem CPPFile)
 		{
+			if( !CPPFile.bExists )
+			{
+				throw new BuildException( "Required source file not found: " + CPPFile.AbsolutePath );
+			}
+
 			var PCHCacheTimerStart = DateTime.UtcNow;
 
 			var BuildPlatform = UEBuildPlatform.GetBuildPlatformForCPPTargetPlatform( ModuleCompileEnvironment.Config.Target.Platform );
diff --git a/Engine/Source/Programs/UnrealBuildTool/System/ExternalExecution.cs b/Engine/Source/Programs/UnrealBuildTool/System/ExternalExecution.cs
--- a/Engine/Source/Programs/UnrealBuildTool/System/ExternalExecution.cs
+++ b/Engine/Source/Programs/UnrealBuildTool/System/ExternalExecution.cs
@@ -535,20 +535,23 @@ namespace UnrealBuildTool
 				{
 					// Since code files are definitely out of date, we'll now finish computing information about the UObject modules for UHT.  We
 					// want to save this work until we know that UHT actually needs to be run to speed up best-case iteration times.
-					foreach( var UHTModuleInfo in UObjectModules )
+					if( UnrealBuildTool.IsGatheringBuild )		// In assembler-only mode, PCH info is loaded from our UBTMakefile!
 					{
-						// Only cache the PCH name if we don't already have one.  When running in 'gather only' mode, this will have already been cached
-						if( string.IsNullOrEmpty( UHTModuleInfo.PCH ) )
+						foreach( var UHTModuleInfo in UObjectModules )
 						{
-							UHTModuleInfo.PCH = "";
-
-							// We need to figure out which PCH header this module is including, so that UHT can inject an include statement for it into any .cpp files it is synthesizing
-							var DependencyModuleCPP = (UEBuildModuleCPP)Target.GetModuleByName( UHTModuleInfo.ModuleName );
-							var ModuleCompileEnvironment = DependencyModuleCPP.CreateModuleCompileEnvironment(GlobalCompileEnvironment);
-							DependencyModuleCPP.CachePCHUsageForModuleSourceFiles(ModuleCompileEnvironment);
-							if (DependencyModuleCPP.ProcessedDependencies.UniquePCHHeaderFile != null)
+							// Only cache the PCH name if we don't already have one.  When running in 'gather only' mode, this will have already been cached
+							if( string.IsNullOrEmpty( UHTModuleInfo.PCH ) )
 							{
-								UHTModuleInfo.PCH = DependencyModuleCPP.ProcessedDependencies.UniquePCHHeaderFile.AbsolutePath;
+								UHTModuleInfo.PCH = "";
+
+								// We need to figure out which PCH header this module is including, so that UHT can inject an include statement for it into any .cpp files it is synthesizing
+								var DependencyModuleCPP = (UEBuildModuleCPP)Target.GetModuleByName( UHTModuleInfo.ModuleName );
+								var ModuleCompileEnvironment = DependencyModuleCPP.CreateModuleCompileEnvironment(GlobalCompileEnvironment);
+								DependencyModuleCPP.CachePCHUsageForModuleSourceFiles(ModuleCompileEnvironment);
+								if (DependencyModuleCPP.ProcessedDependencies.UniquePCHHeaderFile != null)
+								{
+									UHTModuleInfo.PCH = DependencyModuleCPP.ProcessedDependencies.UniquePCHHeaderFile.AbsolutePath;
+								}
 							}
 						}
 					}
