UnrealBuildTool: Fast iteration: Improved console output for UBTMakefile state
- We now print out the reason why we need to generate a new makefile

diff --git a/Engine/Source/Programs/UnrealBuildTool/System/UnrealBuildTool.cs b/Engine/Source/Programs/UnrealBuildTool/System/UnrealBuildTool.cs
--- a/Engine/Source/Programs/UnrealBuildTool/System/UnrealBuildTool.cs
+++ b/Engine/Source/Programs/UnrealBuildTool/System/UnrealBuildTool.cs
@@ -1357,7 +1357,7 @@ namespace UnrealBuildTool
                                 // we might not know about all of the C++ include dependencies for already-up-to-date shared build products
                                 // between the targets
                                 bNeedsFullCPPIncludeRescan = true;
-								Log.TraceInformation( "Forcing full C++ include scan because we're building a different target" );
+								Log.TraceInformation( "Performing full C++ include scan (building a new target)" );
                             }
                         }
                     }
@@ -1392,7 +1392,8 @@ namespace UnrealBuildTool
                         // @todo ubtmake: Mildly terrified of BuildConfiguration/UEBuildConfiguration globals that were set during the Prepare phase but are not set now.  We may need to save/load all of these, otherwise
                         //		we'll need to call SetupGlobalEnvironment on all of the targets (maybe other stuff, too.  See PreBuildStep())
                             
-                        UBTMakefile = LoadUBTMakefile( Targets );
+						string ReasonNotLoaded;
+                        UBTMakefile = LoadUBTMakefile( Targets, out ReasonNotLoaded );
 
                         if( UBTMakefile == null )
                         { 
@@ -1400,7 +1401,7 @@ namespace UnrealBuildTool
                             // a 'gather' and 'assemble' in the same run.  This will take a while longer, but subsequent runs will be fast!
                             UnrealBuildTool.bIsGatheringBuild_Unsafe = true;
 
-                            Log.TraceInformation( "Forcing generation of new UBTMakefile for {0}{1}", Targets[0].GetTargetName(), Targets.Count > 1 ? ( " (and " + ( Targets.Count - 1 ).ToString() + " more)" ) : "" );
+                            Log.TraceInformation( "Creating makefile for {0}{1} ({2})", Targets[0].GetTargetName(), Targets.Count > 1 ? ( " (and " + ( Targets.Count - 1 ).ToString() + " more)" ) : "", ReasonNotLoaded );
                         }
                     }
 
@@ -1992,7 +1993,7 @@ namespace UnrealBuildTool
         {
             if( !UBTMakefile.IsValidMakefile() )
             {
-                throw new BuildException( "Can't save a UBTMakefile that has invalid contents.  See UBTMakefile.IsValidMakefile()" );
+                throw new BuildException( "Can't save a makefile that has invalid contents.  See UBTMakefile.IsValidMakefile()" );
             }
 
             var TimerStartTime = DateTime.UtcNow;
@@ -2014,13 +2015,13 @@ namespace UnrealBuildTool
             }
             catch (Exception Ex)
             {
-                Console.Error.WriteLine("Failed to write UBTMakefile: {0}", Ex.Message);
+                Console.Error.WriteLine("Failed to write makefile: {0}", Ex.Message);
             }
 
             if (BuildConfiguration.bPrintPerformanceInfo)
             {
                 var TimerDuration = DateTime.UtcNow - TimerStartTime;
-                Log.TraceInformation("Saving UBTMakefile took " + TimerDuration.TotalSeconds + "s");
+                Log.TraceInformation("Saving makefile took " + TimerDuration.TotalSeconds + "s");
             }
         }
 
@@ -2029,9 +2030,12 @@ namespace UnrealBuildTool
         /// Loads a UBTMakefile from disk
         /// </summary>
         /// <param name="Targets">List of targets.  Order is not important</param>
-        static UBTMakefile LoadUBTMakefile( List<UEBuildTarget> Targets )
+		/// <param name="ReasonNotLoaded">If the function returns null, this string will contain the reason why</param>
+		/// <returns>The loaded makefile, or null if it failed for some reason.  On failure, the 'ReasonNotLoaded' variable will contain information about why</returns>
+        static UBTMakefile LoadUBTMakefile( List<UEBuildTarget> Targets, out string ReasonNotLoaded )
         {
             var UBTMakefileItem = FileItem.GetItemByFullPath( GetUBTMakefilePath( Targets ) );
+			ReasonNotLoaded = null;
 
             // Check the directory timestamp on the project files directory.  If the user has generated project files more
             // recently than the UBTMakefile, then we need to consider the file to be out of date
@@ -2051,7 +2055,9 @@ namespace UnrealBuildTool
 						{
 							// Engine project files are newer than UBT make file
 							bForceOutOfDate = true;
-							Log.TraceInformation("Existing UBTMakefile is older than generated engine project files, ignoring it" );
+							Log.TraceVerbose("Existing makefile is older than generated engine project files, ignoring it" );
+
+							ReasonNotLoaded = "project files are newer";
 						}
                     }
                 }
@@ -2072,7 +2078,9 @@ namespace UnrealBuildTool
 						{
 							// Game project files are newer than UBT make file
 							bForceOutOfDate = true;
-							Log.TraceInformation("UBTMakefile is older than generated game project files, ignoring it" );
+							Log.TraceVerbose("Makefile is older than generated game project files, ignoring it" );
+
+							ReasonNotLoaded = "game project files are newer";
 						}
                     }	
                 }
@@ -2081,6 +2089,8 @@ namespace UnrealBuildTool
 			{
 				// UBTMakefile doesn't even exist, so we won't bother loading it
 				bForceOutOfDate = true;
+
+				ReasonNotLoaded = "no existing makefile";
 			}
 
 
@@ -2098,13 +2108,15 @@ namespace UnrealBuildTool
                 }
                 catch (Exception Ex)
                 {
-                    Log.TraceWarning("Failed to read UBTMakefile: {0}", Ex.Message);
+                    Log.TraceWarning("Failed to read makefile: {0}", Ex.Message);
+					ReasonNotLoaded = "couldn't read existing makefile";
                 }
 
                 if( LoadedUBTMakefile != null && !LoadedUBTMakefile.IsValidMakefile() )
                 {
-                    Log.TraceWarning("Loaded UBTMakefile appears to have invalid contents, ignoring it ({0})", UBTMakefileItem.AbsolutePath );
+                    Log.TraceWarning("Loaded makefile appears to have invalid contents, ignoring it ({0})", UBTMakefileItem.AbsolutePath );
                     LoadedUBTMakefile = null;
+					ReasonNotLoaded = "existing makefile appears to be invalid";
                 }
             }
 
