UnrealBuildTool: Added friendly command descriptions when using local executor on Windows platform
- Now shows "Compile", "Link" and "Resource", similar to the Mac platform

diff --git a/Engine/Source/Programs/UnrealBuildTool/Windows/VCToolChain.cs b/Engine/Source/Programs/UnrealBuildTool/Windows/VCToolChain.cs
--- a/Engine/Source/Programs/UnrealBuildTool/Windows/VCToolChain.cs
+++ b/Engine/Source/Programs/UnrealBuildTool/Windows/VCToolChain.cs
@@ -726,6 +726,8 @@ namespace UnrealBuildTool
 			foreach (FileItem SourceFile in SourceFiles)
 			{
 				Action CompileAction = new Action(ActionType.Compile);
+				CompileAction.CommandDescription = "Compile";
+
 				StringBuilder FileArguments = new StringBuilder();
 				bool bIsPlainCFile = Path.GetExtension(SourceFile.AbsolutePath).ToUpperInvariant() == ".C";
 
@@ -980,6 +982,7 @@ namespace UnrealBuildTool
 			foreach (FileItem RCFile in RCFiles)
 			{
 				Action CompileAction = new Action(ActionType.Compile);
+				CompileAction.CommandDescription = "Resource";
 				CompileAction.WorkingDirectory = Path.GetFullPath(".");
 				CompileAction.CommandPath = GetVCToolPath(Environment.Config.Target.Platform, Environment.Config.Target.Configuration, "rc");
 				CompileAction.StatusDescription = Path.GetFileName(RCFile.AbsolutePath);
@@ -1240,6 +1243,7 @@ namespace UnrealBuildTool
 
 			// Create an action that invokes the linker.
 			Action LinkAction = new Action(ActionType.Link);
+			LinkAction.CommandDescription = "Link";
 			LinkAction.WorkingDirectory = Path.GetFullPath(".");
 			LinkAction.CommandPath = GetVCToolPath(
 				LinkEnvironment.Config.Target.Platform,
