UnrealBuildTool: UBTMakefiles are now outdated if UnrealBuildTool.exe was compiled more recently

diff --git a/Engine/Source/Programs/UnrealBuildTool/System/UnrealBuildTool.cs b/Engine/Source/Programs/UnrealBuildTool/System/UnrealBuildTool.cs
--- a/Engine/Source/Programs/UnrealBuildTool/System/UnrealBuildTool.cs
+++ b/Engine/Source/Programs/UnrealBuildTool/System/UnrealBuildTool.cs
@@ -1371,7 +1371,7 @@ namespace UnrealBuildTool
                     // NOTE: This is just a safeguard and doesn't have to be perfect.  We also check for newer project file timestamps in LoadUBTMakefile()
                     if( ProjectFileGenerator.bGenerateProjectFiles )	// @todo ubtmake: This is only hit when generating IntelliSense for project files.  Probably should be done right inside ProjectFileGenerator.bat
                     {													// @todo ubtmake: Won't catch multi-target cases as GPF always builds one target at a time for Intellisense
-                        // Delete the UBT makefile
+                        // Delete the UBTMakefile
                         var UBTMakefilePath = UnrealBuildTool.GetUBTMakefilePath( Targets );
                         if (File.Exists(UBTMakefilePath))
                         {
@@ -1386,12 +1386,13 @@ namespace UnrealBuildTool
                     }
 
                     // Were we asked to run in 'assembler only' mode?  If so, let's check to see if that's even possible by seeing if
-                    // we have a valid UBT Makefile already saved to disk, ready for us to load.
+                    // we have a valid UBTMakefile already saved to disk, ready for us to load.
                     if( UnrealBuildTool.bIsAssemblingBuild_Unsafe && !UnrealBuildTool.bIsGatheringBuild_Unsafe )
                     {
                         // @todo ubtmake: Mildly terrified of BuildConfiguration/UEBuildConfiguration globals that were set during the Prepare phase but are not set now.  We may need to save/load all of these, otherwise
                         //		we'll need to call SetupGlobalEnvironment on all of the targets (maybe other stuff, too.  See PreBuildStep())
-                            
+
+                        // Try to load the UBTMakefile.  It will only be loaded if it has valid content and is not determined to be out of date.    
 						string ReasonNotLoaded;
                         UBTMakefile = LoadUBTMakefile( Targets, out ReasonNotLoaded );
 
@@ -2053,7 +2054,7 @@ namespace UnrealBuildTool
                         var EngineProjectFilesLastUpdateTime = new DirectoryInfo( ProjectFileGenerator.IntermediateProjectFilesPath ).LastWriteTime;
 						if( UBTMakefileItem.LastWriteTime.CompareTo( EngineProjectFilesLastUpdateTime ) < 0 )
 						{
-							// Engine project files are newer than UBT make file
+							// Engine project files are newer than UBTMakefile
 							bForceOutOfDate = true;
 							Log.TraceVerbose("Existing makefile is older than generated engine project files, ignoring it" );
 
@@ -2066,24 +2067,43 @@ namespace UnrealBuildTool
                     // Rocket doesn't need to check engine projects for outdatedness
                 }
 
-                // Check the game project directory too
-                if( UnrealBuildTool.HasUProjectFile() )
-                { 
-                    var MasterProjectRelativePath = UnrealBuildTool.GetUProjectPath();
-                    var GameIntermediateProjectFilesPath = Path.Combine( MasterProjectRelativePath, "Intermediate", "ProjectFiles" );
-                    if( Directory.Exists( GameIntermediateProjectFilesPath ) )
-                    {
-                        var GameProjectFilesLastUpdateTime = new DirectoryInfo( GameIntermediateProjectFilesPath ).LastWriteTime;
-						if( UBTMakefileItem.LastWriteTime.CompareTo( GameProjectFilesLastUpdateTime ) < 0 )
+				if( !bForceOutOfDate )
+				{
+					// Check the game project directory too
+					if( UnrealBuildTool.HasUProjectFile() )
+					{ 
+						var MasterProjectRelativePath = UnrealBuildTool.GetUProjectPath();
+						var GameIntermediateProjectFilesPath = Path.Combine( MasterProjectRelativePath, "Intermediate", "ProjectFiles" );
+						if( Directory.Exists( GameIntermediateProjectFilesPath ) )
 						{
-							// Game project files are newer than UBT make file
-							bForceOutOfDate = true;
-							Log.TraceVerbose("Makefile is older than generated game project files, ignoring it" );
+							var GameProjectFilesLastUpdateTime = new DirectoryInfo( GameIntermediateProjectFilesPath ).LastWriteTime;
+							if( UBTMakefileItem.LastWriteTime.CompareTo( GameProjectFilesLastUpdateTime ) < 0 )
+							{
+								// Game project files are newer than UBTMakefile
+								bForceOutOfDate = true;
+								Log.TraceVerbose("Makefile is older than generated game project files, ignoring it" );
+
+								ReasonNotLoaded = "game project files are newer";
+							}
+						}	
+					}
+				}
+
+				if( !bForceOutOfDate )
+				{
+					// Check to see if UnrealBuildTool.exe was compiled more recently than the UBTMakefile
+					var UnrealBuildToolTimestamp = new FileInfo( Assembly.GetExecutingAssembly().Location ).LastWriteTime;
+					if( UBTMakefileItem.LastWriteTime.CompareTo( UnrealBuildToolTimestamp ) < 0 )
+					{
+						// UnrealBuildTool.exe was compiled more recently than the UBTMakefile
+						Log.TraceVerbose("Makefile is older than UnrealBuildTool.exe, ignoring it" );
 
-							ReasonNotLoaded = "game project files are newer";
-						}
-                    }	
-                }
+						ReasonNotLoaded = "UnrealBuildTool.exe is newer";
+						bForceOutOfDate = true;
+					}
+
+
+				}
 			}
 			else
 			{
@@ -2135,7 +2155,7 @@ namespace UnrealBuildTool
 
             if( Targets.Count == 1 )
             {
-                // If there's only one target, just save the UBT makefile in the target's build intermediate directory
+                // If there's only one target, just save the UBTMakefile in the target's build intermediate directory
                 // under a folder for that target (and platform/config combo.)
                 string PlatformIntermediatePath = BuildConfiguration.PlatformIntermediatePath;
                 if (UnrealBuildTool.HasUProjectFile())
