[GitHub] 377 : UBT now deletes hot reload pdb

#ttp 344200

diff --git a/Engine/Source/Programs/UnrealBuildTool/System/ActionGraph.cs b/Engine/Source/Programs/UnrealBuildTool/System/ActionGraph.cs
--- a/Engine/Source/Programs/UnrealBuildTool/System/ActionGraph.cs
+++ b/Engine/Source/Programs/UnrealBuildTool/System/ActionGraph.cs
@@ -459,9 +459,25 @@ namespace UnrealBuildTool
 										}
 										catch (Exception Ex)
 										{
-											// Ingore all exceptions
+											// Ignore all exceptions
 											Log.TraceInformation("Unable to delete old hot reload file: \"{0}\". Error: {0}", JunkFile, Ex.Message);
 										}
+
+										// Delete the PDB file.
+										string JunkPDBFile = JunkFile.Replace(ProducedItemExtension, ".pdb");
+										if (System.IO.File.Exists(JunkPDBFile))
+										{
+											try
+											{
+												Log.TraceInformation("Deleting old hot reload file: \"{0}\".", JunkPDBFile);
+												File.Delete(JunkPDBFile);
+											}
+											catch (Exception Ex)
+											{
+												// Ignore all exceptions
+												Log.TraceInformation("Unable to delete old hot reload file: \"{0}\". Error: {0}", JunkPDBFile, Ex.Message);
+											}
+										}
 									}
 								}
 							}
