fix for deleting junk files when setting up auto sdks and when validating the sdks
#ue4
#ubt

diff --git a/Engine/Source/Programs/UnrealBuildTool/System/UnrealBuildTool.cs b/Engine/Source/Programs/UnrealBuildTool/System/UnrealBuildTool.cs
--- a/Engine/Source/Programs/UnrealBuildTool/System/UnrealBuildTool.cs
+++ b/Engine/Source/Programs/UnrealBuildTool/System/UnrealBuildTool.cs
@@ -661,6 +661,7 @@ namespace UnrealBuildTool
             bool bRunCopyrightVerification = false;
             bool bDumpToFile = false;
             bool bCheckThirdPartyHeaders = false;
+			bool bIgnoreJunk = false;                    
 
             // @todo: Ideally we never need to Mutex unless we are invoked with the same target project,
             // in the same branch/path!  This would allow two clientspecs to build at the same time (even though we need
@@ -680,7 +681,8 @@ namespace UnrealBuildTool
                 if (Utils.ParseCommandLineFlag(Arguments, "-autosdkonly", out AutoSDKOnlyArgumentIndex))
                 {
                     bAutoSDKOnly = true;
-                }
+					bIgnoreJunk = true;
+				}
             }
             bool bValidatePlatforms = false;
             {
@@ -688,7 +690,8 @@ namespace UnrealBuildTool
                 if (Utils.ParseCommandLineFlag(Arguments, "-validateplatform", out ValidatePlatformsArgumentIndex))
                 {
                     bValidatePlatforms = true;
-                }
+					bIgnoreJunk = true;
+				}
             }
             
             // Don't allow simultaneous execution of Unreal Built Tool. Multi-selection in the UI e.g. causes this and you want serial
@@ -729,7 +732,6 @@ namespace UnrealBuildTool
                     var bGenerateMakefiles = false;
                     var bValidPlatformsOnly = false;
                     var bSpecificModulesOnly = false;
-                    var bIgnoreJunk = false;                    
 
                     // We need to be able to identify if one of the arguments is the platform...
                     // Leverage the existing parser function in UEBuildTarget to get this information.
