UnrealBuildTool: Minor console output cleanups
- Changed formatting of some console output
- No longer print notification about full C++ rescan unless that feature is active

diff --git a/Engine/Source/Programs/UnrealBuildTool/System/Distcc.cs b/Engine/Source/Programs/UnrealBuildTool/System/Distcc.cs
--- a/Engine/Source/Programs/UnrealBuildTool/System/Distcc.cs
+++ b/Engine/Source/Programs/UnrealBuildTool/System/Distcc.cs
@@ -48,7 +48,7 @@ namespace UnrealBuildTool
 				string DistccExecutable = BuildConfiguration.DistccExecutablesPath + "/distcc";
 				string GetHostExecutable = BuildConfiguration.DistccExecutablesPath + "/gethost";
 
-				Log.TraceInformation("Performing {0} actions (max {1} parallel jobs)", Actions.Count, MaxActionsToExecuteInParallel, DistccExecutable, GetHostExecutable);
+				Log.TraceInformation("Performing {0} actions ({1} in parallel)", Actions.Count, MaxActionsToExecuteInParallel, DistccExecutable, GetHostExecutable);
 
 				Dictionary<Action, ActionThread> ActionThreadDictionary = new Dictionary<Action, ActionThread>();
 				int JobNumber = 1;
diff --git a/Engine/Source/Programs/UnrealBuildTool/System/ExternalExecution.cs b/Engine/Source/Programs/UnrealBuildTool/System/ExternalExecution.cs
--- a/Engine/Source/Programs/UnrealBuildTool/System/ExternalExecution.cs
+++ b/Engine/Source/Programs/UnrealBuildTool/System/ExternalExecution.cs
@@ -633,7 +633,11 @@ namespace UnrealBuildTool
 						return false;
 					}
 
-					Log.TraceInformation( "Reflection code generation finished for {0} and took {1}", ActualTargetName, (double)s.ElapsedMilliseconds/1000.0 );
+					Log.TraceInformation( "Reflection code generated for {0}", ActualTargetName );
+					if( BuildConfiguration.bPrintPerformanceInfo )
+					{
+						Log.TraceInformation( "UnrealHeaderTool took {1}", ActualTargetName, (double)s.ElapsedMilliseconds/1000.0 );
+					}
 
 					// Now that UHT has successfully finished generating code, we need to update all cached FileItems in case their last write time has changed.
 					// Otherwise UBT might not detect changes UHT made.
diff --git a/Engine/Source/Programs/UnrealBuildTool/System/LocalExecutor.cs b/Engine/Source/Programs/UnrealBuildTool/System/LocalExecutor.cs
--- a/Engine/Source/Programs/UnrealBuildTool/System/LocalExecutor.cs
+++ b/Engine/Source/Programs/UnrealBuildTool/System/LocalExecutor.cs
@@ -384,7 +384,7 @@ namespace UnrealBuildTool
 				MaxActionsToExecuteInParallel = Math.Min(MaxActionsToExecuteInParallel, MaxActionsAffordedByMemory);
             }
 
-            Log.TraceInformation("Performing {0} actions (max {1} parallel jobs)", Actions.Count, MaxActionsToExecuteInParallel);
+            Log.TraceInformation("Performing {0} actions ({1} in parallel)", Actions.Count, MaxActionsToExecuteInParallel);
 
 			Dictionary<Action, ActionThread> ActionThreadDictionary = new Dictionary<Action, ActionThread>();
             int JobNumber = 1;
diff --git a/Engine/Source/Programs/UnrealBuildTool/System/UnrealBuildTool.cs b/Engine/Source/Programs/UnrealBuildTool/System/UnrealBuildTool.cs
--- a/Engine/Source/Programs/UnrealBuildTool/System/UnrealBuildTool.cs
+++ b/Engine/Source/Programs/UnrealBuildTool/System/UnrealBuildTool.cs
@@ -1316,7 +1316,7 @@ namespace UnrealBuildTool
                 // If UBT is being used to build something different than it did last time
                 if( !ProjectFileGenerator.bGenerateProjectFiles )
                 {
-                    bool bIsBuildingSameTargetsAsLastTime = false;
+					if( BuildConfiguration.bUseExperimentalFastDependencyScan )
                     {
                         // If we're building UHT without a Mutex, we'll need to assume that we're building the same targets and that no caches
                         // should be invalidated for this run.  This is important when UBT is invoked from within UBT in order to compile
@@ -1332,6 +1332,7 @@ namespace UnrealBuildTool
                             }
                         }
 
+	                    bool bIsBuildingSameTargetsAsLastTime = false;
                         if( bMustAssumeSameTargets )
                         {
                             bIsBuildingSameTargetsAsLastTime = true;
@@ -1425,7 +1426,7 @@ namespace UnrealBuildTool
                     UEBuildConfiguration.bHotReloadFromIDE = UEBuildConfiguration.bAllowHotReloadFromIDE && !Target.bEditorRecompile && ShouldDoHotReload(Target);
                     if (UEBuildConfiguration.bHotReloadFromIDE)
                     {
-                        Log.TraceInformation("Compiling game modules for Hot-Reload.");
+                        Log.TraceInformation("Compiling game modules for hot reload");
                     }
 
                     // When in 'assembler only' mode, we'll load this cache later on a worker thread.  It takes a long time to load!
