UnrealBuildTool: Fixed bug with fast build iteration mode
- Indirect C++ includes weren't considered outdated if the UBTMakefile was also out of date in the same session

diff --git a/Engine/Source/Programs/UnrealBuildTool/System/ActionGraph.cs b/Engine/Source/Programs/UnrealBuildTool/System/ActionGraph.cs
--- a/Engine/Source/Programs/UnrealBuildTool/System/ActionGraph.cs
+++ b/Engine/Source/Programs/UnrealBuildTool/System/ActionGraph.cs
@@ -748,10 +748,12 @@ namespace UnrealBuildTool
 				bool bFindCPPIncludePrerequisites = false;
 				if( RootAction.ActionType == ActionType.Compile )
 				{
-					// Up to date targets don't need their headers scanned, because presumably they would already be out of dated based on already-cached includes before getting this far
+					// Outdated targets don't need their headers scanned yet, because presumably they would already be out of dated based on already-cached
+					// includes before getting this far.  However, if we find them to be outdated after processing includes, we'll do a deep scan later
+					// on and cache all of the includes so that we have them for a quick outdatedness check the next run.
 					if( !bIsOutdated && 
 						BuildConfiguration.bUseExperimentalFastBuildIteration && 
-						!UnrealBuildTool.IsGatheringBuild && UnrealBuildTool.IsAssemblingBuild && 
+						UnrealBuildTool.IsAssemblingBuild && 
 						RootAction.ActionType == ActionType.Compile )
 					{
 						bFindCPPIncludePrerequisites = true;
