Clean redundant include paths.

diff --git a/Engine/Source/Programs/UnrealBuildTool/Configuration/UEBuildModule.cs b/Engine/Source/Programs/UnrealBuildTool/Configuration/UEBuildModule.cs
--- a/Engine/Source/Programs/UnrealBuildTool/Configuration/UEBuildModule.cs
+++ b/Engine/Source/Programs/UnrealBuildTool/Configuration/UEBuildModule.cs
@@ -537,8 +537,8 @@ namespace UnrealBuildTool
 				VisitedModules.Add(this,true);
 
 				// Add this module's public include paths and definitions.
-				IncludePaths.AddRange(PublicIncludePaths);
-				SystemIncludePaths.AddRange(PublicSystemIncludePaths);
+				AddIncludePathsWithChecks(IncludePaths, PublicIncludePaths);
+				AddIncludePathsWithChecks(SystemIncludePaths, PublicSystemIncludePaths);
 				Definitions.AddRange(PublicDefinitions);
 				
 				// If this module is being built into a DLL or EXE, set up an IMPORTS or EXPORTS definition for it.
@@ -611,7 +611,20 @@ namespace UnrealBuildTool
 				AdditionalFrameworks.AddRange(PublicAdditionalFrameworks);
 			}
 		}
-		
+
+		/** Adds PathsToAdd to IncludePaths, performing path normalization and ignoring duplicates. */
+		protected void AddIncludePathsWithChecks(List<string> IncludePaths, List<string> PathsToAdd)
+		{
+			foreach (var Path in PathsToAdd)
+			{
+				var NormalizedPath = Path.TrimEnd('/');
+				if (!IncludePaths.Contains(NormalizedPath))
+				{
+					IncludePaths.Add(NormalizedPath);
+				}
+			}
+		}
+
 		/** Sets up the environment for compiling this module. */
 		protected virtual void SetupPrivateCompileEnvironment(
 			ref List<string> IncludePaths,
@@ -623,7 +636,7 @@ namespace UnrealBuildTool
 			var VisitedModules = new Dictionary<UEBuildModule, bool>();
 
 			// Add this module's private include paths and definitions.
-			IncludePaths.AddRange( PrivateIncludePaths );
+			AddIncludePathsWithChecks(IncludePaths, PrivateIncludePaths);
 
 			// Allow the module's public dependencies to modify the compile environment.
 			bool bIncludePathsOnly = false;
