DEBUG and TRACE macros defined as expected for all Development/Debug configs.

diff --git a/Engine/Source/Programs/UnrealBuildTool/UnrealBuildTool.csproj b/Engine/Source/Programs/UnrealBuildTool/UnrealBuildTool.csproj
--- a/Engine/Source/Programs/UnrealBuildTool/UnrealBuildTool.csproj
+++ b/Engine/Source/Programs/UnrealBuildTool/UnrealBuildTool.csproj
@@ -38,8 +38,7 @@
     <DebugType>pdbonly</DebugType>
     <Optimize>true</Optimize>
     <OutputPath>..\..\..\Binaries\DotNET\</OutputPath>
-    <DefineConstants>
-    </DefineConstants>
+    <DefineConstants>TRACE</DefineConstants>
     <Optimize>true</Optimize>
     <DebugType>pdbonly</DebugType>
     <PlatformTarget>AnyCPU</PlatformTarget>
@@ -70,6 +69,11 @@
     <CodeAnalysisIgnoreBuiltInRuleSets>true</CodeAnalysisIgnoreBuiltInRuleSets>
     <CodeAnalysisIgnoreBuiltInRules>true</CodeAnalysisIgnoreBuiltInRules>
   </PropertyGroup>
+  <PropertyGroup Condition="'$(Configuration)|$(Platform)' == 'Debug|x86'">
+    <DefineConstants>TRACE;DEBUG</DefineConstants>
+    <Optimize>true</Optimize>
+    <OutputPath>bin\x86\Debug\</OutputPath>
+  </PropertyGroup>
   <ItemGroup>
     <Reference Include="Ionic.Zip.Reduced">
       <HintPath>..\..\..\Binaries\DotNET\Ionic.Zip.Reduced.dll</HintPath>
@@ -131,7 +135,7 @@
     <When Condition="Exists('Android\UEBuildAndroid.cs')">
       <ItemGroup>
         <Compile Include="Android\UEBuildAndroid.cs" />
-		<Compile Include="Android\UEDeployAndroid.cs" />
+        <Compile Include="Android\UEDeployAndroid.cs" />
         <Compile Include="Android\AndroidToolChain.cs" />
         <Compile Include="Android\AndroidProjectGenerator.cs" />
       </ItemGroup>
