TRACE define removed from Development builds for performance reasons.

diff --git a/Engine/Source/Programs/UnrealBuildTool/UnrealBuildTool.csproj b/Engine/Source/Programs/UnrealBuildTool/UnrealBuildTool.csproj
--- a/Engine/Source/Programs/UnrealBuildTool/UnrealBuildTool.csproj
+++ b/Engine/Source/Programs/UnrealBuildTool/UnrealBuildTool.csproj
@@ -38,7 +38,8 @@
     <DebugType>pdbonly</DebugType>
     <Optimize>true</Optimize>
     <OutputPath>..\..\..\Binaries\DotNET\</OutputPath>
-    <DefineConstants>TRACE</DefineConstants>
+    <DefineConstants>
+    </DefineConstants>
     <Optimize>true</Optimize>
     <DebugType>pdbonly</DebugType>
     <PlatformTarget>AnyCPU</PlatformTarget>
@@ -51,7 +52,8 @@
     <DebugType>pdbonly</DebugType>
     <Optimize>true</Optimize>
     <OutputPath>..\..\..\Binaries\DotNET\</OutputPath>
-    <DefineConstants>TRACE</DefineConstants>
+    <DefineConstants>
+    </DefineConstants>
     <Optimize>true</Optimize>
     <DebugType>pdbonly</DebugType>
     <PlatformTarget>x86</PlatformTarget>
