Remove non-existing include paths from compilation command lines.

diff --git a/Engine/Source/Programs/UnrealBuildTool/Configuration/UEBuildModule.cs b/Engine/Source/Programs/UnrealBuildTool/Configuration/UEBuildModule.cs
--- a/Engine/Source/Programs/UnrealBuildTool/Configuration/UEBuildModule.cs
+++ b/Engine/Source/Programs/UnrealBuildTool/Configuration/UEBuildModule.cs
@@ -612,6 +612,14 @@ namespace UnrealBuildTool
 			}
 		}
 
+		static string VCMacroPattern = @"\$\([A-Za-z0-9_]+\)";
+
+		/** Checks if path contains a VC macro */
+		protected bool DoesPathContainVCMacro(string Path)
+		{
+			return Regex.IsMatch(Path, VCMacroPattern);
+		}
+
 		/** Adds PathsToAdd to IncludePaths, performing path normalization and ignoring duplicates. */
 		protected void AddIncludePathsWithChecks(List<string> IncludePaths, List<string> PathsToAdd)
 		{
@@ -620,7 +628,11 @@ namespace UnrealBuildTool
 				var NormalizedPath = Path.TrimEnd('/');
 				if (!IncludePaths.Contains(NormalizedPath))
 				{
-					IncludePaths.Add(NormalizedPath);
+					// If path doesn't exist, it may contain VC macro (which is passed directly to and expanded by compiler).
+					if (Directory.Exists(NormalizedPath) || DoesPathContainVCMacro(NormalizedPath))
+					{
+						IncludePaths.Add(NormalizedPath);
+					}
 				}
 			}
 		}
