Regex fix.

diff --git a/Engine/Source/Programs/UnrealBuildTool/Configuration/UEBuildModule.cs b/Engine/Source/Programs/UnrealBuildTool/Configuration/UEBuildModule.cs
--- a/Engine/Source/Programs/UnrealBuildTool/Configuration/UEBuildModule.cs
+++ b/Engine/Source/Programs/UnrealBuildTool/Configuration/UEBuildModule.cs
@@ -612,12 +612,12 @@ namespace UnrealBuildTool
 			}
 		}
 
-		static string VCMacroPattern = @"\$\([A-Za-z0-9_]+\)";
+		static Regex VCMacroRegex = new Regex(@"\$\([A-Za-z0-9_]+\)");
 
 		/** Checks if path contains a VC macro */
 		protected bool DoesPathContainVCMacro(string Path)
 		{
-			return Regex.IsMatch(Path, VCMacroPattern);
+			return VCMacroRegex.IsMatch(Path);
 		}
 
 		/** Adds PathsToAdd to IncludePaths, performing path normalization and ignoring duplicates. */
