Remove code path from UBT which generates an XML build environment. APIDocTool now scrapes build environment from XGE export instead.

diff --git a/Engine/Source/Programs/UnrealBuildTool/Android/UEBuildAndroid.cs b/Engine/Source/Programs/UnrealBuildTool/Android/UEBuildAndroid.cs
--- a/Engine/Source/Programs/UnrealBuildTool/Android/UEBuildAndroid.cs
+++ b/Engine/Source/Programs/UnrealBuildTool/Android/UEBuildAndroid.cs
@@ -403,15 +403,5 @@ namespace UnrealBuildTool
 			// generate the .apk!
 			return new FileItem[] { FileItem.GetItemByPath(Config.OutputFilePath) };
 		}
-
-		/// <summary>
-		/// Writes an XML summary of the build environment for this binary
-		/// </summary>
-		/// <param name="CompileEnvironment">The environment to compile the binary in</param>
-		/// <param name="LinkEnvironment">The environment to link the binary in</param>
-		/// <returns></returns>
-		public override void WriteBuildEnvironment(CPPEnvironment CompileEnvironment, LinkEnvironment LinkEnvironment, XmlWriter Writer)
-		{
-		}
 	}
 }
diff --git a/Engine/Source/Programs/UnrealBuildTool/Configuration/BuildConfiguration.cs b/Engine/Source/Programs/UnrealBuildTool/Configuration/BuildConfiguration.cs
--- a/Engine/Source/Programs/UnrealBuildTool/Configuration/BuildConfiguration.cs
+++ b/Engine/Source/Programs/UnrealBuildTool/Configuration/BuildConfiguration.cs
@@ -307,12 +307,6 @@ namespace UnrealBuildTool
 		[XmlConfig]
 		public static bool bUseSharedPCHs;
 
-		/// <summary>
-		/// Writes an XML summary of the compile environment for the target to the given path, but does not build it.
-		/// </summary>
-		[XmlConfig]
-		public static string WriteTargetInfoPath;
-
 		/// <summary>
 		/// Whether the dependency cache includes pre-resolved include locations so UBT doesn't have to re-resolve each include location just to check the timestamp.
 		/// This is technically not fully correct because the dependency cache is global and each module could have a different set of include paths that could cause headers
@@ -443,8 +437,6 @@ namespace UnrealBuildTool
 			// take an inordinate amount of time, and intermediates will use up many gigabytes of disk space.
 			bUseSharedPCHs = true;
 
-			WriteTargetInfoPath = null;
-
 			// Using unity build to coalesce source files allows for much faster full rebuild times.  For fastest iteration times on single-file
 			// changes, consider turning this off.  In some cases, unity builds can cause linker errors after iterative changes to files, because
 			// the synthesized file names for unity code files may change between builds.
diff --git a/Engine/Source/Programs/UnrealBuildTool/Configuration/UEBuildBinary.cs b/Engine/Source/Programs/UnrealBuildTool/Configuration/UEBuildBinary.cs
--- a/Engine/Source/Programs/UnrealBuildTool/Configuration/UEBuildBinary.cs
+++ b/Engine/Source/Programs/UnrealBuildTool/Configuration/UEBuildBinary.cs
@@ -191,15 +191,6 @@ namespace UnrealBuildTool
 		/// <returns></returns>
 		public abstract IEnumerable<FileItem> Build(IUEToolChain ToolChain, CPPEnvironment CompileEnvironment,LinkEnvironment LinkEnvironment);
 
-		/// <summary>
-		/// Writes an XML summary of the build environment for this binary
-		/// </summary>
-		/// <param name="CompileEnvironment">The environment to compile the binary in</param>
-		/// <param name="LinkEnvironment">The environment to link the binary in</param>
-		/// <param name="Writer">The output XML writer</param>
-		/// <returns></returns>
-		public abstract void WriteBuildEnvironment(CPPEnvironment CompileEnvironment, LinkEnvironment LinkEnvironment, XmlWriter Writer);
-
 		/// <summary>
 		/// Called to allow the binary to modify the link environment of a different binary containing 
 		/// a module that depends on a module in this binary. */
@@ -623,28 +614,6 @@ namespace UnrealBuildTool
 			return OutputFiles;
 		}
 
-		/// <summary>
-		/// Writes an XML summary of the build environment for this binary
-		/// </summary>
-		/// <param name="CompileEnvironment">The environment to compile the binary in</param>
-		/// <param name="LinkEnvironment">The environment to link the binary in</param>
-		/// <param name="Writer">The output XML writer</param>
-		/// <returns></returns>
-		public override void WriteBuildEnvironment(CPPEnvironment CompileEnvironment, LinkEnvironment LinkEnvironment, XmlWriter Writer)
-		{
-			// Output documentation for each module
-			Writer.WriteStartElement("binary");
-			Writer.WriteAttributeString("name", Path.GetFileName(Config.OutputFilePath));
-			Writer.WriteAttributeString("type", "cpp");
-			foreach (var ModuleName in ModuleNames)
-			{
-				UEBuildModule Module = Target.GetModuleByName(ModuleName);
-				Module.WriteBuildEnvironment(CompileEnvironment, Writer);
-			}
-			Writer.WriteEndElement();
-		}
-
-
 		/// <summary>
 		/// Called to allow the binary to modify the link environment of a different binary containing 
 		/// a module that depends on a module in this binary.
@@ -752,25 +721,5 @@ namespace UnrealBuildTool
 
 			return new FileItem[] { FileItem.GetItemByPath(Config.OutputFilePath) };
 		}
-
-		/// <summary>
-		/// Writes an XML summary of the build environment for this binary
-		/// </summary>
-		/// <param name="CompileEnvironment">The environment to compile the binary in</param>
-		/// <param name="LinkEnvironment">The environment to link the binary in</param>
-		/// <param name="Writer">The output XML writer</param>
-		/// <returns></returns>
-		public override void WriteBuildEnvironment(CPPEnvironment CompileEnvironment, LinkEnvironment LinkEnvironment, XmlWriter Writer)
-		{
-			Writer.WriteStartElement("binary");
-			Writer.WriteAttributeString("name", Path.GetFileName(Config.OutputFilePath));
-			Writer.WriteAttributeString("type", "cs");
-
-			Writer.WriteStartElement("project");
-			Writer.WriteAttributeString("path", Config.ProjectFilePath);
-			Writer.WriteEndElement();
-
-			Writer.WriteEndElement();
-		}
 	};
 }
diff --git a/Engine/Source/Programs/UnrealBuildTool/Configuration/UEBuildModule.cs b/Engine/Source/Programs/UnrealBuildTool/Configuration/UEBuildModule.cs
--- a/Engine/Source/Programs/UnrealBuildTool/Configuration/UEBuildModule.cs
+++ b/Engine/Source/Programs/UnrealBuildTool/Configuration/UEBuildModule.cs
@@ -768,9 +768,6 @@ namespace UnrealBuildTool
 
 		/** Compiles the module, and returns a list of files output by the compiler. */
 		public abstract List<FileItem> Compile( CPPEnvironment GlobalCompileEnvironment, CPPEnvironment CompileEnvironment, bool bCompileMonolithic );
-
-		/** Writes the build environment for this module */
-		public abstract void WriteBuildEnvironment(CPPEnvironment CompileEnvironment, XmlWriter Writer);
 		
 		// Object interface.
 		public override string ToString()
@@ -868,15 +865,6 @@ namespace UnrealBuildTool
 		{
 			return new List<FileItem>();
 		}
-
-		public override void WriteBuildEnvironment(CPPEnvironment CompileEnvironment, XmlWriter Writer)
-		{
-			Writer.WriteStartElement("module");
-			Writer.WriteAttributeString("name", Name);
-			Writer.WriteAttributeString("path", ModuleDirectory);
-			Writer.WriteAttributeString("type", "external");
-			Writer.WriteEndElement();
-		}
 	};
 
 	/** A module that is compiled from C++ code. */
@@ -1790,72 +1778,6 @@ namespace UnrealBuildTool
 			}
 		}
 
-		public override void WriteBuildEnvironment(CPPEnvironment CompileEnvironment, XmlWriter Writer)
-		{
-			// Get the compile environment
-			List<string> PrivateIncludePaths = new List<string>();
-			List<string> PrivateSystemIncludePaths = new List<string>();
-			List<string> PrivateDefinitions = new List<string>();
-			List<UEBuildFramework> PrivateFrameworks = new List<UEBuildFramework>();
-			SetupPrivateCompileEnvironment(ref PrivateIncludePaths, ref PrivateSystemIncludePaths, ref PrivateDefinitions, ref PrivateFrameworks);
-
-			// Add all of the include paths
-			List<string> IncludePaths = new List<string>();
-			IncludePaths.AddRange(CompileEnvironment.Config.CPPIncludeInfo.SystemIncludePaths);
-			IncludePaths.AddRange(PrivateSystemIncludePaths);
-			IncludePaths.AddRange(CompileEnvironment.Config.CPPIncludeInfo.IncludePaths);
-			IncludePaths.AddRange(PrivateIncludePaths);
-
-			// Build the full list of definitions
-			List<string> Definitions = new List<string>();
-			Definitions.AddRange(CompileEnvironment.Config.Definitions);
-			Definitions.AddRange(PrivateDefinitions);
-
-			// Build a list of all the dependencies
-			List<string> Dependencies = new List<string>();
-			Dependencies.AddRange(PublicDependencyModuleNames);
-			Dependencies.AddRange(PrivateDependencyModuleNames);
-
-			// Write the output XML
-			Writer.WriteStartElement("module");
-			Writer.WriteAttributeString("name", Name);
-			Writer.WriteAttributeString("path", ModuleDirectory);
-			Writer.WriteAttributeString("type", "cpp");
-			Writer.WriteStartElement("dependencies");
-			foreach (string Dependency in Dependencies)
-			{
-				Writer.WriteStartElement("dependency");
-				Writer.WriteAttributeString("module", Dependency);
-				Writer.WriteEndElement();
-			}
-			Writer.WriteEndElement();
-			Writer.WriteStartElement("definitions");
-			foreach (string Definition in Definitions)
-			{
-				Writer.WriteStartElement("definition");
-				Writer.WriteAttributeString("name", Definition);
-				Writer.WriteEndElement();
-			}
-			Writer.WriteEndElement();
-			Writer.WriteStartElement("includes");
-			foreach (string IncludePath in IncludePaths)
-			{
-				Writer.WriteStartElement("include");
-				Writer.WriteAttributeString("path", IncludePath);
-				Writer.WriteEndElement();
-			}
-			Writer.WriteEndElement();
-			Writer.WriteStartElement("additionalFrameworks");
-			foreach (UEBuildFramework Framework in PrivateFrameworks)
-			{
-				Writer.WriteStartElement("framework");
-				Writer.WriteAttributeString("name", Framework.FrameworkName);
-				Writer.WriteEndElement();
-			}
-			Writer.WriteEndElement();
-			Writer.WriteEndElement();
-		}
-
 		public HashSet<FileItem> _AllClassesHeaders     = new HashSet<FileItem>();
 		public HashSet<FileItem> _PublicUObjectHeaders  = new HashSet<FileItem>();
 		public HashSet<FileItem> _PrivateUObjectHeaders = new HashSet<FileItem>();
diff --git a/Engine/Source/Programs/UnrealBuildTool/Configuration/UEBuildTarget.cs b/Engine/Source/Programs/UnrealBuildTool/Configuration/UEBuildTarget.cs
--- a/Engine/Source/Programs/UnrealBuildTool/Configuration/UEBuildTarget.cs
+++ b/Engine/Source/Programs/UnrealBuildTool/Configuration/UEBuildTarget.cs
@@ -1421,13 +1421,6 @@ namespace UnrealBuildTool
 				OutputItems.AddRange(Binary.Build(TargetToolChain, GlobalCompileEnvironment, GlobalLinkEnvironment));
 			}
 
-			if (BuildConfiguration.WriteTargetInfoPath != null)	// @todo ubtmake: This has some overlap with what we are doing with the action graph.  It might be easier to support other build systems with formatting like this, but it is a lot of duplicated code
-			{
-				Log.TraceInformation("Writing build environment to " + BuildConfiguration.WriteTargetInfoPath + "...");
-				WriteTargetInfo();
-				OutputItems.Clear();
-			}
-
 			return ECompilationResult.Succeeded;
 		}
 
@@ -1702,28 +1695,6 @@ namespace UnrealBuildTool
 			return FilteredBinaries;
 		}
 
-		/// <summary>
-		/// Writes target info.
-		/// </summary>
-		protected virtual void WriteTargetInfo()
-		{
-			Log.TraceInformation("Writing build environment to " + BuildConfiguration.WriteTargetInfoPath + "...");
-
-			XmlWriterSettings Settings = new XmlWriterSettings();
-			Settings.Indent = true;
-
-			using (XmlWriter Writer = XmlWriter.Create(BuildConfiguration.WriteTargetInfoPath, Settings))
-			{
-				Writer.WriteStartElement("target");
-				Writer.WriteAttributeString("name", AppName);
-				foreach (UEBuildBinary Binary in AppBinaries)
-				{
-					Binary.WriteBuildEnvironment(GlobalCompileEnvironment, GlobalLinkEnvironment, Writer);
-				}
-				Writer.WriteEndElement();
-			}
-		}
-
 		/// <summary>
 		/// All non-program monolithic binaries implicitly depend on all static plugin libraries so they are always linked appropriately
 		/// In order to do this, we create a new module here with a cpp file we emit that invokes an empty function in each library.
diff --git a/Engine/Source/Programs/UnrealBuildTool/IOS/UEBuildIOS.cs b/Engine/Source/Programs/UnrealBuildTool/IOS/UEBuildIOS.cs
--- a/Engine/Source/Programs/UnrealBuildTool/IOS/UEBuildIOS.cs
+++ b/Engine/Source/Programs/UnrealBuildTool/IOS/UEBuildIOS.cs
@@ -329,16 +329,6 @@ namespace UnrealBuildTool
             // generate the .stub!
             return new FileItem[] { FileItem.GetItemByPath(Config.OutputFilePath) };
         }
-
-        /// <summary>
-        /// Writes an XML summary of the build environment for this binary
-        /// </summary>
-        /// <param name="CompileEnvironment">The environment to compile the binary in</param>
-        /// <param name="LinkEnvironment">The environment to link the binary in</param>
-        /// <returns></returns>
-        public override void WriteBuildEnvironment(CPPEnvironment CompileEnvironment, LinkEnvironment LinkEnvironment, XmlWriter Writer)
-        {
-        }
     }
 
 }
diff --git a/Engine/Source/Programs/UnrealBuildTool/System/UnrealBuildTool.cs b/Engine/Source/Programs/UnrealBuildTool/System/UnrealBuildTool.cs
--- a/Engine/Source/Programs/UnrealBuildTool/System/UnrealBuildTool.cs
+++ b/Engine/Source/Programs/UnrealBuildTool/System/UnrealBuildTool.cs
@@ -1853,16 +1853,6 @@ namespace UnrealBuildTool
                 BuildConfiguration.bDisableDebugInfo = false;
                 BuildConfiguration.bOmitPCDebugInfoInDevelopment = false;
             }
-
-            foreach(string Argument in Arguments)
-            {
-                const string WriteTargetInfoArg = "-writetargetinfo=";
-                if (Argument.ToLowerInvariant().StartsWith(WriteTargetInfoArg))
-                {
-                    BuildConfiguration.WriteTargetInfoPath = Argument.Substring(WriteTargetInfoArg.Length);
-                    break;
-                }
-            }
         }
 
         /**
