Reset the .uproject path between calls to UEBuildTarget.CreateTarget() so that every target can set it correctly when generating project files.

diff --git a/Engine/Source/Programs/UnrealBuildTool/System/ProjectFileGenerator.cs b/Engine/Source/Programs/UnrealBuildTool/System/ProjectFileGenerator.cs
--- a/Engine/Source/Programs/UnrealBuildTool/System/ProjectFileGenerator.cs
+++ b/Engine/Source/Programs/UnrealBuildTool/System/ProjectFileGenerator.cs
@@ -1275,12 +1275,21 @@ namespace UnrealBuildTool
 						ArgumentsCopy[ 0 ] = TargetName;
 						Array.Copy(Arguments, 0, ArgumentsCopy, 1, Arguments.Length);
 
+						// CreateTarget mutates the current project setting file if it isn't already set, which prevents other 
+						// projects from being able to do the same. Capture the state beforehand, and reset if after running UBT. 
+						bool bHasProjectFile = UnrealBuildTool.HasUProjectFile();
+
 						bSuccess = UnrealBuildTool.RunUBT( ArgumentsCopy ) == ECompilationResult.Succeeded;
 						if( !bSuccess )
 						{
 							break;
 						}
 
+						if(!bHasProjectFile)
+						{
+							UnrealBuildTool.ResetProjectFile();
+						}
+
 						// Display progress
 						Progress.Write(TargetIndex + 1, TargetFiles.Count);
 					}
diff --git a/Engine/Source/Programs/UnrealBuildTool/System/UnrealBuildTool.cs b/Engine/Source/Programs/UnrealBuildTool/System/UnrealBuildTool.cs
--- a/Engine/Source/Programs/UnrealBuildTool/System/UnrealBuildTool.cs
+++ b/Engine/Source/Programs/UnrealBuildTool/System/UnrealBuildTool.cs
@@ -504,6 +504,14 @@ namespace UnrealBuildTool
             }
         }
 
+		public static void ResetProjectFile()
+		{
+			// Reset the project settings to their original state. Used by project file generator, since it mutates the active project setting.
+			UProjectFile = "";
+			bHasUProjectFile = false;
+			UProjectPath = "";
+		}
+
         private static bool ParseRocketCommandlineArg(string InArg, ref string OutGameName)
         {
             string LowercaseArg = InArg.ToLowerInvariant();
