Mac deployment fix

diff --git a/Engine/Source/Programs/UnrealBuildTool/Mac/UEDeployMac.cs b/Engine/Source/Programs/UnrealBuildTool/Mac/UEDeployMac.cs
--- a/Engine/Source/Programs/UnrealBuildTool/Mac/UEDeployMac.cs
+++ b/Engine/Source/Programs/UnrealBuildTool/Mac/UEDeployMac.cs
@@ -52,7 +52,7 @@ namespace UnrealBuildTool.IOS
 
 					if (!InTarget.GlobalLinkEnvironment.Config.bIsBuildingConsoleApplication)
 					{
-						string RemoteCreateAppBundleScript = Toolchain.ConvertPath(CreateAppBundleScript);
+						string RemoteCreateAppBundleScript = Toolchain.ConvertPath(Path.GetFullPath(CreateAppBundleScript));
 						RemoteCreateAppBundleScript = RemoteCreateAppBundleScript.Replace("../../../../", "../../");
 						RPCUtilHelper.CopyFile("../../" + CreateAppBundleScript, RemoteCreateAppBundleScript, true);
 					}
