UnrealBuildTool: Fixed exception on shutdown when using new fast build iteration mode

diff --git a/Engine/Source/Programs/UnrealBuildTool/System/UnrealBuildTool.cs b/Engine/Source/Programs/UnrealBuildTool/System/UnrealBuildTool.cs
--- a/Engine/Source/Programs/UnrealBuildTool/System/UnrealBuildTool.cs
+++ b/Engine/Source/Programs/UnrealBuildTool/System/UnrealBuildTool.cs
@@ -1288,6 +1288,7 @@ namespace UnrealBuildTool
             var ToolChain = UEToolChain.GetPlatformToolChain(BuildPlatform.GetCPPTargetPlatform(ResetPlatform));
 
             string EULAViolationWarning = null;
+			Thread CPPIncludesThread = null;
 
             try
             {
@@ -1610,7 +1611,6 @@ namespace UnrealBuildTool
                             // Cache indirect includes for all outdated C++ files.  We kick this off as a background thread so that it can
                             // perform the scan while we're compiling.  It usually only takes up to a few seconds, but we don't want to hurt
                             // our best case UBT iteration times for this task which can easily be performed asynchronously
-                            Thread CPPIncludesThread = null;
                             if( BuildConfiguration.bUseExperimentalFastDependencyScan && TargetToOutdatedPrerequisitesMap.Count > 0 )
                             {
                                 CPPIncludesThread = CreateThreadForCachingCPPIncludes( TargetToOutdatedPrerequisitesMap );
@@ -1632,12 +1632,6 @@ namespace UnrealBuildTool
                             {
                                 BuildResult = ECompilationResult.OtherCompilationError;
                             }
-
-                            if( CPPIncludesThread != null )
-                            { 
-                                // Wait until our CPPIncludes dependency scanner thread has finished
-                                CPPIncludesThread.Join();
-                            }
                         }
                     }
                 }
@@ -1654,6 +1648,12 @@ namespace UnrealBuildTool
                 BuildResult = ECompilationResult.OtherCompilationError;
             }
 
+			// Wait until our CPPIncludes dependency scanner thread has finished
+			if( CPPIncludesThread != null )
+			{ 
+			    CPPIncludesThread.Join();
+			}
+
             // Save the include dependency cache.
             { 
 				// NOTE: It's very important that we save the include cache, even if a build exception was thrown (compile error, etc), because we need to make sure that
