Different implementation for reading compiler version.  Now reads the Windows file version instead of the registry, which should work for all editions of VC, not just Professional.

Lots of environmental stuff read once per platform change instead of multiple times per file compiled.
VC environment state moved into its own class/file.
Environment passed around instead of globally referenced.

diff --git a/Engine/Source/Programs/UnrealBuildTool/UnrealBuildTool.csproj b/Engine/Source/Programs/UnrealBuildTool/UnrealBuildTool.csproj
--- a/Engine/Source/Programs/UnrealBuildTool/UnrealBuildTool.csproj
+++ b/Engine/Source/Programs/UnrealBuildTool/UnrealBuildTool.csproj
@@ -266,6 +266,7 @@
     <Compile Include="Configuration\EngineConfiguration.cs" />
     <Compile Include="Utilities\XmlConfigLoader.cs" />
     <Compile Include="Windows\UEBuildWindows.cs" />
+    <Compile Include="Windows\VCEnvironment.cs" />
     <Compile Include="Windows\VCToolChain.cs" />
     <Compile Include="Windows\WindowsProjectGenerator.cs" />
   </ItemGroup>
diff --git a/Engine/Source/Programs/UnrealBuildTool/UnrealBuildTool_Mono.csproj b/Engine/Source/Programs/UnrealBuildTool/UnrealBuildTool_Mono.csproj
--- a/Engine/Source/Programs/UnrealBuildTool/UnrealBuildTool_Mono.csproj
+++ b/Engine/Source/Programs/UnrealBuildTool/UnrealBuildTool_Mono.csproj
@@ -168,6 +168,7 @@
       <SubType>Code</SubType>
     </Compile>
     <Compile Include="Windows\UEBuildWindows.cs" />
+    <Compile Include="Windows\VCEnvironment.cs" />
     <Compile Include="Windows\VCToolChain.cs" />
     <Compile Include="Windows\WindowsProjectGenerator.cs" />
     <Compile Include="Mac\MacToolChain.cs" />
diff --git a/Engine/Source/Programs/UnrealBuildTool/Windows/VCEnvironment.cs b/Engine/Source/Programs/UnrealBuildTool/Windows/VCEnvironment.cs
new file mode 100644
--- /dev/null
+++ b/Engine/Source/Programs/UnrealBuildTool/Windows/VCEnvironment.cs
@@ -0,0 +1,224 @@
+// Copyright 1998-2014 Epic Games, Inc. All Rights Reserved.
+
+using System;
+using System.Collections.Generic;
+using System.Text.RegularExpressions;
+using System.Diagnostics;
+using System.IO;
+using Microsoft.Win32;
+using System.Text;
+
+namespace UnrealBuildTool
+{
+	class VCEnvironment
+	{
+		public readonly CPPTargetPlatform Platform;             // The platform the envvars have been initialized for
+		public readonly string            BaseVSToolPath;       // The path to Visual Studio's /Common7/Tools directory.
+		public readonly string            PlatformVSToolPath;   // The path to the platform tool binaries.
+		public readonly string            WindowsSDKDir;        // Installation folder of the Windows SDK, e.g. C:\Program Files\Microsoft SDKs\Windows\v6.0A\
+		public readonly string            CompilerPath;         // The path to the linker for linking executables
+		public readonly Version           CLExeVersion;         // The version of cl.exe we're running
+		public readonly string            LinkerPath;           // The path to the linker for linking executables
+		public readonly string            LibraryLinkerPath;    // The path to the linker for linking libraries
+		public readonly string            ResourceCompilerPath; // The path to the resource compiler
+		public readonly string            MSBuildPath;          // The path to MSBuild
+
+		/**
+		 * Initializes environment variables required by toolchain. Different for 32 and 64 bit.
+		 */
+		public static VCEnvironment SetEnvironment(CPPTargetPlatform Platform)
+		{
+			if (EnvVars != null && EnvVars.Platform == Platform)
+			{
+				return EnvVars;
+			}
+
+			EnvVars = new VCEnvironment(Platform);
+			return EnvVars;
+		}
+
+		private VCEnvironment(CPPTargetPlatform InPlatform)
+		{
+			Platform = InPlatform;
+
+			// If Visual Studio is not installed, the Windows SDK path will be used, which also happens to be the same
+			// directory. (It installs the toolchain into the folder where Visual Studio would have installed it to).
+			BaseVSToolPath = WindowsPlatform.GetVSComnToolsPath();
+			if (string.IsNullOrEmpty(BaseVSToolPath))
+			{
+				throw new BuildException("Visual Studio 2012 or Visual Studio 2013 must be installed in order to build this target.");
+			}
+
+			WindowsSDKDir        = FindWindowsSDKInstallationFolder();
+			PlatformVSToolPath   = GetPlatformVSToolPath      (Platform, BaseVSToolPath);
+			CompilerPath         = GetCompilerToolPath        (PlatformVSToolPath);
+			CLExeVersion         = FindCLExeVersion           (CompilerPath);
+			LinkerPath           = GetLinkerToolPath          (PlatformVSToolPath);
+			LibraryLinkerPath    = GetLibraryLinkerToolPath   (PlatformVSToolPath);
+			ResourceCompilerPath = GetResourceCompilerToolPath(Platform, WindowsSDKDir);
+
+			var VCVarsBatchFile = Path.Combine(BaseVSToolPath, (Platform == CPPTargetPlatform.Win64) ? "../../VC/bin/x86_amd64/vcvarsx86_amd64.bat" : "vsvars32.bat");
+			Utils.SetEnvironmentVariablesFromBatchFile(VCVarsBatchFile);
+
+			MSBuildPath = GetMSBuildToolPath();
+
+			// When targeting Windows XP on Visual Studio 2012+, we need to override the Windows SDK include and lib path set
+			// by the batch file environment (http://blogs.msdn.com/b/vcblog/archive/2012/10/08/10357555.aspx)
+			if (WindowsPlatform.SupportWindowsXP)
+			{
+				// Lib and bin folders have a x64 subfolder for 64 bit development.
+				var ConfigSuffix = (Platform == CPPTargetPlatform.Win64) ? "\\x64" : "";
+
+				Environment.SetEnvironmentVariable("PATH",    Utils.ResolveEnvironmentVariable(WindowsSDKDir + "bin" + ConfigSuffix + ";%PATH%"));
+				Environment.SetEnvironmentVariable("LIB",     Utils.ResolveEnvironmentVariable(WindowsSDKDir + "lib" + ConfigSuffix + ";%LIB%"));
+				Environment.SetEnvironmentVariable("INCLUDE", Utils.ResolveEnvironmentVariable(WindowsSDKDir + "include;%INCLUDE%"));
+			}
+		}
+
+		/// <returns>The path to Windows SDK directory for the specified version.</returns>
+		private static string FindWindowsSDKInstallationFolder()
+		{
+			// When targeting Windows XP on Visual Studio 2012+, we need to point at the older Windows SDK 7.1A that comes
+			// installed with Visual Studio 2012 Update 1. (http://blogs.msdn.com/b/vcblog/archive/2012/10/08/10357555.aspx)
+			string Version;
+			if (WindowsPlatform.SupportWindowsXP)
+			{
+				Version = "v7.1A";
+			}
+			else switch (WindowsPlatform.Compiler)
+			{
+				case WindowsCompiler.VisualStudio2013:
+					Version = "v8.1";
+					break;
+
+				case WindowsCompiler.VisualStudio2012:
+					Version = "v8.0";
+					break;
+
+				default:
+					throw new BuildException("Unexpected compiler setting when trying to determine Windows SDK folder");
+			}
+
+			// Based on VCVarsQueryRegistry
+			var Result =
+					Microsoft.Win32.Registry.GetValue(@"HKEY_CURRENT_USER\SOFTWARE\Microsoft\Microsoft SDKs\Windows\"              + Version, "InstallationFolder", null)
+				?? Microsoft.Win32.Registry.GetValue(@"HKEY_LOCAL_MACHINE\SOFTWARE\Wow6432Node\Microsoft\Microsoft SDKs\Windows\" + Version, "InstallationFolder", null)
+				?? Microsoft.Win32.Registry.GetValue(@"HKEY_CURRENT_USER\SOFTWARE\Wow6432Node\Microsoft\Microsoft SDKs\Windows\"  + Version, "InstallationFolder", null);
+
+			if (Result == null)
+			{
+				throw new BuildException("Windows SDK {0} must be installed in order to build this target.", Version);
+			}
+
+			return (string)Result;
+		}
+
+		/** Gets the path to the tool binaries for the specified platform. */
+		static string GetPlatformVSToolPath(CPPTargetPlatform Platform, string BaseVSToolPath)
+		{
+			// Regardless of the target, if we're linking on a 64 bit machine, we want to use the 64 bit linker (it's faster than the 32 bit linker)
+			[REDACTED_EMAIL]: Using the 64-bit linker appears to be broken at the moment.
+			if (Platform == CPPTargetPlatform.Win64)
+			{
+				// Use the native 64-bit compiler if present, otherwise use the amd64-on-x86 compiler. VS2012 Express only includes the latter.
+				var Result = Path.Combine(BaseVSToolPath, "../../VC/bin/amd64");
+				if (Directory.Exists(Result))
+				{
+					return Result;
+				}
+
+				return Path.Combine(BaseVSToolPath, "../../VC/bin/x86_amd64");
+			}
+
+			return Path.Combine(BaseVSToolPath, "../../VC/bin");
+		}
+
+		/** Gets the path to the compiler. */
+		static string GetCompilerToolPath(string PlatformVSToolPath)
+		{
+			// If we were asked to use Clang, then we'll redirect the path to the compiler to the LLVM installation directory
+			if (WindowsPlatform.bCompileWithClang)
+			{
+				var Result = Path.Combine( Environment.GetFolderPath( Environment.SpecialFolder.ProgramFilesX86 ), "LLVM", "msbuild-bin", "cl.exe" );
+				if (!File.Exists(Result))
+				{
+					throw new BuildException( "Clang was selected as the Windows compiler, but LLVM/Clang does not appear to be installed.  Could not find: " + Result );
+				}
+
+				return Result;
+			}
+
+			return Path.Combine(PlatformVSToolPath, "cl.exe");
+		}
+
+		/// <returns>The version of the compiler.</returns>
+		private static Version FindCLExeVersion(string CompilerExe)
+		{
+			var ExeVersionInfo = FileVersionInfo.GetVersionInfo(CompilerExe);
+			if (ExeVersionInfo == null)
+			{
+				throw new BuildException("Failed to read the version number of: " + CompilerExe);
+			}
+
+			return new Version(ExeVersionInfo.FileMajorPart, ExeVersionInfo.FileMinorPart, ExeVersionInfo.FileBuildPart, ExeVersionInfo.FilePrivatePart);
+		}
+
+		/** Gets the path to the linker. */
+		static string GetLinkerToolPath(string PlatformVSToolPath)
+		{
+			// If we were asked to use Clang, then we'll redirect the path to the compiler to the LLVM installation directory
+			if( WindowsPlatform.bCompileWithClang && WindowsPlatform.bAllowClangLinker )
+			{
+				var Result = Path.Combine( Environment.GetFolderPath( Environment.SpecialFolder.ProgramFilesX86 ), "LLVM", "bin", "lld.exe" );
+				if( !File.Exists( Result ) )
+				{
+					throw new BuildException( "Clang was selected as the Windows compiler, but LLVM/Clang does not appear to be installed.  Could not find: " + Result );
+				}
+
+				return Result;
+			}
+
+			return Path.Combine(PlatformVSToolPath, "link.exe");
+		}
+
+		/** Gets the path to the library linker. */
+		static string GetLibraryLinkerToolPath(string PlatformVSToolPath)
+		{
+			// Regardless of the target, if we're linking on a 64 bit machine, we want to use the 64 bit linker (it's faster than the 32 bit linker)
+			[REDACTED_EMAIL]: Using the 64-bit linker appears to be broken at the moment.
+			return Path.Combine(PlatformVSToolPath, "lib.exe");
+		}
+
+		/** Gets the path to the resource compiler's rc.exe for the specified platform. */
+		static string GetResourceCompilerToolPath(CPPTargetPlatform Platform, string WindowsSDKDir)
+		{
+			// 64 bit -- we can use the 32 bit version to target 64 bit on 32 bit OS.
+			if (Platform == CPPTargetPlatform.Win64)
+			{
+				return Path.Combine(WindowsSDKDir, "bin/x64/rc.exe");
+			}
+
+			if( !WindowsPlatform.SupportWindowsXP )	// Windows XP requires use to force Windows SDK 7.1 even on the newer compiler, so we need the old path RC.exe
+			{
+				return Path.Combine(WindowsSDKDir, "bin/x86/rc.exe");
+			}
+
+			return Path.Combine(WindowsSDKDir, "bin/rc.exe");
+		}
+
+		/** Gets the path to MSBuild. */
+		static string GetMSBuildToolPath()
+		{
+			string FrameworkDirectory = Environment.GetEnvironmentVariable("FrameworkDir");
+			string FrameworkVersion   = Environment.GetEnvironmentVariable("FrameworkVersion");
+			if (FrameworkDirectory == null || FrameworkVersion == null)
+			{
+				throw new BuildException( "NOTE: Please ensure that 64bit Tools are installed with DevStudio - there is usually an option to install these during install" );
+			}
+
+			return Path.Combine(FrameworkDirectory, FrameworkVersion, "MSBuild.exe");
+		}
+
+		static VCEnvironment EnvVars = null;
+	}
+}
diff --git a/Engine/Source/Programs/UnrealBuildTool/Windows/VCToolChain.cs b/Engine/Source/Programs/UnrealBuildTool/Windows/VCToolChain.cs
--- a/Engine/Source/Programs/UnrealBuildTool/Windows/VCToolChain.cs
+++ b/Engine/Source/Programs/UnrealBuildTool/Windows/VCToolChain.cs
@@ -21,7 +21,7 @@ namespace UnrealBuildTool
 			UEToolChain.RegisterPlatformToolChain(CPPTargetPlatform.Win32, this);
 		}
 
-		static void AppendCLArguments_Global(CPPEnvironment CompileEnvironment, StringBuilder Arguments)
+		static void AppendCLArguments_Global(CPPEnvironment CompileEnvironment, VCEnvironment EnvVars, StringBuilder Arguments)
 		{
 			// @todo fastubt: Use should be using StringBuilder instead of concatenation in all of these toolchain files.  On Mono, string builder is dramatically faster.
 
@@ -210,11 +210,11 @@ namespace UnrealBuildTool
 					{
 						// Allow optimized code to be debugged more easily.  This makes PDBs a bit larger, but doesn't noticeably affect
 						// compile times.  The executable code is not affected at all by this switch, only the debugging information.
-					  	if (WindowsPlatform.Compiler == WindowsCompiler.VisualStudio2013 && EnvVars.UpdateVersion >= new Version("12.0.30723"))
-					  	{
+						if (EnvVars.CLExeVersion >= new Version("18.0.30723"))
+						{
 							// VC2013 Update 3 has a new flag for doing this
 							Arguments.Append(" /Zo");
-					 	}	
+						}
 						else
 						{
 							Arguments.Append(" /d2Zi+");
@@ -658,10 +658,10 @@ namespace UnrealBuildTool
 
 		public override CPPOutput CompileCPPFiles(UEBuildTarget Target, CPPEnvironment CompileEnvironment, List<FileItem> SourceFiles, string ModuleName)
 		{
-			InitializeEnvironmentVariables(CompileEnvironment.Config.Target.Platform);
+			var EnvVars = VCEnvironment.SetEnvironment(CompileEnvironment.Config.Target.Platform);
 
 			StringBuilder Arguments = new StringBuilder();
-			AppendCLArguments_Global(CompileEnvironment, Arguments);
+			AppendCLArguments_Global(CompileEnvironment, EnvVars, Arguments);
 
 			// Add include paths to the argument list.
 			foreach (string IncludePath in CompileEnvironment.Config.CPPIncludeInfo.IncludePaths)
@@ -928,7 +928,7 @@ namespace UnrealBuildTool
 				}
 
 				CompileAction.WorkingDirectory = Path.GetFullPath(".");
-				CompileAction.CommandPath = GetVCToolPath(CompileEnvironment.Config.Target.Platform, CompileEnvironment.Config.Target.Configuration, "cl");
+				CompileAction.CommandPath      = EnvVars.CompilerPath;
 
 				if( !WindowsPlatform.bCompileWithClang )
 				{
@@ -975,6 +975,8 @@ namespace UnrealBuildTool
 
 		public override CPPOutput CompileRCFiles(UEBuildTarget Target, CPPEnvironment Environment, List<FileItem> RCFiles)
 		{
+			var EnvVars = VCEnvironment.SetEnvironment(Environment.Config.Target.Platform);
+
 			CPPOutput Result = new CPPOutput();
 
 			var BuildPlatform = UEBuildPlatform.GetBuildPlatformForCPPTargetPlatform(Environment.Config.Target.Platform);
@@ -983,9 +985,9 @@ namespace UnrealBuildTool
 			{
 				Action CompileAction = new Action(ActionType.Compile);
 				CompileAction.CommandDescription = "Resource";
-				CompileAction.WorkingDirectory = Path.GetFullPath(".");
-				CompileAction.CommandPath = GetVCToolPath(Environment.Config.Target.Platform, Environment.Config.Target.Configuration, "rc");
-				CompileAction.StatusDescription = Path.GetFileName(RCFile.AbsolutePath);
+				CompileAction.WorkingDirectory   = Path.GetFullPath(".");
+				CompileAction.CommandPath        = EnvVars.ResourceCompilerPath;
+				CompileAction.StatusDescription  = Path.GetFileName(RCFile.AbsolutePath);
 
 				if( WindowsPlatform.bCompileWithClang )
 				{ 
@@ -1060,6 +1062,8 @@ namespace UnrealBuildTool
 
 		public override FileItem LinkFiles(LinkEnvironment LinkEnvironment, bool bBuildImportLibraryOnly)
 		{
+			var EnvVars = VCEnvironment.SetEnvironment(LinkEnvironment.Config.Target.Platform);
+
 			if (LinkEnvironment.Config.bIsBuildingDotNetAssembly)
 			{
 				return FileItem.GetItemByPath(LinkEnvironment.Config.OutputFilePath);
@@ -1244,15 +1248,12 @@ namespace UnrealBuildTool
 			// Create an action that invokes the linker.
 			Action LinkAction = new Action(ActionType.Link);
 			LinkAction.CommandDescription = "Link";
-			LinkAction.WorkingDirectory = Path.GetFullPath(".");
-			LinkAction.CommandPath = GetVCToolPath(
-				LinkEnvironment.Config.Target.Platform,
-				LinkEnvironment.Config.Target.Configuration,
-				bIsBuildingLibrary ? "lib" : "link");
-			LinkAction.CommandArguments = Arguments.ToString();
-			LinkAction.ProducedItems.AddRange(ProducedItems);
+			LinkAction.WorkingDirectory   = Path.GetFullPath(".");
+			LinkAction.CommandPath        = bIsBuildingLibrary ? EnvVars.LibraryLinkerPath : EnvVars.LinkerPath;
+			LinkAction.CommandArguments   = Arguments.ToString();
+			LinkAction.ProducedItems    .AddRange(ProducedItems);
 			LinkAction.PrerequisiteItems.AddRange(PrerequisiteItems);
-			LinkAction.StatusDescription = Path.GetFileName(OutputFile.AbsolutePath);
+			LinkAction.StatusDescription  = Path.GetFileName(OutputFile.AbsolutePath);
 
 			if( WindowsPlatform.bCompileWithClang )
 			{ 
@@ -1274,6 +1275,9 @@ namespace UnrealBuildTool
 
 		public override void CompileCSharpProject(CSharpEnvironment CompileEnvironment, string ProjectFileName, string DestinationFile)
 		{
+			// Initialize environment variables required for spawned tools.
+			var EnvVars = VCEnvironment.SetEnvironment(CompileEnvironment.EnvironmentTargetPlatform);
+
 			var BuildProjectAction = new Action(ActionType.BuildProject);
 
 			// Specify the source file (prerequisite) for the action
@@ -1303,9 +1307,9 @@ namespace UnrealBuildTool
 			BuildProjectAction.bCanExecuteRemotely = false;
 
 			// Setup execution via MSBuild.
-			BuildProjectAction.WorkingDirectory = Path.GetFullPath(".");
+			BuildProjectAction.WorkingDirectory  = Path.GetFullPath(".");
 			BuildProjectAction.StatusDescription = Path.GetFileName(ProjectFileName);
-			BuildProjectAction.CommandPath = GetDotNetFrameworkToolPath(CompileEnvironment.EnvironmentTargetPlatform, "MSBuild");
+			BuildProjectAction.CommandPath       = EnvVars.MSBuildPath;
 			if (CompileEnvironment.TargetConfiguration == CSharpTargetConfiguration.Debug)
 			{
 				BuildProjectAction.CommandArguments = " /target:rebuild /property:Configuration=Debug";
@@ -1334,7 +1338,7 @@ namespace UnrealBuildTool
 			Debug.Assert(Platform == CPPTargetPlatform.Win32 || Platform == CPPTargetPlatform.Win64);
 
 			// Make sure we've got the environment variables set up for this target
-			VCToolChain.InitializeEnvironmentVariables(Platform);
+			VCEnvironment.SetEnvironment(Platform);
 
 			// Also add any include paths from the INCLUDE environment variable.  MSVC is not necessarily running with an environment that
 			// matches what UBT extracted from the vcvars*.bat using SetEnvironmentVariablesFromBatchFile().  We'll use the variables we
@@ -1349,214 +1353,6 @@ namespace UnrealBuildTool
 			return IncludePaths;
 		}
 
-		/** Accesses the bin directory for the VC toolchain for the specified platform. */
-		static string GetVCToolPath(CPPTargetPlatform Platform, CPPTargetConfiguration Configuration, string ToolName)
-		{	
-			// Initialize environment variables required for spawned tools.
-			InitializeEnvironmentVariables( Platform );
-
-			// Out variable that is going to contain fully qualified path to executable upon return.
-			string VCToolPath = "";
-
-			// rc.exe resides in the Windows SDK folder.
-			if (ToolName.ToUpperInvariant() == "RC")
-			{
-				// 64 bit -- we can use the 32 bit version to target 64 bit on 32 bit OS.
-				if (Platform == CPPTargetPlatform.Win64)
-				{
-					VCToolPath = Path.Combine(EnvVars.WindowsSDKDir, "bin/x64/rc.exe");
-				}
-				else if( !WindowsPlatform.SupportWindowsXP )	// Windows XP requires use to force Windows SDK 7.1 even on the newer compiler, so we need the old path RC.exe
-				{
-					VCToolPath = Path.Combine(EnvVars.WindowsSDKDir, "bin/x86/rc.exe");
-				}
-				else
-				{
-					VCToolPath = Path.Combine(EnvVars.WindowsSDKDir, "bin/rc.exe");
-				}
-			}
-			// cl.exe and link.exe are found in the toolchain specific folders (32 vs. 64 bit)
-			else
-			{
-				bool bIsRequestingLinkTool = ToolName.Equals( "link", StringComparison.InvariantCultureIgnoreCase );
-				bool bIsRequestingLibTool = ToolName.Equals( "lib", StringComparison.InvariantCultureIgnoreCase );
-
-				// If we were asked to use Clang, then we'll redirect the path to the compiler to the LLVM installation directory
-				if( WindowsPlatform.bCompileWithClang && !bIsRequestingLibTool && !( bIsRequestingLinkTool && !WindowsPlatform.bAllowClangLinker ) )
-				{
-					if( bIsRequestingLinkTool )
-					{
-						ToolName = Path.Combine( "bin", "lld" );
-					}
-					else
-					{
-						ToolName = Path.Combine( "msbuild-bin", "cl" );
-					}
-					VCToolPath = Path.Combine( Environment.GetFolderPath( Environment.SpecialFolder.ProgramFilesX86 ), "LLVM", ToolName + ".exe" );
-					if( !File.Exists( VCToolPath ) )
-					{
-						throw new BuildException( "Clang was selected as the Windows compiler, but LLVM/Clang does not appear to be installed.  Could not find: " + VCToolPath );
-					}
-				}
-				else
-				{
-					// Both target and build machines are 64 bit
-					bool bIs64Bit = (Platform == CPPTargetPlatform.Win64);
-					// Regardless of the target, if we're linking on a 64 bit machine, we want to use the 64 bit linker (it's faster than the 32 bit linker)
-					[REDACTED_EMAIL]: Using the 64-bit linker appears to be broken at the moment.
-					bool bUse64BitLinker = (Platform == CPPTargetPlatform.Win64) && bIsRequestingLinkTool;
-
-					// Use the 64 bit tools if the build machine and target are 64 bit or if we're linking a 32 bit binary on a 64 bit machine
-					if (bIs64Bit || bUse64BitLinker)
-					{
-						// Use the native 64-bit compiler if present, otherwise use the amd64-on-x86 compiler. VS2012 Express only includes the latter.
-						string PlatformToolPath = Path.Combine(EnvVars.BaseVSToolPath, "../../VC/bin/amd64/");
-						if(!Directory.Exists(PlatformToolPath))
-						{
-							PlatformToolPath = Path.Combine(EnvVars.BaseVSToolPath, "../../VC/bin/x86_amd64/");
-						}
-						VCToolPath = PlatformToolPath + ToolName + ".exe";
-					}
-					else
-					{
-						// Use 32 bit for cl.exe and other tools, or for link.exe if 64 bit path doesn't exist and we're targeting 32 bit.
-						VCToolPath = Path.Combine(EnvVars.BaseVSToolPath, "../../VC/bin/" + ToolName + ".exe");
-					}
-				}
-			}
-
-			return VCToolPath;
-		}
-
-		/** Accesses the directory for .NET Framework binaries such as MSBuild */
-		static string GetDotNetFrameworkToolPath(CPPTargetPlatform Platform, string ToolName)
-		{
-			// Initialize environment variables required for spawned tools.
-			InitializeEnvironmentVariables(Platform);
-
-			string FrameworkDirectory = Environment.GetEnvironmentVariable("FrameworkDir");
-			string FrameworkVersion = Environment.GetEnvironmentVariable("FrameworkVersion");
-			if (FrameworkDirectory == null || FrameworkVersion == null)
-			{
-				throw new BuildException( "NOTE: Please ensure that 64bit Tools are installed with DevStudio - there is usually an option to install these during install" );
-			}
-			string DotNetFrameworkBinDir = Path.Combine(FrameworkDirectory, FrameworkVersion);
-			string ToolPath = Path.Combine(DotNetFrameworkBinDir, ToolName + ".exe");
-			return ToolPath;
-		}
-
-		class EnvVarsType
-		{
-			public EnvVarsType(CPPTargetPlatform InPlatform, string InWindowsSDKDir, string InBaseVSToolPath, Version InUpdateVersion)
-			{
-				Platform       = InPlatform;
-				WindowsSDKDir  = InWindowsSDKDir;
-				BaseVSToolPath = InBaseVSToolPath;
-				UpdateVersion  = InUpdateVersion;
-			}
-
-			public readonly CPPTargetPlatform Platform;       // The platform the envvars have been initialized for
-			public readonly string            BaseVSToolPath; // The path to Visual Studio's /Common7/Tools directory.
-			public readonly string            WindowsSDKDir;  // Installation folder of the Windows SDK, e.g. C:\Program Files\Microsoft SDKs\Windows\v6.0A\
-			public readonly Version           UpdateVersion;  // The update version of Visual C++
-		}
-
-		static EnvVarsType EnvVars = null;
-
-		/**
-		 * Initializes environment variables required by toolchain. Different for 32 and 64 bit.
-		 */
-		static void InitializeEnvironmentVariables(CPPTargetPlatform Platform)
-		{
-			if (EnvVars != null && EnvVars.Platform == Platform)
-			{
-				return;
-			}
-
-			// If Visual Studio is not installed, the Windows SDK path will be used, which also happens to be the same
-			// directory. (It installs the toolchain into the folder where Visual Studio would have installed it to).
-			var BaseVSToolPath = WindowsPlatform.GetVSComnToolsPath();
-			if (string.IsNullOrEmpty(BaseVSToolPath))
-			{
-				throw new BuildException("Visual Studio 2012 or Visual Studio 2013 must be installed in order to build this target.");
-			}
-
-			// When targeting Windows XP on Visual Studio 2012+, we need to point at the older Windows SDK 7.1A that comes
-			// installed with Visual Studio 2012 Update 1. (http://blogs.msdn.com/b/vcblog/archive/2012/10/08/10357555.aspx)
-			var WindowsSDKDir = "";
-			if (WindowsPlatform.SupportWindowsXP)
-			{
-				WindowsSDKDir = FindWindowsSDKInstallationFolder("v7.1A");
-			}
-			else if (WindowsPlatform.Compiler == WindowsCompiler.VisualStudio2013)
-			{
-				WindowsSDKDir = FindWindowsSDKInstallationFolder("v8.1");
-			}
-			else if (WindowsPlatform.Compiler == WindowsCompiler.VisualStudio2012)
-			{
-				WindowsSDKDir = FindWindowsSDKInstallationFolder("v8.0");
-			}
-
-			var UpdateVersion = FindVCUpdateVersion();
-
-			var VCVarsBatchFile = Path.Combine(BaseVSToolPath, (Platform == CPPTargetPlatform.Win64) ? "../../VC/bin/x86_amd64/vcvarsx86_amd64.bat" : "vsvars32.bat");
-			Utils.SetEnvironmentVariablesFromBatchFile(VCVarsBatchFile);
-
-			// When targeting Windows XP on Visual Studio 2012+, we need to override the Windows SDK include and lib path set
-			// by the batch file environment (http://blogs.msdn.com/b/vcblog/archive/2012/10/08/10357555.aspx)
-			if (WindowsPlatform.SupportWindowsXP)
-			{
-				// Lib and bin folders have a x64 subfolder for 64 bit development.
-				var ConfigSuffix = (Platform == CPPTargetPlatform.Win64) ? "\\x64" : "";
-
-				Environment.SetEnvironmentVariable("PATH",    Utils.ResolveEnvironmentVariable(WindowsSDKDir + "bin" + ConfigSuffix + ";%PATH%"));
-				Environment.SetEnvironmentVariable("LIB",     Utils.ResolveEnvironmentVariable(WindowsSDKDir + "lib" + ConfigSuffix + ";%LIB%"));
-				Environment.SetEnvironmentVariable("INCLUDE", Utils.ResolveEnvironmentVariable(WindowsSDKDir + "include;%INCLUDE%"));
-			}
-
-			EnvVars = new EnvVarsType(
-				InPlatform      : Platform,
-				InBaseVSToolPath: BaseVSToolPath,
-				InWindowsSDKDir : WindowsSDKDir,
-				InUpdateVersion : UpdateVersion
-			);
-		}
-
-		/// <returns>The path to Windows SDK directory for the specified version.</returns>
-		private static string FindWindowsSDKInstallationFolder(string Version)
-		{
-			// Based on VCVarsQueryRegistry
-			var WinSDKPath =
-				   Microsoft.Win32.Registry.GetValue(@"HKEY_CURRENT_USER\SOFTWARE\Microsoft\Microsoft SDKs\Windows\"              + Version, "InstallationFolder", null)
-				?? Microsoft.Win32.Registry.GetValue(@"HKEY_LOCAL_MACHINE\SOFTWARE\Wow6432Node\Microsoft\Microsoft SDKs\Windows\" + Version, "InstallationFolder", null)
-				?? Microsoft.Win32.Registry.GetValue(@"HKEY_CURRENT_USER\SOFTWARE\Wow6432Node\Microsoft\Microsoft SDKs\Windows\"  + Version, "InstallationFolder", null);
-
-			if (WinSDKPath == null)
-			{
-				throw new BuildException("Windows SDK {0} must be installed in order to build this target.", Version);
-			}
-
-			return (string)WinSDKPath;
-		}
-
-		/// <returns>The installed version of VC.</returns>
-		private static Version FindVCUpdateVersion()
-		{
-			// Based on VCVarsQueryRegistry
-			var UpdateVersionString =
-				   Microsoft.Win32.Registry.GetValue(@"HKEY_LOCAL_MACHINE\SOFTWARE\Microsoft\DevDiv\vc\Servicing\12.0\CompilerCore",             "UpdateVersion", null)
-				?? Microsoft.Win32.Registry.GetValue(@"HKEY_LOCAL_MACHINE\SOFTWARE\Wow6432Node\Microsoft\DevDiv\vc\Servicing\12.0\CompilerCore", "UpdateVersion", null)
-				?? Microsoft.Win32.Registry.GetValue(@"HKEY_CURRENT_USER\SOFTWARE\Wow6432Node\Microsoft\DevDiv\vc\Servicing\12.0\CompilerCore",  "UpdateVersion", null);
-
-			Version Result;
-			if (UpdateVersionString == null || !Version.TryParse((string)UpdateVersionString, out Result))
-			{
-				throw new BuildException("Unable to determine Visual Studio version from registry");
-			}
-
-			return Result;
-		}
-
         public override void AddFilesToManifest(ref FileManifest Manifest, UEBuildBinary Binary)
         {
             // ok, this is pretty awful, we want the import libraries that go with the editor, only on the PC
