Temp workaround for a problem with parsing latest ModuleVersion.h in MacToolChain

diff --git a/Engine/Source/Programs/UnrealBuildTool/Mac/MacToolChain.cs b/Engine/Source/Programs/UnrealBuildTool/Mac/MacToolChain.cs
--- a/Engine/Source/Programs/UnrealBuildTool/Mac/MacToolChain.cs
+++ b/Engine/Source/Programs/UnrealBuildTool/Mac/MacToolChain.cs
@@ -463,23 +463,41 @@ namespace UnrealBuildTool
 			return LauncherVersionMajor + "." + LauncherVersionMinor + "." + LauncherVersionPatch;
 		}
 
+		private int LoadBuiltFromChangelistValue()
+		{
+			string[] VersionHeader = Utils.ReadAllText("../Source/Runtime/Launch/Resources/Version.h").Replace("\r\n", "\n").Replace("\t", " ").Split('\n');
+			foreach (string Line in VersionHeader)
+			{
+				if (Line.StartsWith("#define BUILT_FROM_CHANGELIST "))
+				{
+					return int.Parse(Line.Split(' ')[2]);
+				}
+			}
+			return 0;
+		}
+
 		private string LoadEngineAPIVersion()
 		{
 			int CL = 0;
-			foreach (string Line in File.ReadAllLines("../Source/Runtime/Core/Public/Modules/ModuleVersion.h"))
+			// @todo: Temp solution to work around a problem with parsing ModuleVersion.h updated for 4.4.1 hotfix
+			int BuiltFromChangelist = LoadBuiltFromChangelistValue();
+			if (BuiltFromChangelist > 0)
 			{
-				string[] Tokens = Line.Split(' ', '\t');
-				if (Tokens[0] == "#define" && Tokens[1] == "MODULE_API_VERSION")
+				foreach (string Line in File.ReadAllLines("../Source/Runtime/Core/Public/Modules/ModuleVersion.h"))
 				{
-					if(Tokens[2] == "BUILT_FROM_CHANGELIST")
+					string[] Tokens = Line.Split(' ', '\t');
+					if (Tokens[0] == "#define" && Tokens[1] == "MODULE_API_VERSION")
 					{
-						CL = LoadEngineCL();
-					}
-					else
-					{
-						CL = int.Parse(Tokens[2]);
+						if(Tokens[2] == "BUILT_FROM_CHANGELIST")
+						{
+							CL = LoadEngineCL();
+						}
+						else
+						{
+							CL = int.Parse(Tokens[2]);
+						}
+						break;
 					}
-					break;
 				}
 			}
 			return String.Format("{0}.{1}.{2}", CL / (100 * 100), (CL / 100) % 100, CL % 100);
