Build plugins in debug for DebugGame configurations. Engine always loads DebugGame versions of any modules under the game directory, but they are no longer set as UEBuildModuleType.Game in UBT.

diff --git a/Engine/Source/Programs/UnrealBuildTool/Configuration/UEBuildModule.cs b/Engine/Source/Programs/UnrealBuildTool/Configuration/UEBuildModule.cs
--- a/Engine/Source/Programs/UnrealBuildTool/Configuration/UEBuildModule.cs
+++ b/Engine/Source/Programs/UnrealBuildTool/Configuration/UEBuildModule.cs
@@ -1749,10 +1749,14 @@ namespace UnrealBuildTool
 			Result.Config.OutputDirectory                        = Path.Combine(Binary.Config.IntermediateDirectory, Name);
 
 			// Switch the optimization flag if we're building a game module. Also pass the definition for building in DebugGame along (see ModuleManager.h for notes).
-			if (Target.Configuration == UnrealTargetConfiguration.DebugGame && Type == UEBuildModuleType.Game)
+			if (Target.Configuration == UnrealTargetConfiguration.DebugGame)
 			{
-				Result.Config.Target.Configuration = CPPTargetConfiguration.Debug;
-				Result.Config.Definitions.Add("UE_BUILD_DEVELOPMENT_WITH_DEBUGGAME=1");
+				PluginInfo Plugin = Plugins.GetPluginInfoForModule(Name);
+				if((Plugin != null && Plugin.LoadedFrom == PluginInfo.LoadedFromType.GameProject) || Type == UEBuildModuleType.Game)
+				{
+					Result.Config.Target.Configuration = CPPTargetConfiguration.Debug;
+					Result.Config.Definitions.Add("UE_BUILD_DEVELOPMENT_WITH_DEBUGGAME=1");
+				}
 			}
 
 			// Add the module's private definitions.
