- Fixed a Mac UAT error (hopeful, checking in to move to mac)

diff --git a/Engine/Source/Programs/UnrealBuildTool/UnrealBuildTool_Mono.csproj b/Engine/Source/Programs/UnrealBuildTool/UnrealBuildTool_Mono.csproj
--- a/Engine/Source/Programs/UnrealBuildTool/UnrealBuildTool_Mono.csproj
+++ b/Engine/Source/Programs/UnrealBuildTool/UnrealBuildTool_Mono.csproj
@@ -120,6 +120,10 @@
     <Compile Include="Linux\LinuxProjectGenerator.cs" />
     <Compile Include="Linux\LinuxToolChain.cs" />
     <Compile Include="Linux\UEBuildLinux.cs" />
+    <Compile Include="Android\UEBuildAndroid.cs" />
+    <Compile Include="Android\UEDeployAndroid.cs" />
+    <Compile Include="Android\AndroidToolChain.cs" />
+    <Compile Include="Android\AndroidProjectGenerator.cs" />
     <Compile Include="Properties\AssemblyInfo.cs" />
     <Compile Include="Configuration\UEBuildConfiguration.cs" />
     <Compile Include="Configuration\UEBuildTarget.cs" />
