Fix AutoSDKs not properly using a fallback manual install when necessary.

diff --git a/Engine/Source/Programs/UnrealBuildTool/Configuration/UEBuildPlatform.cs b/Engine/Source/Programs/UnrealBuildTool/Configuration/UEBuildPlatform.cs
--- a/Engine/Source/Programs/UnrealBuildTool/Configuration/UEBuildPlatform.cs
+++ b/Engine/Source/Programs/UnrealBuildTool/Configuration/UEBuildPlatform.cs
@@ -8,1594 +8,1600 @@ using System.IO;
 
 namespace UnrealBuildTool
 {
-    public enum SDKStatus
-    {
-        Valid,			// Desired SDK is installed and set up.
-        Invalid,		// Could not find the desired SDK, SDK setup failed, etc.		
-    };
-
-    public interface IUEBuildPlatform
-    {
-        /**
-         * Whether the required external SDKs are installed for this platform
-         */
-        SDKStatus HasRequiredSDKsInstalled();
-
-        /**
-         * If this platform can be compiled with XGE
-         */
-        bool CanUseXGE();
-
-        /**
-         * If this platform can be compiled with DMUCS/Distcc
-         */
+	public enum SDKStatus
+	{
+		Valid,			// Desired SDK is installed and set up.
+		Invalid,		// Could not find the desired SDK, SDK setup failed, etc.		
+	};
+
+	public interface IUEBuildPlatform
+	{
+		/**
+		 * Whether the required external SDKs are installed for this platform
+		 */
+		SDKStatus HasRequiredSDKsInstalled();
+
+		/**
+		 * If this platform can be compiled with XGE
+		 */
+		bool CanUseXGE();
+
+		/**
+		 * If this platform can be compiled with DMUCS/Distcc
+		 */
 		bool CanUseDistcc();
 
-        /**
-         * Register the platform with the UEBuildPlatform class
-         */
-        void RegisterBuildPlatform();
-
-        /**
-         * Attempt to set up AutoSDK for this paltform
-         */
-        void ManageAndValidateSDK();
-
-        /**
-         * Retrieve the CPPTargetPlatform for the given UnrealTargetPlatform
-         *
-         * @param InUnrealTargetPlatform The UnrealTargetPlatform being build
-         * 
-         * @return CPPTargetPlatform The CPPTargetPlatform to compile for
-         */
-        CPPTargetPlatform GetCPPTargetPlatform(UnrealTargetPlatform InUnrealTargetPlatform);
-
-        /**
-         * Get the extension to use for the given binary type
-         * 
-         * @param InBinaryType The binary type being built
-         * 
-         * @return string The binary extension (i.e. 'exe' or 'dll')
-         */
-        string GetBinaryExtension(UEBuildBinaryType InBinaryType);
-
-        /**
-         * Get the extension to use for debug info for the given binary type
-         * 
-         * @param InBinaryType The binary type being built
-         * 
-         * @return string The debug info extension (i.e. 'pdb')
-         */
-        string GetDebugInfoExtension(UEBuildBinaryType InBinaryType);
-
-        /**
-         * Whether incremental linking should be used
-         * 
-         * @param InPlatform The CPPTargetPlatform being built
-         * @param InConfiguration The CPPTargetConfiguration being built
-         * 
-         * @return bool true if incremental linking should be used, false if not
-         */
-        bool ShouldUseIncrementalLinking(CPPTargetPlatform Platform, CPPTargetConfiguration Configuration);
-
-        /**
-         * Whether PDB files should be used
-         * 
-         * @param InPlatform The CPPTargetPlatform being built
-         * @param InConfiguration The CPPTargetConfiguration being built
-         * @param bInCreateDebugInfo true if debug info is getting create, false if not
-         * 
-         * @return bool true if PDB files should be used, false if not
-         */
-        bool ShouldUsePDBFiles(CPPTargetPlatform Platform, CPPTargetConfiguration Configuration, bool bCreateDebugInfo);
-
-        /**
-         * Whether PCH files should be used
-         * 
-         * @param InPlatform The CPPTargetPlatform being built
-         * @param InConfiguration The CPPTargetConfiguration being built
-         * 
-         * @return bool true if PCH files should be used, false if not
-         */
-        bool ShouldUsePCHFiles(CPPTargetPlatform Platform, CPPTargetConfiguration Configuration);
-
-        /**
-         * Whether the editor should be built for this platform or not
-         * 
-         * @param InPlatform The UnrealTargetPlatform being built
-         * @param InConfiguration The UnrealTargetConfiguration being built
-         * @return bool true if the editor should be built, false if not
-         */
-        bool ShouldNotBuildEditor(UnrealTargetPlatform InPlatform, UnrealTargetConfiguration InConfiguration);
-
-        /**
-         * Whether this build should support ONLY cooked data or not
-         * 
-         * @param InPlatform The UnrealTargetPlatform being built
-         * @param InConfiguration The UnrealTargetConfiguration being built
-         * @return bool true if the editor should be built, false if not
-         */
-        bool BuildRequiresCookedData(UnrealTargetPlatform InPlatform, UnrealTargetConfiguration InConfiguration);
-
-        /**
-         * Whether the platform requires the extra UnityCPPWriter
-         * This is used to add an extra file for UBT to get the #include dependencies from
-         * 
-         * @return bool true if it is required, false if not
-         */
-        bool RequiresExtraUnityCPPWriter();
-
-        /**
-         * Whether this platform should build a monolithic binary
-         */
-        bool ShouldCompileMonolithicBinary(UnrealTargetPlatform InPlatform);
-
-        /**
-         * Get a list of extra modules the platform requires.
-         * This is to allow undisclosed platforms to add modules they need without exposing information about the platform.
-         * 
-         * @param Target The target being build
-         * @param BuildTarget The UEBuildTarget getting build
-         * @param PlatformExtraModules OUTPUT the list of extra modules the platform needs to add to the target
-         */
-        void GetExtraModules(TargetInfo Target, UEBuildTarget BuildTarget, ref List<string> PlatformExtraModules);
-
-        /**
-         * Modify the newly created module passed in for this platform.
-         * This is not required - but allows for hiding details of a
-         * particular platform.
-         * 
-         * @param InModule The newly loaded module
-         * @param Target The target being build
-         */
-        void ModifyNewlyLoadedModule(UEBuildModule InModule, TargetInfo Target);
-
-        /**
-         * Setup the target environment for building
-         * 
-         * @param InBuildTarget The target being built
-         */
-        void SetUpEnvironment(UEBuildTarget InBuildTarget);
-
-        /**
-         * Allow the platform to set an optional architecture
-         */
-        string GetActiveArchitecture();
-
-        /**
-         * Allow the platform to apply architecture-specific name according to its rules
-         * 
-         * @param BinaryName name of the binary, not specific to any architecture
-         *
-         * @return Possibly architecture-specific name
-         */
-        string ApplyArchitectureName(string BinaryName);
-
-        /**
-         * Setup the configuration environment for building
-         * 
-         * @param InBuildTarget The target being built
-         */
-        void SetUpConfigurationEnvironment(UEBuildTarget InBuildTarget);
-
-        /**
-         * Whether this platform should create debug information or not
-         * 
-         * @param InPlatform The UnrealTargetPlatform being built
-         * @param InConfiguration The UnrealTargetConfiguration being built
-         * 
-         * @return bool true if debug info should be generated, false if not
-         */
-        bool ShouldCreateDebugInfo(UnrealTargetPlatform Platform, UnrealTargetConfiguration Configuration);
-
-        /**
-         * Gives the platform a chance to 'override' the configuration settings 
-         * that are overridden on calls to RunUBT.
-         * 
-         * @param InPlatform The UnrealTargetPlatform being built
-         * @param InConfiguration The UnrealTargetConfiguration being built
-         * 
-         * @return bool true if debug info should be generated, false if not
-         */
-        void ResetBuildConfiguration(UnrealTargetPlatform InPlatform, UnrealTargetConfiguration InConfiguration);
-
-        /**
-         * Validate the UEBuildConfiguration for this platform
-         * This is called BEFORE calling UEBuildConfiguration to allow setting 
-         * various fields used in that function such as CompileLeanAndMean...
-         */
-        void ValidateUEBuildConfiguration();
-
-        /**
-         * Validate configuration for this platform
-         * NOTE: This function can/will modify BuildConfiguration!
-         * 
-         * @param InPlatform The CPPTargetPlatform being built
-         * @param InConfiguration The CPPTargetConfiguration being built
-         * @param bInCreateDebugInfo true if debug info is getting create, false if not
-         */
-        void ValidateBuildConfiguration(CPPTargetConfiguration Configuration, CPPTargetPlatform Platform, bool bCreateDebugInfo);
-
-        /**
-         * Return whether this platform has uniquely named binaries across multiple games
-         */
-        bool HasUniqueBinaries();
-
-        /**
-         * Return whether we wish to have this platform's binaries in our builds
-         */
-        bool IsBuildRequired();
-
-        /**
-         * Return whether we wish to have this platform's binaries in our CIS tests
-         */
-        bool IsCISRequired();
-
-        /// <summary>
-        /// Whether the build platform requires deployment prep
-        /// </summary>
-        /// <returns></returns>
-        bool RequiresDeployPrepAfterCompile();
-
-        /**
-         * Return all valid configurations for this platform
-         * 
-         *  Typically, this is always Debug, Development, and Shipping - but Test is a likely future addition for some platforms
-         */
-        List<UnrealTargetConfiguration> GetConfigurations(UnrealTargetPlatform InUnrealTargetPlatform, bool bIncludeDebug);
-
-        /**
-         * Setup the binaries for this specific platform.
-         * 
-         * @param InBuildTarget The target being built
-         */
-        void SetupBinaries(UEBuildTarget InBuildTarget);
-    }
-
-    public abstract partial class UEBuildPlatform : IUEBuildPlatform
-    {
-        public static Dictionary<UnrealTargetPlatform, IUEBuildPlatform> BuildPlatformDictionary = new Dictionary<UnrealTargetPlatform, IUEBuildPlatform>();
-
-        // a mapping of a group to the platforms in the group (ie, Microsoft contains Win32 and Win64)
-        static Dictionary<UnrealPlatformGroup, List<UnrealTargetPlatform>> PlatformGroupDictionary = new Dictionary<UnrealPlatformGroup, List<UnrealTargetPlatform>>();
-
-        /**
-         * Attempt to convert a string to an UnrealTargetPlatform enum entry
-         * 
-         * @return UnrealTargetPlatform.Unknown on failure (the platform didn't match the enum)
-         */
-        public static UnrealTargetPlatform ConvertStringToPlatform(string InPlatformName)
-        {
-            // special case x64, to not break anything
-            // @todo: Is it possible to remove this hack?
-            if (InPlatformName.Equals("X64", StringComparison.InvariantCultureIgnoreCase))
-            {
-                return UnrealTargetPlatform.Win64;
-            }
-
-            // we can't parse the string into an enum because Enum.Parse is case sensitive, so we loop over the enum
-            // looking for matches
-            foreach (string PlatformName in Enum.GetNames(typeof(UnrealTargetPlatform)))
-            {
-                if (InPlatformName.Equals(PlatformName, StringComparison.InvariantCultureIgnoreCase))
-                {
-                    // convert the known good enum string back to the enum value
-                    return (UnrealTargetPlatform)Enum.Parse(typeof(UnrealTargetPlatform), PlatformName);
-                }
-            }
-            return UnrealTargetPlatform.Unknown;
-        }
-
-        /**
-         *	Register the given platforms UEBuildPlatform instance
-         *	
-         *	@param	InPlatform			The UnrealTargetPlatform to register with
-         *	@param	InBuildPlatform		The UEBuildPlatform instance to use for the InPlatform
-         */
-        public static void RegisterBuildPlatform(UnrealTargetPlatform InPlatform, UEBuildPlatform InBuildPlatform)
-        {
-            if (BuildPlatformDictionary.ContainsKey(InPlatform) == true)
-            {
-                Log.TraceWarning("RegisterBuildPlatform Warning: Registering build platform {0} for {1} when it is already set to {2}",
-                    InBuildPlatform.ToString(), InPlatform.ToString(), BuildPlatformDictionary[InPlatform].ToString());
-                BuildPlatformDictionary[InPlatform] = InBuildPlatform;
-            }
-            else
-            {
-                BuildPlatformDictionary.Add(InPlatform, InBuildPlatform);
-            }
-        }
-
-        /**
-         *	Unregister the given platform
-         */
-        public static void UnregisterBuildPlatform(UnrealTargetPlatform InPlatform)
-        {
-            if (BuildPlatformDictionary.ContainsKey(InPlatform) == true)
-            {
-                BuildPlatformDictionary.Remove(InPlatform);
-            }
-        }
-
-        /**
-         * Assign a platform as a member of the given group
-         */
-        public static void RegisterPlatformWithGroup(UnrealTargetPlatform InPlatform, UnrealPlatformGroup InGroup)
-        {
-            // find or add the list of groups for this platform
-            PlatformGroupDictionary.GetOrAddNew(InGroup).Add(InPlatform);
-        }
-
-        /**
-         * Retrieve the list of platforms in this group (if any)
-         */
-        public static List<UnrealTargetPlatform> GetPlatformsInGroup(UnrealPlatformGroup InGroup)
-        {
-            List<UnrealTargetPlatform> PlatformList;
-            PlatformGroupDictionary.TryGetValue(InGroup, out PlatformList);
-            return PlatformList;
-        }
-
-        /**
-         *	Retrieve the IUEBuildPlatform instance for the given TargetPlatform
-         *	
-         *	@param	InPlatform			The UnrealTargetPlatform being built
-         *	@param	bInAllowFailure		If true, do not throw an exception and return null
-         *	
-         *	@return	UEBuildPlatform		The instance of the build platform
-         */
-        public static IUEBuildPlatform GetBuildPlatform(UnrealTargetPlatform InPlatform, bool bInAllowFailure = false)
-        {
-            if (BuildPlatformDictionary.ContainsKey(InPlatform) == true)
-            {
-                return BuildPlatformDictionary[InPlatform];
-            }
-            if (bInAllowFailure == true)
-            {
-                return null;
-            }
-            throw new BuildException("GetBuildPlatform: No BuildPlatform found for {0}", InPlatform.ToString());
-        }
-
-        /**
-         *	Retrieve the IUEBuildPlatform instance for the given CPPTargetPlatform
-         *	
-         *	@param	InPlatform			The CPPTargetPlatform being built
-         *	@param	bInAllowFailure		If true, do not throw an exception and return null
-         *	
-         *	@return	UEBuildPlatform		The instance of the build platform
-         */
-        public static IUEBuildPlatform GetBuildPlatformForCPPTargetPlatform(CPPTargetPlatform InPlatform, bool bInAllowFailure = false)
-        {
-            UnrealTargetPlatform UTPlatform = UEBuildTarget.CPPTargetPlatformToUnrealTargetPlatform(InPlatform);
-            if (BuildPlatformDictionary.ContainsKey(UTPlatform) == true)
-            {
-                return BuildPlatformDictionary[UTPlatform];
-            }
-            if (bInAllowFailure == true)
-            {
-                return null;
-            }
-            throw new BuildException("UEBuildPlatform::GetBuildPlatformForCPPTargetPlatform: No BuildPlatform found for {0}", InPlatform.ToString());
-        }
-
-        /**
-         *	Allow all registered build platforms to modify the newly created module 
-         *	passed in for the given platform.
-         *	This is not required - but allows for hiding details of a particular platform.
-         *	
-         *	@param	InModule		The newly loaded module
-         *	@param	Target			The target being build
-         *	@param	Only			If this is not unknown, then only run that platform
-         */
-        public static void PlatformModifyNewlyLoadedModule(UEBuildModule InModule, TargetInfo Target, UnrealTargetPlatform Only = UnrealTargetPlatform.Unknown)
-        {
-            foreach (var PlatformEntry in BuildPlatformDictionary)
-            {
-                if (Only == UnrealTargetPlatform.Unknown || PlatformEntry.Key == Only || PlatformEntry.Key == Target.Platform)
-                {
-                    PlatformEntry.Value.ModifyNewlyLoadedModule(InModule, Target);
-                }
-            }
-        }
-
-        /**
-         * Returns the delimiter used to separate paths in the PATH environment variable for the platform we are executing on.
-         */
-        public String GetPathVarDelimiter()
-        {
-            switch (BuildHostPlatform.Current.Platform)
-            {
-                case UnrealTargetPlatform.Linux:
-                case UnrealTargetPlatform.Mac:
-                    return ":";
-                case UnrealTargetPlatform.Win32:
-                case UnrealTargetPlatform.Win64:
-                    return ";";
-                default:
-                    Log.TraceWarning("PATH var delimiter unknown for platform " + BuildHostPlatform.Current.Platform.ToString() + " using ';'");
-                    return ";";
-            }
-        }
-
-
-
-        /**
-         *	If this platform can be compiled with XGE
-         */
-        public virtual bool CanUseXGE()
-        {
-            return true;
-        }
-
-		/**
-         *	If this platform can be compiled with DMUCS/Distcc
-         */
+		/**
+		 * Register the platform with the UEBuildPlatform class
+		 */
+		void RegisterBuildPlatform();
+
+		/**
+		 * Attempt to set up AutoSDK for this paltform
+		 */
+		void ManageAndValidateSDK();
+
+		/**
+		 * Retrieve the CPPTargetPlatform for the given UnrealTargetPlatform
+		 *
+		 * @param InUnrealTargetPlatform The UnrealTargetPlatform being build
+		 * 
+		 * @return CPPTargetPlatform The CPPTargetPlatform to compile for
+		 */
+		CPPTargetPlatform GetCPPTargetPlatform(UnrealTargetPlatform InUnrealTargetPlatform);
+
+		/**
+		 * Get the extension to use for the given binary type
+		 * 
+		 * @param InBinaryType The binary type being built
+		 * 
+		 * @return string The binary extension (i.e. 'exe' or 'dll')
+		 */
+		string GetBinaryExtension(UEBuildBinaryType InBinaryType);
+
+		/**
+		 * Get the extension to use for debug info for the given binary type
+		 * 
+		 * @param InBinaryType The binary type being built
+		 * 
+		 * @return string The debug info extension (i.e. 'pdb')
+		 */
+		string GetDebugInfoExtension(UEBuildBinaryType InBinaryType);
+
+		/**
+		 * Whether incremental linking should be used
+		 * 
+		 * @param InPlatform The CPPTargetPlatform being built
+		 * @param InConfiguration The CPPTargetConfiguration being built
+		 * 
+		 * @return bool true if incremental linking should be used, false if not
+		 */
+		bool ShouldUseIncrementalLinking(CPPTargetPlatform Platform, CPPTargetConfiguration Configuration);
+
+		/**
+		 * Whether PDB files should be used
+		 * 
+		 * @param InPlatform The CPPTargetPlatform being built
+		 * @param InConfiguration The CPPTargetConfiguration being built
+		 * @param bInCreateDebugInfo true if debug info is getting create, false if not
+		 * 
+		 * @return bool true if PDB files should be used, false if not
+		 */
+		bool ShouldUsePDBFiles(CPPTargetPlatform Platform, CPPTargetConfiguration Configuration, bool bCreateDebugInfo);
+
+		/**
+		 * Whether PCH files should be used
+		 * 
+		 * @param InPlatform The CPPTargetPlatform being built
+		 * @param InConfiguration The CPPTargetConfiguration being built
+		 * 
+		 * @return bool true if PCH files should be used, false if not
+		 */
+		bool ShouldUsePCHFiles(CPPTargetPlatform Platform, CPPTargetConfiguration Configuration);
+
+		/**
+		 * Whether the editor should be built for this platform or not
+		 * 
+		 * @param InPlatform The UnrealTargetPlatform being built
+		 * @param InConfiguration The UnrealTargetConfiguration being built
+		 * @return bool true if the editor should be built, false if not
+		 */
+		bool ShouldNotBuildEditor(UnrealTargetPlatform InPlatform, UnrealTargetConfiguration InConfiguration);
+
+		/**
+		 * Whether this build should support ONLY cooked data or not
+		 * 
+		 * @param InPlatform The UnrealTargetPlatform being built
+		 * @param InConfiguration The UnrealTargetConfiguration being built
+		 * @return bool true if the editor should be built, false if not
+		 */
+		bool BuildRequiresCookedData(UnrealTargetPlatform InPlatform, UnrealTargetConfiguration InConfiguration);
+
+		/**
+		 * Whether the platform requires the extra UnityCPPWriter
+		 * This is used to add an extra file for UBT to get the #include dependencies from
+		 * 
+		 * @return bool true if it is required, false if not
+		 */
+		bool RequiresExtraUnityCPPWriter();
+
+		/**
+		 * Whether this platform should build a monolithic binary
+		 */
+		bool ShouldCompileMonolithicBinary(UnrealTargetPlatform InPlatform);
+
+		/**
+		 * Get a list of extra modules the platform requires.
+		 * This is to allow undisclosed platforms to add modules they need without exposing information about the platform.
+		 * 
+		 * @param Target The target being build
+		 * @param BuildTarget The UEBuildTarget getting build
+		 * @param PlatformExtraModules OUTPUT the list of extra modules the platform needs to add to the target
+		 */
+		void GetExtraModules(TargetInfo Target, UEBuildTarget BuildTarget, ref List<string> PlatformExtraModules);
+
+		/**
+		 * Modify the newly created module passed in for this platform.
+		 * This is not required - but allows for hiding details of a
+		 * particular platform.
+		 * 
+		 * @param InModule The newly loaded module
+		 * @param Target The target being build
+		 */
+		void ModifyNewlyLoadedModule(UEBuildModule InModule, TargetInfo Target);
+
+		/**
+		 * Setup the target environment for building
+		 * 
+		 * @param InBuildTarget The target being built
+		 */
+		void SetUpEnvironment(UEBuildTarget InBuildTarget);
+
+		/**
+		 * Allow the platform to set an optional architecture
+		 */
+		string GetActiveArchitecture();
+
+		/**
+		 * Allow the platform to apply architecture-specific name according to its rules
+		 * 
+		 * @param BinaryName name of the binary, not specific to any architecture
+		 *
+		 * @return Possibly architecture-specific name
+		 */
+		string ApplyArchitectureName(string BinaryName);
+
+		/**
+		 * Setup the configuration environment for building
+		 * 
+		 * @param InBuildTarget The target being built
+		 */
+		void SetUpConfigurationEnvironment(UEBuildTarget InBuildTarget);
+
+		/**
+		 * Whether this platform should create debug information or not
+		 * 
+		 * @param InPlatform The UnrealTargetPlatform being built
+		 * @param InConfiguration The UnrealTargetConfiguration being built
+		 * 
+		 * @return bool true if debug info should be generated, false if not
+		 */
+		bool ShouldCreateDebugInfo(UnrealTargetPlatform Platform, UnrealTargetConfiguration Configuration);
+
+		/**
+		 * Gives the platform a chance to 'override' the configuration settings 
+		 * that are overridden on calls to RunUBT.
+		 * 
+		 * @param InPlatform The UnrealTargetPlatform being built
+		 * @param InConfiguration The UnrealTargetConfiguration being built
+		 * 
+		 * @return bool true if debug info should be generated, false if not
+		 */
+		void ResetBuildConfiguration(UnrealTargetPlatform InPlatform, UnrealTargetConfiguration InConfiguration);
+
+		/**
+		 * Validate the UEBuildConfiguration for this platform
+		 * This is called BEFORE calling UEBuildConfiguration to allow setting 
+		 * various fields used in that function such as CompileLeanAndMean...
+		 */
+		void ValidateUEBuildConfiguration();
+
+		/**
+		 * Validate configuration for this platform
+		 * NOTE: This function can/will modify BuildConfiguration!
+		 * 
+		 * @param InPlatform The CPPTargetPlatform being built
+		 * @param InConfiguration The CPPTargetConfiguration being built
+		 * @param bInCreateDebugInfo true if debug info is getting create, false if not
+		 */
+		void ValidateBuildConfiguration(CPPTargetConfiguration Configuration, CPPTargetPlatform Platform, bool bCreateDebugInfo);
+
+		/**
+		 * Return whether this platform has uniquely named binaries across multiple games
+		 */
+		bool HasUniqueBinaries();
+
+		/**
+		 * Return whether we wish to have this platform's binaries in our builds
+		 */
+		bool IsBuildRequired();
+
+		/**
+		 * Return whether we wish to have this platform's binaries in our CIS tests
+		 */
+		bool IsCISRequired();
+
+		/// <summary>
+		/// Whether the build platform requires deployment prep
+		/// </summary>
+		/// <returns></returns>
+		bool RequiresDeployPrepAfterCompile();
+
+		/**
+		 * Return all valid configurations for this platform
+		 * 
+		 *  Typically, this is always Debug, Development, and Shipping - but Test is a likely future addition for some platforms
+		 */
+		List<UnrealTargetConfiguration> GetConfigurations(UnrealTargetPlatform InUnrealTargetPlatform, bool bIncludeDebug);
+
+		/**
+		 * Setup the binaries for this specific platform.
+		 * 
+		 * @param InBuildTarget The target being built
+		 */
+		void SetupBinaries(UEBuildTarget InBuildTarget);
+	}
+
+	public abstract partial class UEBuildPlatform : IUEBuildPlatform
+	{
+		public static Dictionary<UnrealTargetPlatform, IUEBuildPlatform> BuildPlatformDictionary = new Dictionary<UnrealTargetPlatform, IUEBuildPlatform>();
+
+		// a mapping of a group to the platforms in the group (ie, Microsoft contains Win32 and Win64)
+		static Dictionary<UnrealPlatformGroup, List<UnrealTargetPlatform>> PlatformGroupDictionary = new Dictionary<UnrealPlatformGroup, List<UnrealTargetPlatform>>();
+
+		/**
+		 * Attempt to convert a string to an UnrealTargetPlatform enum entry
+		 * 
+		 * @return UnrealTargetPlatform.Unknown on failure (the platform didn't match the enum)
+		 */
+		public static UnrealTargetPlatform ConvertStringToPlatform(string InPlatformName)
+		{
+			// special case x64, to not break anything
+			// @todo: Is it possible to remove this hack?
+			if (InPlatformName.Equals("X64", StringComparison.InvariantCultureIgnoreCase))
+			{
+				return UnrealTargetPlatform.Win64;
+			}
+
+			// we can't parse the string into an enum because Enum.Parse is case sensitive, so we loop over the enum
+			// looking for matches
+			foreach (string PlatformName in Enum.GetNames(typeof(UnrealTargetPlatform)))
+			{
+				if (InPlatformName.Equals(PlatformName, StringComparison.InvariantCultureIgnoreCase))
+				{
+					// convert the known good enum string back to the enum value
+					return (UnrealTargetPlatform)Enum.Parse(typeof(UnrealTargetPlatform), PlatformName);
+				}
+			}
+			return UnrealTargetPlatform.Unknown;
+		}
+
+		/**
+		 *	Register the given platforms UEBuildPlatform instance
+		 *	
+		 *	@param	InPlatform			The UnrealTargetPlatform to register with
+		 *	@param	InBuildPlatform		The UEBuildPlatform instance to use for the InPlatform
+		 */
+		public static void RegisterBuildPlatform(UnrealTargetPlatform InPlatform, UEBuildPlatform InBuildPlatform)
+		{
+			if (BuildPlatformDictionary.ContainsKey(InPlatform) == true)
+			{
+				Log.TraceWarning("RegisterBuildPlatform Warning: Registering build platform {0} for {1} when it is already set to {2}",
+					InBuildPlatform.ToString(), InPlatform.ToString(), BuildPlatformDictionary[InPlatform].ToString());
+				BuildPlatformDictionary[InPlatform] = InBuildPlatform;
+			}
+			else
+			{
+				BuildPlatformDictionary.Add(InPlatform, InBuildPlatform);
+			}
+		}
+
+		/**
+		 *	Unregister the given platform
+		 */
+		public static void UnregisterBuildPlatform(UnrealTargetPlatform InPlatform)
+		{
+			if (BuildPlatformDictionary.ContainsKey(InPlatform) == true)
+			{
+				BuildPlatformDictionary.Remove(InPlatform);
+			}
+		}
+
+		/**
+		 * Assign a platform as a member of the given group
+		 */
+		public static void RegisterPlatformWithGroup(UnrealTargetPlatform InPlatform, UnrealPlatformGroup InGroup)
+		{
+			// find or add the list of groups for this platform
+			PlatformGroupDictionary.GetOrAddNew(InGroup).Add(InPlatform);
+		}
+
+		/**
+		 * Retrieve the list of platforms in this group (if any)
+		 */
+		public static List<UnrealTargetPlatform> GetPlatformsInGroup(UnrealPlatformGroup InGroup)
+		{
+			List<UnrealTargetPlatform> PlatformList;
+			PlatformGroupDictionary.TryGetValue(InGroup, out PlatformList);
+			return PlatformList;
+		}
+
+		/**
+		 *	Retrieve the IUEBuildPlatform instance for the given TargetPlatform
+		 *	
+		 *	@param	InPlatform			The UnrealTargetPlatform being built
+		 *	@param	bInAllowFailure		If true, do not throw an exception and return null
+		 *	
+		 *	@return	UEBuildPlatform		The instance of the build platform
+		 */
+		public static IUEBuildPlatform GetBuildPlatform(UnrealTargetPlatform InPlatform, bool bInAllowFailure = false)
+		{
+			if (BuildPlatformDictionary.ContainsKey(InPlatform) == true)
+			{
+				return BuildPlatformDictionary[InPlatform];
+			}
+			if (bInAllowFailure == true)
+			{
+				return null;
+			}
+			throw new BuildException("GetBuildPlatform: No BuildPlatform found for {0}", InPlatform.ToString());
+		}
+
+		/**
+		 *	Retrieve the IUEBuildPlatform instance for the given CPPTargetPlatform
+		 *	
+		 *	@param	InPlatform			The CPPTargetPlatform being built
+		 *	@param	bInAllowFailure		If true, do not throw an exception and return null
+		 *	
+		 *	@return	UEBuildPlatform		The instance of the build platform
+		 */
+		public static IUEBuildPlatform GetBuildPlatformForCPPTargetPlatform(CPPTargetPlatform InPlatform, bool bInAllowFailure = false)
+		{
+			UnrealTargetPlatform UTPlatform = UEBuildTarget.CPPTargetPlatformToUnrealTargetPlatform(InPlatform);
+			if (BuildPlatformDictionary.ContainsKey(UTPlatform) == true)
+			{
+				return BuildPlatformDictionary[UTPlatform];
+			}
+			if (bInAllowFailure == true)
+			{
+				return null;
+			}
+			throw new BuildException("UEBuildPlatform::GetBuildPlatformForCPPTargetPlatform: No BuildPlatform found for {0}", InPlatform.ToString());
+		}
+
+		/**
+		 *	Allow all registered build platforms to modify the newly created module 
+		 *	passed in for the given platform.
+		 *	This is not required - but allows for hiding details of a particular platform.
+		 *	
+		 *	@param	InModule		The newly loaded module
+		 *	@param	Target			The target being build
+		 *	@param	Only			If this is not unknown, then only run that platform
+		 */
+		public static void PlatformModifyNewlyLoadedModule(UEBuildModule InModule, TargetInfo Target, UnrealTargetPlatform Only = UnrealTargetPlatform.Unknown)
+		{
+			foreach (var PlatformEntry in BuildPlatformDictionary)
+			{
+				if (Only == UnrealTargetPlatform.Unknown || PlatformEntry.Key == Only || PlatformEntry.Key == Target.Platform)
+				{
+					PlatformEntry.Value.ModifyNewlyLoadedModule(InModule, Target);
+				}
+			}
+		}
+
+		/**
+		 * Returns the delimiter used to separate paths in the PATH environment variable for the platform we are executing on.
+		 */
+		public String GetPathVarDelimiter()
+		{
+			switch (BuildHostPlatform.Current.Platform)
+			{
+				case UnrealTargetPlatform.Linux:
+				case UnrealTargetPlatform.Mac:
+					return ":";
+				case UnrealTargetPlatform.Win32:
+				case UnrealTargetPlatform.Win64:
+					return ";";
+				default:
+					Log.TraceWarning("PATH var delimiter unknown for platform " + BuildHostPlatform.Current.Platform.ToString() + " using ';'");
+					return ";";
+			}
+		}
+
+
+
+		/**
+		 *	If this platform can be compiled with XGE
+		 */
+		public virtual bool CanUseXGE()
+		{
+			return true;
+		}
+
+		/**
+		 *	If this platform can be compiled with DMUCS/Distcc
+		 */
 		public virtual bool CanUseDistcc()
 		{
 			return false;
 		}
 
-        /**
-         *	Register the platform with the UEBuildPlatform class
-         */
-        public void RegisterBuildPlatform()
-        {
-            ManageAndValidateSDK();
-            RegisterBuildPlatformInternal();
-        }
-
-        protected abstract void RegisterBuildPlatformInternal();
-
-        /**
-         *	Retrieve the CPPTargetPlatform for the given UnrealTargetPlatform
-         *
-         *	@param	InUnrealTargetPlatform		The UnrealTargetPlatform being build
-         *	
-         *	@return	CPPTargetPlatform			The CPPTargetPlatform to compile for
-         */
-        public abstract CPPTargetPlatform GetCPPTargetPlatform(UnrealTargetPlatform InUnrealTargetPlatform);
-
-        /// <summary>
-        /// Return whether the given platform requires a monolithic build
-        /// </summary>
-        /// <param name="InPlatform">The platform of interest</param>
-        /// <param name="InConfiguration">The configuration of interest</param>
-        /// <returns></returns>
-        public static bool PlatformRequiresMonolithicBuilds(UnrealTargetPlatform InPlatform, UnrealTargetConfiguration InConfiguration)
-        {
-            // Some platforms require monolithic builds...
-            var BuildPlatform = GetBuildPlatform(InPlatform, true);
-            if (BuildPlatform != null)
-            {
-                return BuildPlatform.ShouldCompileMonolithicBinary(InPlatform);
-            }
-
-            // We assume it does not
-            return false;
-        }
-
-        /**
-         *	Get the extension to use for the given binary type
-         *	
-         *	@param	InBinaryType		The binary type being built
-         *	
-         *	@return	string				The binary extension (i.e. 'exe' or 'dll')
-         */
-        public virtual string GetBinaryExtension(UEBuildBinaryType InBinaryType)
-        {
-            throw new BuildException("GetBinaryExtensiton for {0} not handled in {1}", InBinaryType.ToString(), this.ToString());
-        }
-
-        /**
-         *	Get the extension to use for debug info for the given binary type
-         *	
-         *	@param	InBinaryType		The binary type being built
-         *	
-         *	@return	string				The debug info extension (i.e. 'pdb')
-         */
-        public virtual string GetDebugInfoExtension(UEBuildBinaryType InBinaryType)
-        {
-            throw new BuildException("GetDebugInfoExtension for {0} not handled in {1}", InBinaryType.ToString(), this.ToString());
-        }
-
-        /**
-         *	Whether incremental linking should be used
-         *	
-         *	@param	InPlatform			The CPPTargetPlatform being built
-         *	@param	InConfiguration		The CPPTargetConfiguration being built
-         *	
-         *	@return	bool	true if incremental linking should be used, false if not
-         */
-        public virtual bool ShouldUseIncrementalLinking(CPPTargetPlatform Platform, CPPTargetConfiguration Configuration)
-        {
-            return false;
-        }
-
-        /**
-         *	Whether PDB files should be used
-         *	
-         *	@param	InPlatform			The CPPTargetPlatform being built
-         *	@param	InConfiguration		The CPPTargetConfiguration being built
-         *	@param	bInCreateDebugInfo	true if debug info is getting create, false if not
-         *	
-         *	@return	bool	true if PDB files should be used, false if not
-         */
-        public virtual bool ShouldUsePDBFiles(CPPTargetPlatform Platform, CPPTargetConfiguration Configuration, bool bCreateDebugInfo)
-        {
-            return false;
-        }
-
-        /**
-         *	Whether PCH files should be used
-         *	
-         *	@param	InPlatform			The CPPTargetPlatform being built
-         *	@param	InConfiguration		The CPPTargetConfiguration being built
-         *	
-         *	@return	bool				true if PCH files should be used, false if not
-         */
-        public virtual bool ShouldUsePCHFiles(CPPTargetPlatform Platform, CPPTargetConfiguration Configuration)
-        {
-            return BuildConfiguration.bUsePCHFiles;
-        }
-
-        /**
-         *	Whether the editor should be built for this platform or not
-         *	
-         *	@param	InPlatform		The UnrealTargetPlatform being built
-         *	@param	InConfiguration	The UnrealTargetConfiguration being built
-         *	@return	bool			true if the editor should be built, false if not
-         */
-        public virtual bool ShouldNotBuildEditor(UnrealTargetPlatform InPlatform, UnrealTargetConfiguration InConfiguration)
-        {
-            return false;
-        }
-
-        /**
-         *	Whether this build should support ONLY cooked data or not
-         *	
-         *	@param	InPlatform		The UnrealTargetPlatform being built
-         *	@param	InConfiguration	The UnrealTargetConfiguration being built
-         *	@return	bool			true if the editor should be built, false if not
-         */
-        public virtual bool BuildRequiresCookedData(UnrealTargetPlatform InPlatform, UnrealTargetConfiguration InConfiguration)
-        {
-            return false;
-        }
-
-        /**
-         *	Whether the platform requires the extra UnityCPPWriter
-         *	This is used to add an extra file for UBT to get the #include dependencies from
-         *	
-         *	@return	bool	true if it is required, false if not
-         */
-        public virtual bool RequiresExtraUnityCPPWriter()
-        {
-            return false;
-        }
-
-        /**
-         * Whether this platform should build a monolithic binary
-         */
-        public virtual bool ShouldCompileMonolithicBinary(UnrealTargetPlatform InPlatform)
-        {
-            return false;
-        }
-
-        /**
-         *	Get a list of extra modules the platform requires.
-         *	This is to allow undisclosed platforms to add modules they need without exposing information about the platform.
-         *	
-         *	@param	Target						The target being build
-         *	@param	BuildTarget					The UEBuildTarget getting build
-         *	@param	PlatformExtraModules		OUTPUT the list of extra modules the platform needs to add to the target
-         */
-        public virtual void GetExtraModules(TargetInfo Target, UEBuildTarget BuildTarget, ref List<string> PlatformExtraModules)
-        {
-        }
-
-        /**
-         *	Modify the newly created module passed in for this platform.
-         *	This is not required - but allows for hiding details of a
-         *	particular platform.
-         *	
-         *	@param	InModule		The newly loaded module
-         *	@param	Target			The target being build
-         */
-        public virtual void ModifyNewlyLoadedModule(UEBuildModule InModule, TargetInfo Target)
-        {
-        }
-
-        /**
-         *	Setup the target environment for building
-         *	
-         *	@param	InBuildTarget		The target being built
-         */
-        public abstract void SetUpEnvironment(UEBuildTarget InBuildTarget);
-
-        /**
-         * Allow the platform to set an optional architecture
-         */
-        public virtual string GetActiveArchitecture()
-        {
-            // by default, use an empty architecture (which is really just a modifer to the platform for some paths/names)
-            return "";
-        }
-
-        /**
-         * Allow the platform to override the NMake output name
-         */
-        public virtual string ModifyNMakeOutput(string ExeName)
-        {
-            // by default, use original
-            return ExeName;
-        }
-
-        /**
-         * Allow the platform to apply architecture-specific name according to its rules
-         * 
-         * @param BinaryName name of the binary, not specific to any architecture
-         *
-         * @return Possibly architecture-specific name
-         */
-        public virtual string ApplyArchitectureName(string BinaryName)
-        {
-            // by default, use logic that assumes architectures to start with "-"
-            return BinaryName + GetActiveArchitecture();
-        }
-
-        /**
-         *	Setup the configuration environment for building
-         *	
-         *	@param	InBuildTarget		The target being built
-         */
-        public virtual void SetUpConfigurationEnvironment(UEBuildTarget InBuildTarget)
-        {
-            // Determine the C++ compile/link configuration based on the Unreal configuration.
-            CPPTargetConfiguration CompileConfiguration;
-            UnrealTargetConfiguration CheckConfig = InBuildTarget.Configuration;
-            //@todo SAS: Add a true Debug mode!
-            if (UnrealBuildTool.RunningRocket())
-            {
-                if (Utils.IsFileUnderDirectory(InBuildTarget.OutputPath, UnrealBuildTool.GetUProjectPath()))
-                {
-                    if (CheckConfig == UnrealTargetConfiguration.Debug)
-                    {
-                        CheckConfig = UnrealTargetConfiguration.DebugGame;
-                    }
-                }
-                else
-                {
-                    // Only Development and Shipping are supported for engine modules
-                    if (CheckConfig != UnrealTargetConfiguration.Development && CheckConfig != UnrealTargetConfiguration.Shipping)
-                    {
-                        CheckConfig = UnrealTargetConfiguration.Development;
-                    }
-                }
-            }
-            switch (CheckConfig)
-            {
-                default:
-                case UnrealTargetConfiguration.Debug:
-                    CompileConfiguration = CPPTargetConfiguration.Debug;
-                    if (BuildConfiguration.bDebugBuildsActuallyUseDebugCRT)
-                    {
-                        InBuildTarget.GlobalCompileEnvironment.Config.Definitions.Add("_DEBUG=1"); // the engine doesn't use this, but lots of 3rd party stuff does
-                    }
-                    else
-                    {
-                        InBuildTarget.GlobalCompileEnvironment.Config.Definitions.Add("NDEBUG=1"); // the engine doesn't use this, but lots of 3rd party stuff does
-                    }
-                    InBuildTarget.GlobalCompileEnvironment.Config.Definitions.Add("UE_BUILD_DEBUG=1");
-                    break;
-                case UnrealTargetConfiguration.DebugGame:
-                // Individual game modules can be switched to be compiled in debug as necessary. By default, everything is compiled in development.
-                case UnrealTargetConfiguration.Development:
-                    CompileConfiguration = CPPTargetConfiguration.Development;
-                    InBuildTarget.GlobalCompileEnvironment.Config.Definitions.Add("NDEBUG=1"); // the engine doesn't use this, but lots of 3rd party stuff does
-                    InBuildTarget.GlobalCompileEnvironment.Config.Definitions.Add("UE_BUILD_DEVELOPMENT=1");
-                    break;
-                case UnrealTargetConfiguration.Shipping:
-                    CompileConfiguration = CPPTargetConfiguration.Shipping;
-                    InBuildTarget.GlobalCompileEnvironment.Config.Definitions.Add("NDEBUG=1"); // the engine doesn't use this, but lots of 3rd party stuff does
-                    InBuildTarget.GlobalCompileEnvironment.Config.Definitions.Add("UE_BUILD_SHIPPING=1");
-                    break;
-                case UnrealTargetConfiguration.Test:
-                    CompileConfiguration = CPPTargetConfiguration.Shipping;
-                    InBuildTarget.GlobalCompileEnvironment.Config.Definitions.Add("NDEBUG=1"); // the engine doesn't use this, but lots of 3rd party stuff does
-                    InBuildTarget.GlobalCompileEnvironment.Config.Definitions.Add("UE_BUILD_TEST=1");
-                    break;
-            }
-
-            // Set up the global C++ compilation and link environment.
-            InBuildTarget.GlobalCompileEnvironment.Config.Target.Configuration = CompileConfiguration;
-            InBuildTarget.GlobalLinkEnvironment.Config.Target.Configuration = CompileConfiguration;
-
-            // Create debug info based on the heuristics specified by the user.
-            InBuildTarget.GlobalCompileEnvironment.Config.bCreateDebugInfo =
-                !BuildConfiguration.bDisableDebugInfo && ShouldCreateDebugInfo(InBuildTarget.Platform, CheckConfig);
-            InBuildTarget.GlobalLinkEnvironment.Config.bCreateDebugInfo = InBuildTarget.GlobalCompileEnvironment.Config.bCreateDebugInfo;
-        }
-
-        /**
-         *	Whether this platform should create debug information or not
-         *	
-         *	@param	InPlatform			The UnrealTargetPlatform being built
-         *	@param	InConfiguration		The UnrealTargetConfiguration being built
-         *	
-         *	@return	bool				true if debug info should be generated, false if not
-         */
-        public abstract bool ShouldCreateDebugInfo(UnrealTargetPlatform Platform, UnrealTargetConfiguration Configuration);
-
-        /**
-         *	Gives the platform a chance to 'override' the configuration settings 
-         *	that are overridden on calls to RunUBT.
-         *	
-         *	@param	InPlatform			The UnrealTargetPlatform being built
-         *	@param	InConfiguration		The UnrealTargetConfiguration being built
-         *	
-         *	@return	bool				true if debug info should be generated, false if not
-         */
-        public virtual void ResetBuildConfiguration(UnrealTargetPlatform InPlatform, UnrealTargetConfiguration InConfiguration)
-        {
-        }
-
-        /**
-         *	Validate the UEBuildConfiguration for this platform
-         *	This is called BEFORE calling UEBuildConfiguration to allow setting 
-         *	various fields used in that function such as CompileLeanAndMean...
-         */
-        public virtual void ValidateUEBuildConfiguration()
-        {
-        }
-
-        /**
-         *	Validate configuration for this platform
-         *	NOTE: This function can/will modify BuildConfiguration!
-         *	
-         *	@param	InPlatform			The CPPTargetPlatform being built
-         *	@param	InConfiguration		The CPPTargetConfiguration being built
-         *	@param	bInCreateDebugInfo	true if debug info is getting create, false if not
-         */
-        public virtual void ValidateBuildConfiguration(CPPTargetConfiguration Configuration, CPPTargetPlatform Platform, bool bCreateDebugInfo)
-        {
-        }
-
-        /**
-         *	Return whether this platform has uniquely named binaries across multiple games
-         */
-        public virtual bool HasUniqueBinaries()
-        {
-            return true;
-        }
-
-        /**
-         *	Return whether we wish to have this platform's binaries in our builds
-         */
-        public virtual bool IsBuildRequired()
-        {
-            return true;
-        }
-
-        /**
-         *	Return whether we wish to have this platform's binaries in our CIS tests
-         */
-        public virtual bool IsCISRequired()
-        {
-            return true;
-        }
-
-        /// <summary>
-        /// Whether the build platform requires deployment prep
-        /// </summary>
-        /// <returns></returns>
-        public virtual bool RequiresDeployPrepAfterCompile()
-        {
-            return false;
-        }
-
-        /**
-         *	Return all valid configurations for this platform
-         *	
-         *  Typically, this is always Debug, Development, and Shipping - but Test is a likely future addition for some platforms
-         */
-        public virtual List<UnrealTargetConfiguration> GetConfigurations(UnrealTargetPlatform InUnrealTargetPlatform, bool bIncludeDebug)
-        {
-            List<UnrealTargetConfiguration> Configurations = new List<UnrealTargetConfiguration>()
+		/**
+		 *	Register the platform with the UEBuildPlatform class
+		 */
+		public void RegisterBuildPlatform()
+		{
+			ManageAndValidateSDK();
+			RegisterBuildPlatformInternal();
+		}
+
+		protected abstract void RegisterBuildPlatformInternal();
+
+		/**
+		 *	Retrieve the CPPTargetPlatform for the given UnrealTargetPlatform
+		 *
+		 *	@param	InUnrealTargetPlatform		The UnrealTargetPlatform being build
+		 *	
+		 *	@return	CPPTargetPlatform			The CPPTargetPlatform to compile for
+		 */
+		public abstract CPPTargetPlatform GetCPPTargetPlatform(UnrealTargetPlatform InUnrealTargetPlatform);
+
+		/// <summary>
+		/// Return whether the given platform requires a monolithic build
+		/// </summary>
+		/// <param name="InPlatform">The platform of interest</param>
+		/// <param name="InConfiguration">The configuration of interest</param>
+		/// <returns></returns>
+		public static bool PlatformRequiresMonolithicBuilds(UnrealTargetPlatform InPlatform, UnrealTargetConfiguration InConfiguration)
+		{
+			// Some platforms require monolithic builds...
+			var BuildPlatform = GetBuildPlatform(InPlatform, true);
+			if (BuildPlatform != null)
+			{
+				return BuildPlatform.ShouldCompileMonolithicBinary(InPlatform);
+			}
+
+			// We assume it does not
+			return false;
+		}
+
+		/**
+		 *	Get the extension to use for the given binary type
+		 *	
+		 *	@param	InBinaryType		The binary type being built
+		 *	
+		 *	@return	string				The binary extension (i.e. 'exe' or 'dll')
+		 */
+		public virtual string GetBinaryExtension(UEBuildBinaryType InBinaryType)
+		{
+			throw new BuildException("GetBinaryExtensiton for {0} not handled in {1}", InBinaryType.ToString(), this.ToString());
+		}
+
+		/**
+		 *	Get the extension to use for debug info for the given binary type
+		 *	
+		 *	@param	InBinaryType		The binary type being built
+		 *	
+		 *	@return	string				The debug info extension (i.e. 'pdb')
+		 */
+		public virtual string GetDebugInfoExtension(UEBuildBinaryType InBinaryType)
+		{
+			throw new BuildException("GetDebugInfoExtension for {0} not handled in {1}", InBinaryType.ToString(), this.ToString());
+		}
+
+		/**
+		 *	Whether incremental linking should be used
+		 *	
+		 *	@param	InPlatform			The CPPTargetPlatform being built
+		 *	@param	InConfiguration		The CPPTargetConfiguration being built
+		 *	
+		 *	@return	bool	true if incremental linking should be used, false if not
+		 */
+		public virtual bool ShouldUseIncrementalLinking(CPPTargetPlatform Platform, CPPTargetConfiguration Configuration)
+		{
+			return false;
+		}
+
+		/**
+		 *	Whether PDB files should be used
+		 *	
+		 *	@param	InPlatform			The CPPTargetPlatform being built
+		 *	@param	InConfiguration		The CPPTargetConfiguration being built
+		 *	@param	bInCreateDebugInfo	true if debug info is getting create, false if not
+		 *	
+		 *	@return	bool	true if PDB files should be used, false if not
+		 */
+		public virtual bool ShouldUsePDBFiles(CPPTargetPlatform Platform, CPPTargetConfiguration Configuration, bool bCreateDebugInfo)
+		{
+			return false;
+		}
+
+		/**
+		 *	Whether PCH files should be used
+		 *	
+		 *	@param	InPlatform			The CPPTargetPlatform being built
+		 *	@param	InConfiguration		The CPPTargetConfiguration being built
+		 *	
+		 *	@return	bool				true if PCH files should be used, false if not
+		 */
+		public virtual bool ShouldUsePCHFiles(CPPTargetPlatform Platform, CPPTargetConfiguration Configuration)
+		{
+			return BuildConfiguration.bUsePCHFiles;
+		}
+
+		/**
+		 *	Whether the editor should be built for this platform or not
+		 *	
+		 *	@param	InPlatform		The UnrealTargetPlatform being built
+		 *	@param	InConfiguration	The UnrealTargetConfiguration being built
+		 *	@return	bool			true if the editor should be built, false if not
+		 */
+		public virtual bool ShouldNotBuildEditor(UnrealTargetPlatform InPlatform, UnrealTargetConfiguration InConfiguration)
+		{
+			return false;
+		}
+
+		/**
+		 *	Whether this build should support ONLY cooked data or not
+		 *	
+		 *	@param	InPlatform		The UnrealTargetPlatform being built
+		 *	@param	InConfiguration	The UnrealTargetConfiguration being built
+		 *	@return	bool			true if the editor should be built, false if not
+		 */
+		public virtual bool BuildRequiresCookedData(UnrealTargetPlatform InPlatform, UnrealTargetConfiguration InConfiguration)
+		{
+			return false;
+		}
+
+		/**
+		 *	Whether the platform requires the extra UnityCPPWriter
+		 *	This is used to add an extra file for UBT to get the #include dependencies from
+		 *	
+		 *	@return	bool	true if it is required, false if not
+		 */
+		public virtual bool RequiresExtraUnityCPPWriter()
+		{
+			return false;
+		}
+
+		/**
+		 * Whether this platform should build a monolithic binary
+		 */
+		public virtual bool ShouldCompileMonolithicBinary(UnrealTargetPlatform InPlatform)
+		{
+			return false;
+		}
+
+		/**
+		 *	Get a list of extra modules the platform requires.
+		 *	This is to allow undisclosed platforms to add modules they need without exposing information about the platform.
+		 *	
+		 *	@param	Target						The target being build
+		 *	@param	BuildTarget					The UEBuildTarget getting build
+		 *	@param	PlatformExtraModules		OUTPUT the list of extra modules the platform needs to add to the target
+		 */
+		public virtual void GetExtraModules(TargetInfo Target, UEBuildTarget BuildTarget, ref List<string> PlatformExtraModules)
+		{
+		}
+
+		/**
+		 *	Modify the newly created module passed in for this platform.
+		 *	This is not required - but allows for hiding details of a
+		 *	particular platform.
+		 *	
+		 *	@param	InModule		The newly loaded module
+		 *	@param	Target			The target being build
+		 */
+		public virtual void ModifyNewlyLoadedModule(UEBuildModule InModule, TargetInfo Target)
+		{
+		}
+
+		/**
+		 *	Setup the target environment for building
+		 *	
+		 *	@param	InBuildTarget		The target being built
+		 */
+		public abstract void SetUpEnvironment(UEBuildTarget InBuildTarget);
+
+		/**
+		 * Allow the platform to set an optional architecture
+		 */
+		public virtual string GetActiveArchitecture()
+		{
+			// by default, use an empty architecture (which is really just a modifer to the platform for some paths/names)
+			return "";
+		}
+
+		/**
+		 * Allow the platform to override the NMake output name
+		 */
+		public virtual string ModifyNMakeOutput(string ExeName)
+		{
+			// by default, use original
+			return ExeName;
+		}
+
+		/**
+		 * Allow the platform to apply architecture-specific name according to its rules
+		 * 
+		 * @param BinaryName name of the binary, not specific to any architecture
+		 *
+		 * @return Possibly architecture-specific name
+		 */
+		public virtual string ApplyArchitectureName(string BinaryName)
+		{
+			// by default, use logic that assumes architectures to start with "-"
+			return BinaryName + GetActiveArchitecture();
+		}
+
+		/**
+		 *	Setup the configuration environment for building
+		 *	
+		 *	@param	InBuildTarget		The target being built
+		 */
+		public virtual void SetUpConfigurationEnvironment(UEBuildTarget InBuildTarget)
+		{
+			// Determine the C++ compile/link configuration based on the Unreal configuration.
+			CPPTargetConfiguration CompileConfiguration;
+			UnrealTargetConfiguration CheckConfig = InBuildTarget.Configuration;
+			//@todo SAS: Add a true Debug mode!
+			if (UnrealBuildTool.RunningRocket())
+			{
+				if (Utils.IsFileUnderDirectory(InBuildTarget.OutputPath, UnrealBuildTool.GetUProjectPath()))
+				{
+					if (CheckConfig == UnrealTargetConfiguration.Debug)
+					{
+						CheckConfig = UnrealTargetConfiguration.DebugGame;
+					}
+				}
+				else
+				{
+					// Only Development and Shipping are supported for engine modules
+					if (CheckConfig != UnrealTargetConfiguration.Development && CheckConfig != UnrealTargetConfiguration.Shipping)
+					{
+						CheckConfig = UnrealTargetConfiguration.Development;
+					}
+				}
+			}
+			switch (CheckConfig)
+			{
+				default:
+				case UnrealTargetConfiguration.Debug:
+					CompileConfiguration = CPPTargetConfiguration.Debug;
+					if (BuildConfiguration.bDebugBuildsActuallyUseDebugCRT)
+					{
+						InBuildTarget.GlobalCompileEnvironment.Config.Definitions.Add("_DEBUG=1"); // the engine doesn't use this, but lots of 3rd party stuff does
+					}
+					else
+					{
+						InBuildTarget.GlobalCompileEnvironment.Config.Definitions.Add("NDEBUG=1"); // the engine doesn't use this, but lots of 3rd party stuff does
+					}
+					InBuildTarget.GlobalCompileEnvironment.Config.Definitions.Add("UE_BUILD_DEBUG=1");
+					break;
+				case UnrealTargetConfiguration.DebugGame:
+				// Individual game modules can be switched to be compiled in debug as necessary. By default, everything is compiled in development.
+				case UnrealTargetConfiguration.Development:
+					CompileConfiguration = CPPTargetConfiguration.Development;
+					InBuildTarget.GlobalCompileEnvironment.Config.Definitions.Add("NDEBUG=1"); // the engine doesn't use this, but lots of 3rd party stuff does
+					InBuildTarget.GlobalCompileEnvironment.Config.Definitions.Add("UE_BUILD_DEVELOPMENT=1");
+					break;
+				case UnrealTargetConfiguration.Shipping:
+					CompileConfiguration = CPPTargetConfiguration.Shipping;
+					InBuildTarget.GlobalCompileEnvironment.Config.Definitions.Add("NDEBUG=1"); // the engine doesn't use this, but lots of 3rd party stuff does
+					InBuildTarget.GlobalCompileEnvironment.Config.Definitions.Add("UE_BUILD_SHIPPING=1");
+					break;
+				case UnrealTargetConfiguration.Test:
+					CompileConfiguration = CPPTargetConfiguration.Shipping;
+					InBuildTarget.GlobalCompileEnvironment.Config.Definitions.Add("NDEBUG=1"); // the engine doesn't use this, but lots of 3rd party stuff does
+					InBuildTarget.GlobalCompileEnvironment.Config.Definitions.Add("UE_BUILD_TEST=1");
+					break;
+			}
+
+			// Set up the global C++ compilation and link environment.
+			InBuildTarget.GlobalCompileEnvironment.Config.Target.Configuration = CompileConfiguration;
+			InBuildTarget.GlobalLinkEnvironment.Config.Target.Configuration = CompileConfiguration;
+
+			// Create debug info based on the heuristics specified by the user.
+			InBuildTarget.GlobalCompileEnvironment.Config.bCreateDebugInfo =
+				!BuildConfiguration.bDisableDebugInfo && ShouldCreateDebugInfo(InBuildTarget.Platform, CheckConfig);
+			InBuildTarget.GlobalLinkEnvironment.Config.bCreateDebugInfo = InBuildTarget.GlobalCompileEnvironment.Config.bCreateDebugInfo;
+		}
+
+		/**
+		 *	Whether this platform should create debug information or not
+		 *	
+		 *	@param	InPlatform			The UnrealTargetPlatform being built
+		 *	@param	InConfiguration		The UnrealTargetConfiguration being built
+		 *	
+		 *	@return	bool				true if debug info should be generated, false if not
+		 */
+		public abstract bool ShouldCreateDebugInfo(UnrealTargetPlatform Platform, UnrealTargetConfiguration Configuration);
+
+		/**
+		 *	Gives the platform a chance to 'override' the configuration settings 
+		 *	that are overridden on calls to RunUBT.
+		 *	
+		 *	@param	InPlatform			The UnrealTargetPlatform being built
+		 *	@param	InConfiguration		The UnrealTargetConfiguration being built
+		 *	
+		 *	@return	bool				true if debug info should be generated, false if not
+		 */
+		public virtual void ResetBuildConfiguration(UnrealTargetPlatform InPlatform, UnrealTargetConfiguration InConfiguration)
+		{
+		}
+
+		/**
+		 *	Validate the UEBuildConfiguration for this platform
+		 *	This is called BEFORE calling UEBuildConfiguration to allow setting 
+		 *	various fields used in that function such as CompileLeanAndMean...
+		 */
+		public virtual void ValidateUEBuildConfiguration()
+		{
+		}
+
+		/**
+		 *	Validate configuration for this platform
+		 *	NOTE: This function can/will modify BuildConfiguration!
+		 *	
+		 *	@param	InPlatform			The CPPTargetPlatform being built
+		 *	@param	InConfiguration		The CPPTargetConfiguration being built
+		 *	@param	bInCreateDebugInfo	true if debug info is getting create, false if not
+		 */
+		public virtual void ValidateBuildConfiguration(CPPTargetConfiguration Configuration, CPPTargetPlatform Platform, bool bCreateDebugInfo)
+		{
+		}
+
+		/**
+		 *	Return whether this platform has uniquely named binaries across multiple games
+		 */
+		public virtual bool HasUniqueBinaries()
+		{
+			return true;
+		}
+
+		/**
+		 *	Return whether we wish to have this platform's binaries in our builds
+		 */
+		public virtual bool IsBuildRequired()
+		{
+			return true;
+		}
+
+		/**
+		 *	Return whether we wish to have this platform's binaries in our CIS tests
+		 */
+		public virtual bool IsCISRequired()
+		{
+			return true;
+		}
+
+		/// <summary>
+		/// Whether the build platform requires deployment prep
+		/// </summary>
+		/// <returns></returns>
+		public virtual bool RequiresDeployPrepAfterCompile()
+		{
+			return false;
+		}
+
+		/**
+		 *	Return all valid configurations for this platform
+		 *	
+		 *  Typically, this is always Debug, Development, and Shipping - but Test is a likely future addition for some platforms
+		 */
+		public virtual List<UnrealTargetConfiguration> GetConfigurations(UnrealTargetPlatform InUnrealTargetPlatform, bool bIncludeDebug)
+		{
+			List<UnrealTargetConfiguration> Configurations = new List<UnrealTargetConfiguration>()
 			{
 				UnrealTargetConfiguration.Development, 
 			};
 
-            if (bIncludeDebug)
-            {
-                Configurations.Insert(0, UnrealTargetConfiguration.Debug);
-            }
-
-            return Configurations;
-        }
-
-        /**
-         *	Setup the binaries for this specific platform.
-         *	
-         *	@param	InBuildTarget		The target being built
-         */
-        public virtual void SetupBinaries(UEBuildTarget InBuildTarget)
-        {
-        }
-    }
-
-    // AutoSDKs handling portion
-    public abstract partial class UEBuildPlatform : IUEBuildPlatform
-    {
-
-        #region protected AutoSDKs Utility
-
-        /** Name of the file that holds currently install SDK version string */
-        protected static string CurrentlyInstalledSDKStringManifest = "CurrentlyInstalled.txt";
-
-        /** name of the file that holds the last succesfully run SDK setup script version */
-        protected static string LastRunScriptVersionManifest = "CurrentlyInstalled.Version.txt";
-
-        /** Name of the file that holds environment variables of current SDK */
-        protected static string SDKEnvironmentVarsFile = "OutputEnvVars.txt";
-
-        protected static readonly string SDKRootEnvVar = "UE_SDKS_ROOT";
-
-        protected static string AutoSetupEnvVar = "AutoSDKSetup";
-
-        /** 
-         * Whether platform supports switching SDKs during runtime
-         * 
-         * @return true if supports
-         */
-        protected virtual bool PlatformSupportsAutoSDKs()
-        {
-            return false;
-        }
-
-        static private bool bCheckedAutoSDKRootEnvVar = false;
-        static private bool bAutoSDKSystemEnabled = false;
-        static private bool HasAutoSDKSystemEnabled()
-        {
-            if (!bCheckedAutoSDKRootEnvVar)
-            {
-                string SDKRoot = Environment.GetEnvironmentVariable(SDKRootEnvVar);
-                if (SDKRoot != null)
-                {
-                    bAutoSDKSystemEnabled = true;
-                }
-                bCheckedAutoSDKRootEnvVar = true;
-            }
-            return bAutoSDKSystemEnabled;
-        }
-
-        // Whether AutoSDK setup is safe. AutoSDKs will damage manual installs on some platforms.
-        protected bool IsAutoSDKSafe()
-        {
-            return !IsAutoSDKDestructive() || !HasAnyManualInstall();
-        }
-
-        /** 
-         * Returns SDK string as required by the platform 
-         * 
-         * @return Valid SDK string
-         */
-        protected virtual string GetRequiredSDKString()
-        {
-            return "";
-        }
-
-        /**
-        * Gets the version number of the SDK setup script itself.  The version in the base should ALWAYS be the master revision from the last refactor.  
-        * If you need to force a rebuild for a given platform, override this for the given platform.
-        * @return Setup script version
-        */
-        protected virtual String GetRequiredScriptVersionString()
-        {
-            return "3.0";
-        }
-
-        /** 
-         * Returns path to platform SDKs
-         * 
-         * @return Valid SDK string
-         */
-        protected string GetPathToPlatformAutoSDKs()
-        {
-            string SDKPath = "";
-            string SDKRoot = Environment.GetEnvironmentVariable(SDKRootEnvVar);
-            if (SDKRoot != null)
-            {
-                if (SDKRoot != "")
-                {
-                    SDKPath = Path.Combine(SDKRoot, "Host" + BuildHostPlatform.Current.Platform, GetSDKTargetPlatformName());
-                }
-            }
-            return SDKPath;
-        }
-
-        /**
-         * Because most ManualSDK determination depends on reading env vars, if this process is spawned by a process that ALREADY set up
-         * AutoSDKs then all the SDK env vars will exist, and we will spuriously detect a Manual SDK. (children inherit the environment of the parent process).
-         * Therefore we write out an env var to set in the command file (OutputEnvVars.txt) such that child processes can determine if their manual SDK detection
-         * is bogus.  Make it platform specific so that platforms can be in different states.
-         */
-        protected string GetPlatformAutoSDKSetupEnvVar()
-        {
-            return GetSDKTargetPlatformName() + AutoSetupEnvVar;
-        }
-
-        /**
-         * Gets currently installed version
-         * 
-         * @param PlatformSDKRoot absolute path to platform SDK root
-         * @param OutInstalledSDKVersionString version string as currently installed
-         * 
-         * @return true if was able to read it
-         */
-        protected bool GetCurrentlyInstalledSDKString(string PlatformSDKRoot, out string OutInstalledSDKVersionString)
-        {
-            if (Directory.Exists(PlatformSDKRoot))
-            {
-                string VersionFilename = Path.Combine(PlatformSDKRoot, CurrentlyInstalledSDKStringManifest);
-                if (File.Exists(VersionFilename))
-                {
-                    using (StreamReader Reader = new StreamReader(VersionFilename))
-                    {
-                        string Version = Reader.ReadLine();
-                        string Type = Reader.ReadLine();
-
-                        // don't allow ManualSDK installs to count as an AutoSDK install version.
-                        if (Type != null && Type == "AutoSDK")
-                        {
-                            if (Version != null)
-                            {
-                                OutInstalledSDKVersionString = Version;
-                                return true;
-                            }
-                        }
-                    }
-                }
-            }
-
-            OutInstalledSDKVersionString = "";
-            return false;
-        }
-
-        /**
-         * Gets the version of the last successfully run setup script.
-         * 
-         * @param PlatformSDKRoot absolute path to platform SDK root
-         * @param OutLastRunScriptVersion version string
-         * 
-         * @return true if was able to read it
-         */
-        protected bool GetLastRunScriptVersionString(string PlatformSDKRoot, out string OutLastRunScriptVersion)
-        {
-            if (Directory.Exists(PlatformSDKRoot))
-            {
-                string VersionFilename = Path.Combine(PlatformSDKRoot, LastRunScriptVersionManifest);
-                if (File.Exists(VersionFilename))
-                {
-                    using (StreamReader Reader = new StreamReader(VersionFilename))
-                    {
-                        string Version = Reader.ReadLine();
-                        if (Version != null)
-                        {
-                            OutLastRunScriptVersion = Version;
-                            return true;
-                        }
-                    }
-                }
-            }
-
-            OutLastRunScriptVersion = "";
-            return false;
-        }
-
-        /**
-         * Sets currently installed version
-         * 
-         * @param PlatformSDKRoot absolute path to platform SDK root
-         * @param InstalledSDKVersionString SDK version string to set
-         * 
-         * @return true if was able to set it
-         */
-        protected bool SetCurrentlyInstalledAutoSDKString(String InstalledSDKVersionString)
-        {
-            String PlatformSDKRoot = GetPathToPlatformAutoSDKs();
-            if (Directory.Exists(PlatformSDKRoot))
-            {
-                string VersionFilename = Path.Combine(PlatformSDKRoot, CurrentlyInstalledSDKStringManifest);
-                if (File.Exists(VersionFilename))
-                {
-                    File.Delete(VersionFilename);
-                }
-
-                using (StreamWriter Writer = File.CreateText(VersionFilename))
-                {
-                    Writer.WriteLine(InstalledSDKVersionString);
-                    Writer.WriteLine("AutoSDK");
-                    return true;
-                }
-            }
-
-            return false;
-        }
-
-        protected void SetupManualSDK()
-        {
-            if (PlatformSupportsAutoSDKs() && HasAutoSDKSystemEnabled())
-            {
-                String InstalledSDKVersionString = GetRequiredSDKString();
-                String PlatformSDKRoot = GetPathToPlatformAutoSDKs();
-                if (Directory.Exists(PlatformSDKRoot))
-                {
-                    string VersionFilename = Path.Combine(PlatformSDKRoot, CurrentlyInstalledSDKStringManifest);
-                    if (File.Exists(VersionFilename))
-                    {
-                        File.Delete(VersionFilename);
-                    }
-
-                    string EnvVarFile = Path.Combine(PlatformSDKRoot, SDKEnvironmentVarsFile);
-                    if (File.Exists(EnvVarFile))
-                    {
-                        File.Delete(EnvVarFile);
-                    }
-
-                    using (StreamWriter Writer = File.CreateText(VersionFilename))
-                    {
-                        Writer.WriteLine(InstalledSDKVersionString);
-                        Writer.WriteLine("ManualSDK");
-                    }
-                }
-            }
-        }
-
-        protected bool SetLastRunAutoSDKScriptVersion(string LastRunScriptVersion)
-        {
-            String PlatformSDKRoot = GetPathToPlatformAutoSDKs();
-            if (Directory.Exists(PlatformSDKRoot))
-            {
-                string VersionFilename = Path.Combine(PlatformSDKRoot, LastRunScriptVersionManifest);
-                if (File.Exists(VersionFilename))
-                {
-                    File.Delete(VersionFilename);
-                }
-
-                using (StreamWriter Writer = File.CreateText(VersionFilename))
-                {
-                    Writer.WriteLine(LastRunScriptVersion);
-                    return true;
-                }
-            }
-            return false;
-        }
-
-        /**
-        * Returns Hook names as needed by the platform
-        * (e.g. can be overriden with custom executables or scripts)
-        *
-        * @param Hook Hook type
-        */
-        protected virtual string GetHookExecutableName(SDKHookType Hook)
-        {
-            if (Hook == SDKHookType.Uninstall)
-            {
-                return "unsetup.bat";
-            }
-
-            return "setup.bat";
-        }
-
-        /**
-         * Runs install/uninstall hooks for SDK
-         * 
-         * @param PlatformSDKRoot absolute path to platform SDK root
-         * @param SDKVersionString version string to run for (can be empty!)
-         * @param Hook which one of hooks to run
-         * @param bHookCanBeNonExistent whether a non-existing hook means failure
-         * 
-         * @return true if succeeded
-         */
-        protected virtual bool RunAutoSDKHooks(string PlatformSDKRoot, string SDKVersionString, SDKHookType Hook, bool bHookCanBeNonExistent = true)
-        {
-            if (!IsAutoSDKSafe())
-            {
-                Console.ForegroundColor = ConsoleColor.Red;
-                Console.WriteLine(GetSDKTargetPlatformName() + " attempted to run SDK hook which could have damaged manual SDK install!");
-                Console.ResetColor();
-
-                return false;
-            }
-            if (SDKVersionString != "")
-            {
-                string SDKDirectory = Path.Combine(PlatformSDKRoot, SDKVersionString);
-                string HookExe = Path.Combine(SDKDirectory, GetHookExecutableName(Hook));
-
-                if (File.Exists(HookExe))
-                {
-                    Console.WriteLine("Running {0} hook {1}", Hook, HookExe);
-
-                    // run it
-                    Process HookProcess = new Process();
-                    HookProcess.StartInfo.WorkingDirectory = SDKDirectory;
-                    HookProcess.StartInfo.FileName = HookExe;
-                    HookProcess.StartInfo.Arguments = "";
-                    HookProcess.StartInfo.WindowStyle = ProcessWindowStyle.Hidden;
-
-                    //installers may require administrator access to succeed. so run as an admmin.
-                    HookProcess.StartInfo.Verb = "runas";
-                    HookProcess.Start();
-                    HookProcess.WaitForExit();
-
-                    if (HookProcess.ExitCode != 0)
-                    {
-                        Console.WriteLine("Hook exited uncleanly (returned {0}), considering it failed.", HookProcess.ExitCode);
-                        return false;
-                    }
-
-                    return true;
-                }
-                else
-                {
-                    Console.WriteLine("File {0} does not exist", HookExe);
-                }
-            }
-            else
-            {
-                Console.WriteLine("Version string is blank for {0}. Can't determine {1} hook.", PlatformSDKRoot, Hook.ToString());
-            }
-
-            return bHookCanBeNonExistent;
-        }
-
-        /**
-         * Loads environment variables from SDK
-         * If any commands are added or removed the handling needs to be duplicated in
-         * TargetPlatformManagerModule.cpp
-         * 
-         * @param PlatformSDKRoot absolute path to platform SDK         
-         * 
-         * @return true if succeeded
-         */
-        protected bool SetupEnvironmentFromAutoSDK(string PlatformSDKRoot)
-        {
-            string EnvVarFile = Path.Combine(PlatformSDKRoot, SDKEnvironmentVarsFile);
-            if (File.Exists(EnvVarFile))
-            {
-                using (StreamReader Reader = new StreamReader(EnvVarFile))
-                {
-                    List<string> PathAdds = new List<string>();
-                    List<string> PathRemoves = new List<string>();
-
-                    List<string> EnvVarNames = new List<string>();
-                    List<string> EnvVarValues = new List<string>();
-
-                    bool bNeedsToWriteAutoSetupEnvVar = true;
-                    String PlatformSetupEnvVar = GetPlatformAutoSDKSetupEnvVar();
-                    for (; ; )
-                    {
-                        string VariableString = Reader.ReadLine();
-                        if (VariableString == null)
-                        {
-                            break;
-                        }
-
-                        string[] Parts = VariableString.Split('=');
-                        if (Parts.Length != 2)
-                        {
-                            Console.WriteLine("Incorrect environment variable declaration:");
-                            Console.WriteLine(VariableString);
-                            return false;
-                        }
-
-                        if (String.Compare(Parts[0], "strippath", true) == 0)
-                        {
-                            PathRemoves.Add(Parts[1]);
-                        }
-                        else if (String.Compare(Parts[0], "addpath", true) == 0)
-                        {
-                            PathAdds.Add(Parts[1]);
-                        }
-                        else
-                        {
-                            if (String.Compare(Parts[0], PlatformSetupEnvVar) == 0)
-                            {
-                                bNeedsToWriteAutoSetupEnvVar = false;
-                            }
-                            // convenience for setup.bat writers.  Trim any accidental whitespace from var names/values.
-                            EnvVarNames.Add(Parts[0].Trim());
-                            EnvVarValues.Add(Parts[1].Trim());
-                        }
-                    }
-
-                    // don't actually set anything until we successfully validate and read all values in.
-                    // we don't want to set a few vars, return a failure, and then have a platform try to
-                    // build against a manually installed SDK with half-set env vars.
-                    for (int i = 0; i < EnvVarNames.Count; ++i)
-                    {
-                        string EnvVarName = EnvVarNames[i];
-                        string EnvVarValue = EnvVarValues[i];
-                        if (BuildConfiguration.bPrintDebugInfo)
-                        {
-                            Console.WriteLine("Setting variable '{0}' to '{1}'", EnvVarName, EnvVarValue);
-                        }
-                        Environment.SetEnvironmentVariable(EnvVarName, EnvVarValue);
-                    }
-
-
-                    // actually perform the PATH stripping / adding.
-                    String OrigPathVar = Environment.GetEnvironmentVariable("PATH");
-                    String PathDelimiter = GetPathVarDelimiter();
-                    String[] PathVars = OrigPathVar.Split(PathDelimiter.ToCharArray());
-
-                    List<String> ModifiedPathVars = new List<string>();
-                    ModifiedPathVars.AddRange(PathVars);
-
-                    // perform removes first, in case they overlap with any adds.
-                    foreach (String PathRemove in PathRemoves)
-                    {
-                        foreach (String PathVar in PathVars)
-                        {
-                            if (PathVar.IndexOf(PathRemove, StringComparison.OrdinalIgnoreCase) >= 0)
-                            {
-                                if (BuildConfiguration.bPrintDebugInfo)
-                                {
-                                    Console.WriteLine("Removing Path: '{0}'", PathVar);
-                                }
-                                ModifiedPathVars.Remove(PathVar);
-                            }
-                        }
-                    }
-
-                    // remove all the of ADDs so that if this function is executed multiple times, the paths will be guarateed to be in the same order after each run.
-                    // If we did not do this, a 'remove' that matched some, but not all, of our 'adds' would cause the order to change.
-                    foreach (String PathAdd in PathAdds)
-                    {
-                        foreach (String PathVar in PathVars)
-                        {
-                            if (String.Compare(PathAdd, PathVar, true) == 0)
-                            {
-                                if (BuildConfiguration.bPrintDebugInfo)
-                                {
-                                    Console.WriteLine("Removing Path: '{0}'", PathVar);
-                                }
-                                ModifiedPathVars.Remove(PathVar);
-                            }
-                        }
-                    }
-
-                    // perform adds, but don't add duplicates
-                    foreach (String PathAdd in PathAdds)
-                    {
-                        if (!ModifiedPathVars.Contains(PathAdd))
-                        {
-                            if (BuildConfiguration.bPrintDebugInfo)
-                            {
-                                Console.WriteLine("Adding Path: '{0}'", PathAdd);
-                            }
-                            ModifiedPathVars.Add(PathAdd);
-                        }
-                    }
-
-                    String ModifiedPath = String.Join(PathDelimiter, ModifiedPathVars);
-                    Environment.SetEnvironmentVariable("PATH", ModifiedPath);
-
-                    Reader.Close();
-
-                    // write out env var command so any process using this commandfile will mark itself as having had autosdks set up.
-                    // avoids child processes spuriously detecting manualsdks.
-                    if (bNeedsToWriteAutoSetupEnvVar)
-                    {
-                        using (StreamWriter Writer = File.AppendText(EnvVarFile))
-                        {
-                            Writer.WriteLine("{0}=1", PlatformSetupEnvVar);
-                        }
-                        // set the var in the local environment in case this process spawns any others.
-                        Environment.SetEnvironmentVariable(PlatformSetupEnvVar, "1");
-                    }
-
-                    // make sure we know that we've modified the local environment, invalidating manual installs for this run.
-                    bLocalProcessSetupAutoSDK = true;
-
-                    return true;
-                }
-            }
-            else
-            {
-                Console.WriteLine("Cannot set up environment for {1} because command file {1} does not exist.", PlatformSDKRoot, EnvVarFile);
-            }
-
-            return false;
-        }
-
-        protected void InvalidateCurrentlyInstalledAutoSDK()
-        {
-            String PlatformSDKRoot = GetPathToPlatformAutoSDKs();
-            if (Directory.Exists(PlatformSDKRoot))
-            {
-                string SDKFilename = Path.Combine(PlatformSDKRoot, CurrentlyInstalledSDKStringManifest);
-                if (File.Exists(SDKFilename))
-                {
-                    File.Delete(SDKFilename);
-                }
-
-                string VersionFilename = Path.Combine(PlatformSDKRoot, LastRunScriptVersionManifest);
-                if (File.Exists(VersionFilename))
-                {
-                    File.Delete(VersionFilename);
-                }
-
-                string EnvVarFile = Path.Combine(PlatformSDKRoot, SDKEnvironmentVarsFile);
-                if (File.Exists(EnvVarFile))
-                {
-                    File.Delete(EnvVarFile);
-                }
-            }
-        }
-
-        /** 
-         * Currently installed AutoSDK is written out to a text file in a known location.  
-         * This function just compares the file's contents with the current requirements.
-         */
-        protected SDKStatus HasRequiredAutoSDKInstalled()
-        {
-            if (PlatformSupportsAutoSDKs() && HasAutoSDKSystemEnabled())
-            {
-                string AutoSDKRoot = GetPathToPlatformAutoSDKs();
-                if (AutoSDKRoot != "")
-                {
-                    // check script version so script fixes can be propagated without touching every build machine's CurrentlyInstalled file manually.
-                    bool bScriptVersionMatches = false;
-                    string CurrentScriptVersionString;
-                    if (GetLastRunScriptVersionString(AutoSDKRoot, out CurrentScriptVersionString) && CurrentScriptVersionString == GetRequiredScriptVersionString())
-                    {
-                        bScriptVersionMatches = true;
-                    }
-
-                    // check to make sure OutputEnvVars doesn't need regenerating
-                    string EnvVarFile = Path.Combine(AutoSDKRoot, SDKEnvironmentVarsFile);
-                    bool bEnvVarFileExists = File.Exists(EnvVarFile);
-
-                    string CurrentSDKString;
-                    if (bEnvVarFileExists && GetCurrentlyInstalledSDKString(AutoSDKRoot, out CurrentSDKString) && CurrentSDKString == GetRequiredSDKString() && bScriptVersionMatches)
-                    {
-                        return SDKStatus.Valid;
-                    }
-                    return SDKStatus.Invalid;
-                }
-            }
-            return SDKStatus.Invalid;
-        }
-
-        // This tracks if we have already checked the sdk installation.
-        private Int32 SDKCheckStatus = -1;
-
-        // true if we've ever overridden the process's environment with AutoSDK data.  After that, manual installs cannot be considered valid ever again.
-        private bool bLocalProcessSetupAutoSDK = false;
-
-        protected bool HasSetupAutoSDK()
-        {
-            return bLocalProcessSetupAutoSDK || HasParentProcessSetupAutoSDK();
-        }
-
-        protected bool HasParentProcessSetupAutoSDK()
-        {
-            bool bParentProcessSetupAutoSDK = false;
-            String AutoSDKSetupVarName = GetPlatformAutoSDKSetupEnvVar();
-            String AutoSDKSetupVar = Environment.GetEnvironmentVariable(AutoSDKSetupVarName);
-            if (!String.IsNullOrEmpty(AutoSDKSetupVar))
-            {
-                bParentProcessSetupAutoSDK = true;
-            }
-            return bParentProcessSetupAutoSDK;
-        }
-
-        protected SDKStatus HasRequiredManualSDK()
-        {
-            if (HasSetupAutoSDK())
-            {
-                return SDKStatus.Invalid;
-            }
-
-            // manual installs are always invalid if we have modified the process's environment for AutoSDKs
-            return HasRequiredManualSDKInternal();
-        }
-
-        // for platforms with destructive AutoSDK.  Report if any manual sdk is installed that may be damaged by an autosdk.
-        protected virtual bool HasAnyManualInstall()
-        {
-            return false;
-        }
-
-        // tells us if the user has a valid manual install.
-        protected abstract SDKStatus HasRequiredManualSDKInternal();
-
-        // some platforms will fail if there is a manual install that is the WRONG manual install.
-        protected virtual bool AllowInvalidManualInstall()
-        {
-            return true;
-        }
-
-        // platforms can choose if they prefer a correct the the AutoSDK install over the manual install.
-        protected virtual bool PreferAutoSDK()
-        {
-            return true;
-        }
-
-        // some platforms don't support parallel SDK installs.  AutoSDK on these platforms will
-        // actively damage an existing manual install by overwriting files in it.  AutoSDK must NOT
-        // run any setup if a manual install exists in this case.
-        protected virtual bool IsAutoSDKDestructive()
-        {
-            return false;
-        }
-
-        /**
-        * Runs batch files if necessary to set up required AutoSDK.
-        * AutoSDKs are SDKs that have not been setup through a formal installer, but rather come from
-        * a source control directory, or other local copy.
-        */
-        private void SetupAutoSDK()
-        {
-            if (IsAutoSDKSafe() && PlatformSupportsAutoSDKs() && HasAutoSDKSystemEnabled())
-            {
-                // run installation for autosdk if necessary.
-                if (HasRequiredAutoSDKInstalled() == SDKStatus.Invalid)
-                {
-                    //reset check status so any checking sdk status after the attempted setup will do a real check again.
-                    SDKCheckStatus = -1;
-
-                    string AutoSDKRoot = GetPathToPlatformAutoSDKs();
-                    string CurrentSDKString;
-                    GetCurrentlyInstalledSDKString(AutoSDKRoot, out CurrentSDKString);
-
-                    // switch over (note that version string can be empty)
-                    if (!RunAutoSDKHooks(AutoSDKRoot, CurrentSDKString, SDKHookType.Uninstall))
-                    {
-                        Console.WriteLine("Failed to uninstall currently installed SDK {0}", CurrentSDKString);
-                        InvalidateCurrentlyInstalledAutoSDK();
-                        return;
-                    }
-                    // delete Manifest file to avoid multiple uninstalls
-                    InvalidateCurrentlyInstalledAutoSDK();
-
-                    if (!RunAutoSDKHooks(AutoSDKRoot, GetRequiredSDKString(), SDKHookType.Install, false))
-                    {
-                        Console.WriteLine("Failed to install required SDK {0}.  Attemping to uninstall", GetRequiredSDKString());
-                        RunAutoSDKHooks(AutoSDKRoot, GetRequiredSDKString(), SDKHookType.Uninstall, false);
-                        return;
-                    }
-
-                    string EnvVarFile = Path.Combine(AutoSDKRoot, SDKEnvironmentVarsFile);
-                    if (!File.Exists(EnvVarFile))
-                    {
-                        Console.WriteLine("Installation of required SDK {0}.  Did not generate Environment file {1}", GetRequiredSDKString(), EnvVarFile);
-                        RunAutoSDKHooks(AutoSDKRoot, GetRequiredSDKString(), SDKHookType.Uninstall, false);
-                        return;
-                    }
-
-                    SetCurrentlyInstalledAutoSDKString(GetRequiredSDKString());
-                    SetLastRunAutoSDKScriptVersion(GetRequiredScriptVersionString());
-                }
-
-                // fixup process environment to match autosdk
-                SetupEnvironmentFromAutoSDK();
-            }
-        }
-
-        #endregion
-
-        #region public AutoSDKs Utility
-
-        /** Enum describing types of hooks a platform SDK can have */
-        public enum SDKHookType
-        {
-            Install,
-            Uninstall
-        };
-
-        /** 
-         * Returns platform-specific name used in SDK repository
-         * 
-         * @return path to SDK Repository
-         */
-        public virtual string GetSDKTargetPlatformName()
-        {
-            return "";
-        }
-
-        /* Whether or not we should try to automatically switch SDKs when asked to validate the platform's SDK state. */
-        public static bool bAllowAutoSDKSwitching = true;
-
-        public SDKStatus SetupEnvironmentFromAutoSDK()
-        {
-            string PlatformSDKRoot = GetPathToPlatformAutoSDKs();
-
-            // load environment variables from current SDK
-            if (!SetupEnvironmentFromAutoSDK(PlatformSDKRoot))
-            {
-                Console.WriteLine("Failed to load environment from required SDK {0}", GetRequiredSDKString());
-                InvalidateCurrentlyInstalledAutoSDK();
-                return SDKStatus.Invalid;
-            }
-            return SDKStatus.Valid;
-        }
-
-        /**
-         *	Whether the required external SDKs are installed for this platform.  
-         *	Could be either a manual install or an AutoSDK.
-         */
-        public SDKStatus HasRequiredSDKsInstalled()
-        {
-            // avoid redundant potentially expensive SDK checks.
-            if (SDKCheckStatus == -1)
-            {
-                bool bHasManualSDK = HasRequiredManualSDK() == SDKStatus.Valid;
-                bool bHasAutoSDK = HasRequiredAutoSDKInstalled() == SDKStatus.Valid;
-
-                // Per-Platform implementations can choose how to handle non-Auto SDK detection / handling.
-                SDKCheckStatus = (bHasManualSDK || bHasAutoSDK) ? 1 : 0;
-            }
-            return SDKCheckStatus == 1 ? SDKStatus.Valid : SDKStatus.Invalid;
-        }
-
-        // Arbitrates between manual SDKs and setting up AutoSDK based on program options and platform preferences.
-        public void ManageAndValidateSDK()
-        {
-            // do not modify installed manifests if parent process has already set everything up.
-            // this avoids problems with determining IsAutoSDKSafe and doing an incorrect invalidate.
-            if (bAllowAutoSDKSwitching && !HasParentProcessSetupAutoSDK())
-            {
-                bool bHasRequiredManualSDK = HasRequiredManualSDK() == SDKStatus.Valid;
-                if (IsAutoSDKSafe() && (PreferAutoSDK() || !bHasRequiredManualSDK))
-                {
-                    SetupAutoSDK();
-                }
-                else if (bHasRequiredManualSDK)
-                {
-                    SetupManualSDK();
-                }
-                else
-                {
-                    InvalidateCurrentlyInstalledAutoSDK();
-                }
-            }
-
-            if (BuildConfiguration.bPrintDebugInfo)
-            {
-                PrintSDKInfo();
-            }
-        }
-
-        public void PrintSDKInfo()
-        {
-
-
-            if (HasRequiredSDKsInstalled() == SDKStatus.Valid)
-            {
-                bool bHasRequiredManualSDK = HasRequiredManualSDK() == SDKStatus.Valid;
-                if (HasSetupAutoSDK())
-                {
-                    string PlatformSDKRoot = GetPathToPlatformAutoSDKs();
-                    Console.WriteLine(GetSDKTargetPlatformName() + " using SDK from: " + Path.Combine(PlatformSDKRoot, GetRequiredSDKString()));
-                }
-                else if (bHasRequiredManualSDK)
-                {
-                    Console.WriteLine(this.ToString() + " using manually installed SDK " + GetRequiredSDKString());
-                }
-                else
-                {
-                    Console.WriteLine(this.ToString() + " setup error.  Inform platform team.");
-                }
-            }
-            else
-            {
-                Console.WriteLine(this.ToString() + " has no valid SDK");
-            }
-        }
-
-        #endregion
-    }
+			if (bIncludeDebug)
+			{
+				Configurations.Insert(0, UnrealTargetConfiguration.Debug);
+			}
+
+			return Configurations;
+		}
+
+		/**
+		 *	Setup the binaries for this specific platform.
+		 *	
+		 *	@param	InBuildTarget		The target being built
+		 */
+		public virtual void SetupBinaries(UEBuildTarget InBuildTarget)
+		{
+		}
+	}
+
+	// AutoSDKs handling portion
+	public abstract partial class UEBuildPlatform : IUEBuildPlatform
+	{
+
+		#region protected AutoSDKs Utility
+
+		/** Name of the file that holds currently install SDK version string */
+		protected static string CurrentlyInstalledSDKStringManifest = "CurrentlyInstalled.txt";
+
+		/** name of the file that holds the last succesfully run SDK setup script version */
+		protected static string LastRunScriptVersionManifest = "CurrentlyInstalled.Version.txt";
+
+		/** Name of the file that holds environment variables of current SDK */
+		protected static string SDKEnvironmentVarsFile = "OutputEnvVars.txt";
+
+		protected static readonly string SDKRootEnvVar = "UE_SDKS_ROOT";
+
+		protected static string AutoSetupEnvVar = "AutoSDKSetup";
+
+		/** 
+		 * Whether platform supports switching SDKs during runtime
+		 * 
+		 * @return true if supports
+		 */
+		protected virtual bool PlatformSupportsAutoSDKs()
+		{
+			return false;
+		}
+
+		static private bool bCheckedAutoSDKRootEnvVar = false;
+		static private bool bAutoSDKSystemEnabled = false;
+		static private bool HasAutoSDKSystemEnabled()
+		{
+			if (!bCheckedAutoSDKRootEnvVar)
+			{
+				string SDKRoot = Environment.GetEnvironmentVariable(SDKRootEnvVar);
+				if (SDKRoot != null)
+				{
+					bAutoSDKSystemEnabled = true;
+				}
+				bCheckedAutoSDKRootEnvVar = true;
+			}
+			return bAutoSDKSystemEnabled;
+		}
+
+		// Whether AutoSDK setup is safe. AutoSDKs will damage manual installs on some platforms.
+		protected bool IsAutoSDKSafe()
+		{
+			return !IsAutoSDKDestructive() || !HasAnyManualInstall();
+		}
+
+		/** 
+		 * Returns SDK string as required by the platform 
+		 * 
+		 * @return Valid SDK string
+		 */
+		protected virtual string GetRequiredSDKString()
+		{
+			return "";
+		}
+
+		/**
+		* Gets the version number of the SDK setup script itself.  The version in the base should ALWAYS be the master revision from the last refactor.  
+		* If you need to force a rebuild for a given platform, override this for the given platform.
+		* @return Setup script version
+		*/
+		protected virtual String GetRequiredScriptVersionString()
+		{
+			return "3.0";
+		}
+
+		/** 
+		 * Returns path to platform SDKs
+		 * 
+		 * @return Valid SDK string
+		 */
+		protected string GetPathToPlatformAutoSDKs()
+		{
+			string SDKPath = "";
+			string SDKRoot = Environment.GetEnvironmentVariable(SDKRootEnvVar);
+			if (SDKRoot != null)
+			{
+				if (SDKRoot != "")
+				{
+					SDKPath = Path.Combine(SDKRoot, "Host" + BuildHostPlatform.Current.Platform, GetSDKTargetPlatformName());
+				}
+			}
+			return SDKPath;
+		}
+
+		/**
+		 * Because most ManualSDK determination depends on reading env vars, if this process is spawned by a process that ALREADY set up
+		 * AutoSDKs then all the SDK env vars will exist, and we will spuriously detect a Manual SDK. (children inherit the environment of the parent process).
+		 * Therefore we write out an env var to set in the command file (OutputEnvVars.txt) such that child processes can determine if their manual SDK detection
+		 * is bogus.  Make it platform specific so that platforms can be in different states.
+		 */
+		protected string GetPlatformAutoSDKSetupEnvVar()
+		{
+			return GetSDKTargetPlatformName() + AutoSetupEnvVar;
+		}
+
+		/**
+		 * Gets currently installed version
+		 * 
+		 * @param PlatformSDKRoot absolute path to platform SDK root
+		 * @param OutInstalledSDKVersionString version string as currently installed
+		 * 
+		 * @return true if was able to read it
+		 */
+		protected bool GetCurrentlyInstalledSDKString(string PlatformSDKRoot, out string OutInstalledSDKVersionString)
+		{
+			if (Directory.Exists(PlatformSDKRoot))
+			{
+				string VersionFilename = Path.Combine(PlatformSDKRoot, CurrentlyInstalledSDKStringManifest);
+				if (File.Exists(VersionFilename))
+				{
+					using (StreamReader Reader = new StreamReader(VersionFilename))
+					{
+						string Version = Reader.ReadLine();
+						string Type = Reader.ReadLine();
+
+						// don't allow ManualSDK installs to count as an AutoSDK install version.
+						if (Type != null && Type == "AutoSDK")
+						{
+							if (Version != null)
+							{
+								OutInstalledSDKVersionString = Version;
+								return true;
+							}
+						}
+					}
+				}
+			}
+
+			OutInstalledSDKVersionString = "";
+			return false;
+		}
+
+		/**
+		 * Gets the version of the last successfully run setup script.
+		 * 
+		 * @param PlatformSDKRoot absolute path to platform SDK root
+		 * @param OutLastRunScriptVersion version string
+		 * 
+		 * @return true if was able to read it
+		 */
+		protected bool GetLastRunScriptVersionString(string PlatformSDKRoot, out string OutLastRunScriptVersion)
+		{
+			if (Directory.Exists(PlatformSDKRoot))
+			{
+				string VersionFilename = Path.Combine(PlatformSDKRoot, LastRunScriptVersionManifest);
+				if (File.Exists(VersionFilename))
+				{
+					using (StreamReader Reader = new StreamReader(VersionFilename))
+					{
+						string Version = Reader.ReadLine();
+						if (Version != null)
+						{
+							OutLastRunScriptVersion = Version;
+							return true;
+						}
+					}
+				}
+			}
+
+			OutLastRunScriptVersion = "";
+			return false;
+		}
+
+		/**
+		 * Sets currently installed version
+		 * 
+		 * @param PlatformSDKRoot absolute path to platform SDK root
+		 * @param InstalledSDKVersionString SDK version string to set
+		 * 
+		 * @return true if was able to set it
+		 */
+		protected bool SetCurrentlyInstalledAutoSDKString(String InstalledSDKVersionString)
+		{
+			String PlatformSDKRoot = GetPathToPlatformAutoSDKs();
+			if (Directory.Exists(PlatformSDKRoot))
+			{
+				string VersionFilename = Path.Combine(PlatformSDKRoot, CurrentlyInstalledSDKStringManifest);
+				if (File.Exists(VersionFilename))
+				{
+					File.Delete(VersionFilename);
+				}
+
+				using (StreamWriter Writer = File.CreateText(VersionFilename))
+				{
+					Writer.WriteLine(InstalledSDKVersionString);
+					Writer.WriteLine("AutoSDK");
+					return true;
+				}
+			}
+
+			return false;
+		}
+
+		protected void SetupManualSDK()
+		{
+			if (PlatformSupportsAutoSDKs() && HasAutoSDKSystemEnabled())
+			{
+				String InstalledSDKVersionString = GetRequiredSDKString();
+				String PlatformSDKRoot = GetPathToPlatformAutoSDKs();
+				if (Directory.Exists(PlatformSDKRoot))
+				{
+					string VersionFilename = Path.Combine(PlatformSDKRoot, CurrentlyInstalledSDKStringManifest);
+					if (File.Exists(VersionFilename))
+					{
+						File.Delete(VersionFilename);
+					}
+
+					string EnvVarFile = Path.Combine(PlatformSDKRoot, SDKEnvironmentVarsFile);
+					if (File.Exists(EnvVarFile))
+					{
+						File.Delete(EnvVarFile);
+					}
+
+					using (StreamWriter Writer = File.CreateText(VersionFilename))
+					{
+						Writer.WriteLine(InstalledSDKVersionString);
+						Writer.WriteLine("ManualSDK");
+					}					
+				}
+			}
+		}
+
+		protected bool SetLastRunAutoSDKScriptVersion(string LastRunScriptVersion)
+		{
+			String PlatformSDKRoot = GetPathToPlatformAutoSDKs();
+			if (Directory.Exists(PlatformSDKRoot))
+			{
+				string VersionFilename = Path.Combine(PlatformSDKRoot, LastRunScriptVersionManifest);
+				if (File.Exists(VersionFilename))
+				{
+					File.Delete(VersionFilename);
+				}
+
+				using (StreamWriter Writer = File.CreateText(VersionFilename))
+				{
+					Writer.WriteLine(LastRunScriptVersion);
+					return true;
+				}
+			}
+			return false;
+		}
+
+		/**
+		* Returns Hook names as needed by the platform
+		* (e.g. can be overriden with custom executables or scripts)
+		*
+		* @param Hook Hook type
+		*/
+		protected virtual string GetHookExecutableName(SDKHookType Hook)
+		{
+			if (Hook == SDKHookType.Uninstall)
+			{
+				return "unsetup.bat";
+			}
+
+			return "setup.bat";
+		}
+
+		/**
+		 * Runs install/uninstall hooks for SDK
+		 * 
+		 * @param PlatformSDKRoot absolute path to platform SDK root
+		 * @param SDKVersionString version string to run for (can be empty!)
+		 * @param Hook which one of hooks to run
+		 * @param bHookCanBeNonExistent whether a non-existing hook means failure
+		 * 
+		 * @return true if succeeded
+		 */
+		protected virtual bool RunAutoSDKHooks(string PlatformSDKRoot, string SDKVersionString, SDKHookType Hook, bool bHookCanBeNonExistent = true)
+		{
+			if (!IsAutoSDKSafe())
+			{
+				Console.ForegroundColor = ConsoleColor.Red;
+				Console.WriteLine(GetSDKTargetPlatformName() + " attempted to run SDK hook which could have damaged manual SDK install!");
+				Console.ResetColor();
+
+				return false;
+			}
+			if (SDKVersionString != "")
+			{
+				string SDKDirectory = Path.Combine(PlatformSDKRoot, SDKVersionString);
+				string HookExe = Path.Combine(SDKDirectory, GetHookExecutableName(Hook));
+
+				if (File.Exists(HookExe))
+				{
+					Console.WriteLine("Running {0} hook {1}", Hook, HookExe);
+
+					// run it
+					Process HookProcess = new Process();
+					HookProcess.StartInfo.WorkingDirectory = SDKDirectory;
+					HookProcess.StartInfo.FileName = HookExe;
+					HookProcess.StartInfo.Arguments = "";
+					HookProcess.StartInfo.WindowStyle = ProcessWindowStyle.Hidden;
+
+					//installers may require administrator access to succeed. so run as an admmin.
+					HookProcess.StartInfo.Verb = "runas";
+					HookProcess.Start();
+					HookProcess.WaitForExit();
+
+					if (HookProcess.ExitCode != 0)
+					{
+						Console.WriteLine("Hook exited uncleanly (returned {0}), considering it failed.", HookProcess.ExitCode);
+						return false;
+					}
+
+					return true;
+				}
+				else
+				{
+					Console.WriteLine("File {0} does not exist", HookExe);
+				}
+			}
+			else
+			{
+				Console.WriteLine("Version string is blank for {0}. Can't determine {1} hook.", PlatformSDKRoot, Hook.ToString());
+			}
+
+			return bHookCanBeNonExistent;
+		}
+
+		/**
+		 * Loads environment variables from SDK
+		 * If any commands are added or removed the handling needs to be duplicated in
+		 * TargetPlatformManagerModule.cpp
+		 * 
+		 * @param PlatformSDKRoot absolute path to platform SDK         
+		 * 
+		 * @return true if succeeded
+		 */
+		protected bool SetupEnvironmentFromAutoSDK(string PlatformSDKRoot)
+		{
+			string EnvVarFile = Path.Combine(PlatformSDKRoot, SDKEnvironmentVarsFile);
+			if (File.Exists(EnvVarFile))
+			{
+				using (StreamReader Reader = new StreamReader(EnvVarFile))
+				{
+					List<string> PathAdds = new List<string>();
+					List<string> PathRemoves = new List<string>();
+
+					List<string> EnvVarNames = new List<string>();
+					List<string> EnvVarValues = new List<string>();
+
+					bool bNeedsToWriteAutoSetupEnvVar = true;
+					String PlatformSetupEnvVar = GetPlatformAutoSDKSetupEnvVar();
+					for (; ; )
+					{
+						string VariableString = Reader.ReadLine();
+						if (VariableString == null)
+						{
+							break;
+						}
+
+						string[] Parts = VariableString.Split('=');
+						if (Parts.Length != 2)
+						{
+							Console.WriteLine("Incorrect environment variable declaration:");
+							Console.WriteLine(VariableString);
+							return false;
+						}
+
+						if (String.Compare(Parts[0], "strippath", true) == 0)
+						{
+							PathRemoves.Add(Parts[1]);
+						}
+						else if (String.Compare(Parts[0], "addpath", true) == 0)
+						{
+							PathAdds.Add(Parts[1]);
+						}
+						else
+						{
+							if (String.Compare(Parts[0], PlatformSetupEnvVar) == 0)
+							{
+								bNeedsToWriteAutoSetupEnvVar = false;
+							}
+							// convenience for setup.bat writers.  Trim any accidental whitespace from var names/values.
+							EnvVarNames.Add(Parts[0].Trim());
+							EnvVarValues.Add(Parts[1].Trim());
+						}
+					}
+
+					// don't actually set anything until we successfully validate and read all values in.
+					// we don't want to set a few vars, return a failure, and then have a platform try to
+					// build against a manually installed SDK with half-set env vars.
+					for (int i = 0; i < EnvVarNames.Count; ++i)
+					{
+						string EnvVarName = EnvVarNames[i];
+						string EnvVarValue = EnvVarValues[i];
+						if (BuildConfiguration.bPrintDebugInfo)
+						{
+							Console.WriteLine("Setting variable '{0}' to '{1}'", EnvVarName, EnvVarValue);
+						}
+						Environment.SetEnvironmentVariable(EnvVarName, EnvVarValue);
+					}
+
+
+					// actually perform the PATH stripping / adding.
+					String OrigPathVar = Environment.GetEnvironmentVariable("PATH");
+					String PathDelimiter = GetPathVarDelimiter();
+					String[] PathVars = OrigPathVar.Split(PathDelimiter.ToCharArray());
+
+					List<String> ModifiedPathVars = new List<string>();
+					ModifiedPathVars.AddRange(PathVars);
+
+					// perform removes first, in case they overlap with any adds.
+					foreach (String PathRemove in PathRemoves)
+					{
+						foreach (String PathVar in PathVars)
+						{
+							if (PathVar.IndexOf(PathRemove, StringComparison.OrdinalIgnoreCase) >= 0)
+							{
+								if (BuildConfiguration.bPrintDebugInfo)
+								{
+									Console.WriteLine("Removing Path: '{0}'", PathVar);
+								}
+								ModifiedPathVars.Remove(PathVar);
+							}
+						}
+					}
+
+					// remove all the of ADDs so that if this function is executed multiple times, the paths will be guarateed to be in the same order after each run.
+					// If we did not do this, a 'remove' that matched some, but not all, of our 'adds' would cause the order to change.
+					foreach (String PathAdd in PathAdds)
+					{
+						foreach (String PathVar in PathVars)
+						{
+							if (String.Compare(PathAdd, PathVar, true) == 0)
+							{
+								if (BuildConfiguration.bPrintDebugInfo)
+								{
+									Console.WriteLine("Removing Path: '{0}'", PathVar);
+								}
+								ModifiedPathVars.Remove(PathVar);
+							}
+						}
+					}
+
+					// perform adds, but don't add duplicates
+					foreach (String PathAdd in PathAdds)
+					{
+						if (!ModifiedPathVars.Contains(PathAdd))
+						{
+							if (BuildConfiguration.bPrintDebugInfo)
+							{
+								Console.WriteLine("Adding Path: '{0}'", PathAdd);
+							}
+							ModifiedPathVars.Add(PathAdd);
+						}
+					}
+
+					String ModifiedPath = String.Join(PathDelimiter, ModifiedPathVars);
+					Environment.SetEnvironmentVariable("PATH", ModifiedPath);
+
+					Reader.Close();
+
+					// write out env var command so any process using this commandfile will mark itself as having had autosdks set up.
+					// avoids child processes spuriously detecting manualsdks.
+					if (bNeedsToWriteAutoSetupEnvVar)
+					{
+						using (StreamWriter Writer = File.AppendText(EnvVarFile))
+						{
+							Writer.WriteLine("{0}=1", PlatformSetupEnvVar);
+						}
+						// set the var in the local environment in case this process spawns any others.
+						Environment.SetEnvironmentVariable(PlatformSetupEnvVar, "1");
+					}
+
+					// make sure we know that we've modified the local environment, invalidating manual installs for this run.
+					bLocalProcessSetupAutoSDK = true;
+
+					return true;
+				}
+			}
+			else
+			{
+				Console.WriteLine("Cannot set up environment for {1} because command file {1} does not exist.", PlatformSDKRoot, EnvVarFile);
+			}
+
+			return false;
+		}
+
+		protected void InvalidateCurrentlyInstalledAutoSDK()
+		{
+			String PlatformSDKRoot = GetPathToPlatformAutoSDKs();
+			if (Directory.Exists(PlatformSDKRoot))
+			{
+				string SDKFilename = Path.Combine(PlatformSDKRoot, CurrentlyInstalledSDKStringManifest);
+				if (File.Exists(SDKFilename))
+				{
+					File.Delete(SDKFilename);
+				}
+
+				string VersionFilename = Path.Combine(PlatformSDKRoot, LastRunScriptVersionManifest);
+				if (File.Exists(VersionFilename))
+				{
+					File.Delete(VersionFilename);
+				}
+
+				string EnvVarFile = Path.Combine(PlatformSDKRoot, SDKEnvironmentVarsFile);
+				if (File.Exists(EnvVarFile))
+				{
+					File.Delete(EnvVarFile);
+				}
+			}
+		}
+
+		/** 
+		 * Currently installed AutoSDK is written out to a text file in a known location.  
+		 * This function just compares the file's contents with the current requirements.
+		 */
+		protected SDKStatus HasRequiredAutoSDKInstalled()
+		{
+			if (PlatformSupportsAutoSDKs() && HasAutoSDKSystemEnabled())
+			{
+				string AutoSDKRoot = GetPathToPlatformAutoSDKs();
+				if (AutoSDKRoot != "")
+				{
+					// check script version so script fixes can be propagated without touching every build machine's CurrentlyInstalled file manually.
+					bool bScriptVersionMatches = false;
+					string CurrentScriptVersionString;
+					if (GetLastRunScriptVersionString(AutoSDKRoot, out CurrentScriptVersionString) && CurrentScriptVersionString == GetRequiredScriptVersionString())
+					{
+						bScriptVersionMatches = true;
+					}
+
+					// check to make sure OutputEnvVars doesn't need regenerating
+					string EnvVarFile = Path.Combine(AutoSDKRoot, SDKEnvironmentVarsFile);
+					bool bEnvVarFileExists = File.Exists(EnvVarFile);
+
+					string CurrentSDKString;
+					if (bEnvVarFileExists && GetCurrentlyInstalledSDKString(AutoSDKRoot, out CurrentSDKString) && CurrentSDKString == GetRequiredSDKString() && bScriptVersionMatches)
+					{
+						return SDKStatus.Valid;
+					}
+					return SDKStatus.Invalid;
+				}
+			}
+			return SDKStatus.Invalid;
+		}
+
+		// This tracks if we have already checked the sdk installation.
+		private Int32 SDKCheckStatus = -1;
+
+		// true if we've ever overridden the process's environment with AutoSDK data.  After that, manual installs cannot be considered valid ever again.
+		private bool bLocalProcessSetupAutoSDK = false;
+
+		protected bool HasSetupAutoSDK()
+		{
+			return bLocalProcessSetupAutoSDK || HasParentProcessSetupAutoSDK();
+		}
+
+		protected bool HasParentProcessSetupAutoSDK()
+		{
+			bool bParentProcessSetupAutoSDK = false;
+			String AutoSDKSetupVarName = GetPlatformAutoSDKSetupEnvVar();
+			String AutoSDKSetupVar = Environment.GetEnvironmentVariable(AutoSDKSetupVarName);
+			if (!String.IsNullOrEmpty(AutoSDKSetupVar))
+			{
+				bParentProcessSetupAutoSDK = true;
+			}
+			return bParentProcessSetupAutoSDK;
+		}
+
+		protected SDKStatus HasRequiredManualSDK()
+		{
+			if (HasSetupAutoSDK())
+			{
+				return SDKStatus.Invalid;
+			}
+
+			// manual installs are always invalid if we have modified the process's environment for AutoSDKs
+			return HasRequiredManualSDKInternal();
+		}
+
+		// for platforms with destructive AutoSDK.  Report if any manual sdk is installed that may be damaged by an autosdk.
+		protected virtual bool HasAnyManualInstall()
+		{
+			return false;
+		}
+
+		// tells us if the user has a valid manual install.
+		protected abstract SDKStatus HasRequiredManualSDKInternal();
+
+		// some platforms will fail if there is a manual install that is the WRONG manual install.
+		protected virtual bool AllowInvalidManualInstall()
+		{
+			return true;
+		}
+
+		// platforms can choose if they prefer a correct the the AutoSDK install over the manual install.
+		protected virtual bool PreferAutoSDK()
+		{
+			return true;
+		}
+
+		// some platforms don't support parallel SDK installs.  AutoSDK on these platforms will
+		// actively damage an existing manual install by overwriting files in it.  AutoSDK must NOT
+		// run any setup if a manual install exists in this case.
+		protected virtual bool IsAutoSDKDestructive()
+		{
+			return false;
+		}
+
+		/**
+		* Runs batch files if necessary to set up required AutoSDK.
+		* AutoSDKs are SDKs that have not been setup through a formal installer, but rather come from
+		* a source control directory, or other local copy.
+		*/
+		private void SetupAutoSDK()
+		{
+			if (IsAutoSDKSafe() && PlatformSupportsAutoSDKs() && HasAutoSDKSystemEnabled())
+			{
+				// run installation for autosdk if necessary.
+				if (HasRequiredAutoSDKInstalled() == SDKStatus.Invalid)
+				{
+					//reset check status so any checking sdk status after the attempted setup will do a real check again.
+					SDKCheckStatus = -1;
+
+					string AutoSDKRoot = GetPathToPlatformAutoSDKs();
+					string CurrentSDKString;
+					GetCurrentlyInstalledSDKString(AutoSDKRoot, out CurrentSDKString);
+
+					// switch over (note that version string can be empty)
+					if (!RunAutoSDKHooks(AutoSDKRoot, CurrentSDKString, SDKHookType.Uninstall))
+					{
+						Console.WriteLine("Failed to uninstall currently installed SDK {0}", CurrentSDKString);
+						InvalidateCurrentlyInstalledAutoSDK();
+						return;
+					}
+					// delete Manifest file to avoid multiple uninstalls
+					InvalidateCurrentlyInstalledAutoSDK();
+
+					if (!RunAutoSDKHooks(AutoSDKRoot, GetRequiredSDKString(), SDKHookType.Install, false))
+					{
+						Console.WriteLine("Failed to install required SDK {0}.  Attemping to uninstall", GetRequiredSDKString());
+						RunAutoSDKHooks(AutoSDKRoot, GetRequiredSDKString(), SDKHookType.Uninstall, false);
+						return;
+					}
+
+					string EnvVarFile = Path.Combine(AutoSDKRoot, SDKEnvironmentVarsFile);
+					if (!File.Exists(EnvVarFile))
+					{
+						Console.WriteLine("Installation of required SDK {0}.  Did not generate Environment file {1}", GetRequiredSDKString(), EnvVarFile);
+						RunAutoSDKHooks(AutoSDKRoot, GetRequiredSDKString(), SDKHookType.Uninstall, false);
+						return;
+					}
+
+					SetCurrentlyInstalledAutoSDKString(GetRequiredSDKString());
+					SetLastRunAutoSDKScriptVersion(GetRequiredScriptVersionString());
+				}
+
+				// fixup process environment to match autosdk
+				SetupEnvironmentFromAutoSDK();
+			}
+		}
+
+		#endregion
+
+		#region public AutoSDKs Utility
+
+		/** Enum describing types of hooks a platform SDK can have */
+		public enum SDKHookType
+		{
+			Install,
+			Uninstall
+		};
+
+		/** 
+		 * Returns platform-specific name used in SDK repository
+		 * 
+		 * @return path to SDK Repository
+		 */
+		public virtual string GetSDKTargetPlatformName()
+		{
+			return "";
+		}
+
+		/* Whether or not we should try to automatically switch SDKs when asked to validate the platform's SDK state. */
+		public static bool bAllowAutoSDKSwitching = true;
+
+		public SDKStatus SetupEnvironmentFromAutoSDK()
+		{
+			string PlatformSDKRoot = GetPathToPlatformAutoSDKs();
+
+			// load environment variables from current SDK
+			if (!SetupEnvironmentFromAutoSDK(PlatformSDKRoot))
+			{
+				Console.WriteLine("Failed to load environment from required SDK {0}", GetRequiredSDKString());
+				InvalidateCurrentlyInstalledAutoSDK();
+				return SDKStatus.Invalid;
+			}
+			return SDKStatus.Valid;
+		}
+
+		/**
+		 *	Whether the required external SDKs are installed for this platform.  
+		 *	Could be either a manual install or an AutoSDK.
+		 */
+		public SDKStatus HasRequiredSDKsInstalled()
+		{
+			// avoid redundant potentially expensive SDK checks.
+			if (SDKCheckStatus == -1)
+			{
+				bool bHasManualSDK = HasRequiredManualSDK() == SDKStatus.Valid;
+				bool bHasAutoSDK = HasRequiredAutoSDKInstalled() == SDKStatus.Valid;
+
+				// Per-Platform implementations can choose how to handle non-Auto SDK detection / handling.
+				SDKCheckStatus = (bHasManualSDK || bHasAutoSDK) ? 1 : 0;
+			}
+			return SDKCheckStatus == 1 ? SDKStatus.Valid : SDKStatus.Invalid;
+		}
+
+		// Arbitrates between manual SDKs and setting up AutoSDK based on program options and platform preferences.
+		public void ManageAndValidateSDK()
+		{
+			// do not modify installed manifests if parent process has already set everything up.
+			// this avoids problems with determining IsAutoSDKSafe and doing an incorrect invalidate.
+			if (bAllowAutoSDKSwitching && !HasParentProcessSetupAutoSDK())
+			{
+				bool bSetSomeSDK = false;
+				bool bHasRequiredManualSDK = HasRequiredManualSDK() == SDKStatus.Valid;
+				if (IsAutoSDKSafe() && (PreferAutoSDK() || !bHasRequiredManualSDK))
+				{
+					SetupAutoSDK();
+					bSetSomeSDK = true;
+				}
+
+				//Setup manual SDK if autoSDK setup was skipped or failed for whatever reason.
+				if (bHasRequiredManualSDK && (HasRequiredAutoSDKInstalled() != SDKStatus.Valid))
+				{
+					SetupManualSDK();
+					bSetSomeSDK = true;
+				}
+
+				if (!bSetSomeSDK)
+				{
+					InvalidateCurrentlyInstalledAutoSDK();
+				}
+			}
+
+			if (BuildConfiguration.bPrintDebugInfo)
+			{
+				PrintSDKInfo();
+			}
+		}
+
+		public void PrintSDKInfo()
+		{
+
+
+			if (HasRequiredSDKsInstalled() == SDKStatus.Valid)
+			{
+				bool bHasRequiredManualSDK = HasRequiredManualSDK() == SDKStatus.Valid;
+				if (HasSetupAutoSDK())
+				{
+					string PlatformSDKRoot = GetPathToPlatformAutoSDKs();
+					Console.WriteLine(GetSDKTargetPlatformName() + " using SDK from: " + Path.Combine(PlatformSDKRoot, GetRequiredSDKString()));
+				}
+				else if (bHasRequiredManualSDK)
+				{
+					Console.WriteLine(this.ToString() + " using manually installed SDK " + GetRequiredSDKString());
+				}
+				else
+				{
+					Console.WriteLine(this.ToString() + " setup error.  Inform platform team.");
+				}
+			}
+			else
+			{
+				Console.WriteLine(this.ToString() + " has no valid SDK");
+			}
+		}
+
+		#endregion
+	}
 }
