- Fixed the Android manifest generation

diff --git a/Engine/Source/Programs/UnrealBuildTool/Android/UEBuildAndroid.cs b/Engine/Source/Programs/UnrealBuildTool/Android/UEBuildAndroid.cs
--- a/Engine/Source/Programs/UnrealBuildTool/Android/UEBuildAndroid.cs
+++ b/Engine/Source/Programs/UnrealBuildTool/Android/UEBuildAndroid.cs
@@ -351,17 +351,18 @@ namespace UnrealBuildTool
 
 		public override void SetupBinaries(UEBuildTarget InBuildTarget)
 		{
-			// dangerously fast mode doesn't generate stub files
-			if (!IOSToolChain.bUseDangerouslyFastMode)
+			// for "fat" binaries, we need to expand the .so's across the platforms
+			string[] Architectures = AndroidToolChain.GetAllArchitectures();
+
+			// all .so's have architecture name, the manifest won't by default, so expand across the binaries
+			List<UEBuildBinary> NewBinaries = new List<UEBuildBinary>();
+			foreach (string Architecture in Architectures)
 			{
-				List<UEBuildBinary> NewBinaries = new List<UEBuildBinary>();
-				// add the .apk to the binaries
 				foreach (var Binary in InBuildTarget.AppBinaries)
 				{
-					// make a binary that just points to the .stub of this executable
 					UEBuildBinaryConfiguration NewConfig = new UEBuildBinaryConfiguration(
 																	InType: Binary.Config.Type,
-																	InOutputFilePath: Path.ChangeExtension(Binary.Config.OutputFilePath, ".apk"),
+																	InOutputFilePath: AndroidToolChain.InlineArchName(Binary.Config.OutputFilePath, Architecture),
 																	InIntermediateDirectory: Binary.Config.IntermediateDirectory,
 																	bInCreateImportLibrarySeparately: Binary.Config.bCreateImportLibrarySeparately,
 																	bInAllowExports: Binary.Config.bAllowExports,
@@ -369,38 +370,26 @@ namespace UnrealBuildTool
 
 					NewBinaries.Add(new UEStubDummyBinary(InBuildTarget, NewConfig));
 				}
-
-				InBuildTarget.AppBinaries.AddRange(NewBinaries);
 			}
-		}
-	}
 
-	/// <summary>
-	/// A .stub that has the executable and metadata included
-	/// Note that this doesn't actually build anything. It just makes the .apk get checked in
-	/// THis could be shared with IOS, etc
-	/// </summary>
-	class UEApkDummyBinary : UEBuildBinary
-	{
-		/// <summary>
-		/// Create an instance initialized to the given configuration
-		/// </summary>
-		/// <param name="InConfig">The build binary configuration to initialize the instance to</param>
-		public UEApkDummyBinary(UEBuildTarget InTarget, UEBuildBinaryConfiguration InConfig)
-			: base(InTarget, InConfig)
-		{
-		}
 
-		/// <summary>
-		/// Builds the binary.
-		/// </summary>
-		/// <param name="CompileEnvironment">The environment to compile the binary in</param>
-		/// <param name="LinkEnvironment">The environment to link the binary in</param>
-		/// <returns></returns>
-		public override IEnumerable<FileItem> Build(IUEToolChain ToolChain, CPPEnvironment CompileEnvironment, LinkEnvironment LinkEnvironment)
-		{
-			// generate the .apk!
-			return new FileItem[] { FileItem.GetItemByPath(Config.OutputFilePath) };
+			// add the .apk to the binaries (UBT always generates the apk without an extension with all .so's inside)
+			foreach (var Binary in InBuildTarget.AppBinaries)
+			{
+				// make a binary that just points to the .stub of this executable
+				UEBuildBinaryConfiguration NewConfig = new UEBuildBinaryConfiguration(
+																InType: Binary.Config.Type,
+																InOutputFilePath: Path.ChangeExtension(Binary.Config.OutputFilePath, ".apk"),
+																InIntermediateDirectory: Binary.Config.IntermediateDirectory,
+																bInCreateImportLibrarySeparately: Binary.Config.bCreateImportLibrarySeparately,
+																bInAllowExports: Binary.Config.bAllowExports,
+																InModuleNames: Binary.Config.ModuleNames);
+
+				NewBinaries.Add(new UEStubDummyBinary(InBuildTarget, NewConfig));
+			}
+
+			// replace the old binaries with the new
+			InBuildTarget.AppBinaries = NewBinaries;
 		}
 	}
 }
