Initial Sunstone ThirdParty setup.

diff --git a/Engine/Source/Programs/UnrealBuildTool/Android/AndroidToolChain.cs b/Engine/Source/Programs/UnrealBuildTool/Android/AndroidToolChain.cs
--- a/Engine/Source/Programs/UnrealBuildTool/Android/AndroidToolChain.cs
+++ b/Engine/Source/Programs/UnrealBuildTool/Android/AndroidToolChain.cs
@@ -227,6 +227,7 @@ namespace UnrealBuildTool
 			Result += " -Wno-reorder";					// member initialization order
 			Result += " -Wno-unknown-pragmas";			// probably should kill this one, sign of another issue in PhysX?
 			Result += " -Wno-invalid-offsetof";			// needed to suppress warnings about using offsetof on non-POD types.
+			Result += " -Wno-logical-op-parentheses";	// needed for external headers we can't change
 
 			// shipping builds will cause this warning with "ensure", so disable only in those case
 			if (CompileEnvironment.Config.Target.Configuration == CPPTargetConfiguration.Shipping)
