- Fixed another Android manifest issue

diff --git a/Engine/Source/Programs/UnrealBuildTool/Android/AndroidToolChain.cs b/Engine/Source/Programs/UnrealBuildTool/Android/AndroidToolChain.cs
--- a/Engine/Source/Programs/UnrealBuildTool/Android/AndroidToolChain.cs
+++ b/Engine/Source/Programs/UnrealBuildTool/Android/AndroidToolChain.cs
@@ -712,7 +712,16 @@ namespace UnrealBuildTool
 				LinkAction.CommandArguments = LinkEnvironment.Config.bIsBuildingLibrary ? GetArArguments(LinkEnvironment) : GetLinkArguments(LinkEnvironment, Arch);
 
 				// Add the output file as a production of the link action.
-				FileItem OutputFile = FileItem.GetItemByPath(InlineArchName(LinkEnvironment.Config.OutputFilePath, Arch));
+				FileItem OutputFile;
+				// libraries don't have the platform name in it, so inline the name
+				if (LinkEnvironment.Config.bIsBuildingLibrary)
+				{
+					OutputFile = FileItem.GetItemByPath(InlineArchName(LinkEnvironment.Config.OutputFilePath, Arch));
+				}
+				else
+				{
+					OutputFile = FileItem.GetItemByPath(LinkEnvironment.Config.OutputFilePath);
+				}
 				Outputs.Add(OutputFile);
 				LinkAction.ProducedItems.Add(OutputFile);
 				LinkAction.StatusDescription = string.Format("{0}", Path.GetFileName(OutputFile.AbsolutePath));
diff --git a/Engine/Source/Programs/UnrealBuildTool/Android/UEBuildAndroid.cs b/Engine/Source/Programs/UnrealBuildTool/Android/UEBuildAndroid.cs
--- a/Engine/Source/Programs/UnrealBuildTool/Android/UEBuildAndroid.cs
+++ b/Engine/Source/Programs/UnrealBuildTool/Android/UEBuildAndroid.cs
@@ -368,7 +368,7 @@ namespace UnrealBuildTool
 																	bInAllowExports: Binary.Config.bAllowExports,
 																	InModuleNames: Binary.Config.ModuleNames);
 
-					NewBinaries.Add(new UEStubDummyBinary(InBuildTarget, NewConfig));
+					NewBinaries.Add(new UEBuildBinaryCPP(InBuildTarget, NewConfig));
 				}
 			}
 
