Add command line option to disable shared pchs; allows catching compile errors where the correct shared pch header isn't included.

diff --git a/Engine/Source/Programs/UnrealBuildTool/System/UnrealBuildTool.cs b/Engine/Source/Programs/UnrealBuildTool/System/UnrealBuildTool.cs
--- a/Engine/Source/Programs/UnrealBuildTool/System/UnrealBuildTool.cs
+++ b/Engine/Source/Programs/UnrealBuildTool/System/UnrealBuildTool.cs
@@ -1822,6 +1822,11 @@ namespace UnrealBuildTool
                 BuildConfiguration.bUsePCHFiles = false;
             }
 
+            if (Utils.ParseCommandLineFlag(Arguments, "-nosharedpch", out ArgumentIndex))
+            {
+                BuildConfiguration.bUseSharedPCHs = false;
+            }
+
             if (Utils.ParseCommandLineFlag(Arguments, "-skipActionHistory", out ArgumentIndex))
             {
                 BuildConfiguration.bUseActionHistory = false;
