Fix exception in Rocket build; libcurl on Android adds include paths that aren't there.

diff --git a/Engine/Source/Programs/UnrealBuildTool/Configuration/UEBuildTarget.cs b/Engine/Source/Programs/UnrealBuildTool/Configuration/UEBuildTarget.cs
--- a/Engine/Source/Programs/UnrealBuildTool/Configuration/UEBuildTarget.cs
+++ b/Engine/Source/Programs/UnrealBuildTool/Configuration/UEBuildTarget.cs
@@ -1118,12 +1118,15 @@ namespace UnrealBuildTool
 				// Add all the include paths
 				foreach(string IncludePath in AllIncludePaths.Where(x => !x.StartsWith("$(")))
 				{
-					foreach(string IncludeFileName in Directory.EnumerateFiles(IncludePath, "*", SearchOption.AllDirectories))
+					if(Directory.Exists(IncludePath))
 					{
-						string Extension = Path.GetExtension(IncludeFileName).ToLower();
-						if(Extension == ".h" || Extension == ".inl")
+						foreach(string IncludeFileName in Directory.EnumerateFiles(IncludePath, "*", SearchOption.AllDirectories))
 						{
-							FileNames.Add(IncludeFileName);
+							string Extension = Path.GetExtension(IncludeFileName).ToLower();
+							if(Extension == ".h" || Extension == ".inl")
+							{
+								FileNames.Add(IncludeFileName);
+							}
 						}
 					}
 				}
