Temporarily disabled use of precompiled header files on Mac when building remotely from PC

diff --git a/Engine/Source/Programs/UnrealBuildTool/Mac/UEBuildMac.cs b/Engine/Source/Programs/UnrealBuildTool/Mac/UEBuildMac.cs
--- a/Engine/Source/Programs/UnrealBuildTool/Mac/UEBuildMac.cs
+++ b/Engine/Source/Programs/UnrealBuildTool/Mac/UEBuildMac.cs
@@ -158,7 +158,11 @@ namespace UnrealBuildTool
 
 		public override void ValidateBuildConfiguration(CPPTargetConfiguration Configuration, CPPTargetPlatform Platform, bool bCreateDebugInfo)
 		{
-//			BuildConfiguration.bUsePCHFiles = false;
+			if (BuildHostPlatform.Current.Platform != UnrealTargetPlatform.Mac)
+			{
+				// @todo: Temporarily disable precompiled header files when building remotely due to errors
+				BuildConfiguration.bUsePCHFiles = false;
+			}
 			BuildConfiguration.bCheckExternalHeadersForModification = BuildHostPlatform.Current.Platform != UnrealTargetPlatform.Mac;
 			BuildConfiguration.bCheckSystemHeadersForModification = BuildHostPlatform.Current.Platform != UnrealTargetPlatform.Mac;
 			BuildConfiguration.ProcessorCountMultiplier = MacToolChain.GetAdjustedProcessorCountMultiplier();
