UnrealBuildTool: Added additional performance diagnostic logging

diff --git a/Engine/Source/Programs/UnrealBuildTool/System/JunkDeleter.cs b/Engine/Source/Programs/UnrealBuildTool/System/JunkDeleter.cs
--- a/Engine/Source/Programs/UnrealBuildTool/System/JunkDeleter.cs
+++ b/Engine/Source/Programs/UnrealBuildTool/System/JunkDeleter.cs
@@ -15,11 +15,19 @@ namespace UnrealBuildTool
 		/// </summary>
 		static public void DeleteJunk()
 		{
+			var JunkStartTime = DateTime.UtcNow;
+
 			if (UnrealBuildTool.RunningRocket() == false)
 			{
 				List<string> JunkManifest = LoadJunkManifest();
 				DeleteAllJunk(JunkManifest);
 			}
+
+			if( BuildConfiguration.bPrintPerformanceInfo )
+			{ 
+				var JunkTime = (DateTime.UtcNow - JunkStartTime).TotalSeconds;
+				Log.TraceInformation( "DeleteJunk took " + JunkTime + "s" );
+			}
 		}
 
 		/// <summary>
diff --git a/Engine/Source/Programs/UnrealBuildTool/System/RulesCompiler.cs b/Engine/Source/Programs/UnrealBuildTool/System/RulesCompiler.cs
--- a/Engine/Source/Programs/UnrealBuildTool/System/RulesCompiler.cs
+++ b/Engine/Source/Programs/UnrealBuildTool/System/RulesCompiler.cs
@@ -1701,6 +1701,8 @@ namespace UnrealBuildTool
 		public static UEBuildTarget CreateTarget(string TargetName, TargetInfo Target,
 			List<string> InAdditionalDefinitions, string InRemoteRoot, List<OnlyModule> InOnlyModules, bool bInEditorRecompile)
 		{
+			var CreateTargetStartTime = DateTime.UtcNow;
+
 			string TargetFileName;
 			TargetRules RulesObject = CreateTargetRules(TargetName, Target, bInEditorRecompile, out TargetFileName);
 			if (bInEditorRecompile)
@@ -1791,6 +1793,12 @@ namespace UnrealBuildTool
 					break;
 			}
 
+			if( BuildConfiguration.bPrintPerformanceInfo )
+			{ 
+				var CreateTargetTime = (DateTime.UtcNow - CreateTargetStartTime).TotalSeconds;
+				Log.TraceInformation( "CreateTarget for " + TargetName + " took " + CreateTargetTime + "s" );
+			}
+
 			if (BuildTarget == null)
 			{
 				throw new BuildException("Failed to create build target for '{0}'.", TargetName);
diff --git a/Engine/Source/Programs/UnrealBuildTool/System/Telemetry.cs b/Engine/Source/Programs/UnrealBuildTool/System/Telemetry.cs
--- a/Engine/Source/Programs/UnrealBuildTool/System/Telemetry.cs
+++ b/Engine/Source/Programs/UnrealBuildTool/System/Telemetry.cs
@@ -58,6 +58,8 @@ namespace UnrealBuildTool
 		/// </summary>
 		public static void Initialize()
 		{
+			var TelemetryInitStartTime = DateTime.UtcNow;
+
 			// reflect over the assembly and find the first eligible provider class.
 			Provider = 
 				(from reflectionAssembly in AppDomain.CurrentDomain.GetReferencingAssemblies(Assembly.GetExecutingAssembly())
@@ -66,6 +68,12 @@ namespace UnrealBuildTool
 				where T.IsDefined(typeof(ProviderAttribute), true) && T.GetInterfaces().Contains(typeof(IProvider)) && T.IsClass && !T.IsAbstract && T.GetConstructor(Type.EmptyTypes) != null
 				// create the instance
 				select (IProvider)Activator.CreateInstance(T)).FirstOrDefault();
+			
+			if( BuildConfiguration.bPrintPerformanceInfo )
+			{ 
+				var TelemetryInitTime = (DateTime.UtcNow - TelemetryInitStartTime).TotalSeconds;
+				Log.TraceInformation( "Telemetry initialization took " + TelemetryInitTime + "s" );
+			}
 		}
 
 		/// <summary>
diff --git a/Engine/Source/Programs/UnrealBuildTool/System/UnrealBuildTool.cs b/Engine/Source/Programs/UnrealBuildTool/System/UnrealBuildTool.cs
--- a/Engine/Source/Programs/UnrealBuildTool/System/UnrealBuildTool.cs
+++ b/Engine/Source/Programs/UnrealBuildTool/System/UnrealBuildTool.cs
@@ -657,6 +657,8 @@ namespace UnrealBuildTool
                 UProjectInfo.FillProjectInfo();
             }
 
+			var BasicInitStartTime = DateTime.UtcNow;
+
             ECompilationResult Result = ECompilationResult.Succeeded;
 
             // Reset early so we can access BuildConfiguration even before RunUBT() is called
@@ -702,6 +704,7 @@ namespace UnrealBuildTool
 				}
             }
             
+			
             // Don't allow simultaneous execution of Unreal Built Tool. Multi-selection in the UI e.g. causes this and you want serial
             // execution in this case to avoid contention issues with shared produced items.
             bool bCreatedMutex = false;
@@ -997,6 +1000,12 @@ namespace UnrealBuildTool
                     RegisterAllUBTClasses();
                     ProjectFileGenerator.bGenerateProjectFiles = false;
 
+					if( BuildConfiguration.bPrintPerformanceInfo )
+					{ 
+						var BasicInitTime = (DateTime.UtcNow - BasicInitStartTime).TotalSeconds;
+						Log.TraceInformation( "Basic UBT initialization took " + BasicInitTime + "s" );
+					}
+
                     // now that we know the available platforms, we can delete other platforms' junk. if we're only building specific modules from the editor, don't touch anything else (it may be in use).
                     if (!bSpecificModulesOnly && !bIgnoreJunk)
                     {
@@ -1257,6 +1266,10 @@ namespace UnrealBuildTool
         {
             bool bSuccess = true;
 
+			
+			var RunUBTInitStartTime = DateTime.UtcNow;
+
+
             // Reset global configurations
             ActionGraph.ResetAllActions();
 
@@ -1309,15 +1322,31 @@ namespace UnrealBuildTool
                 }
 
 
-                foreach (string[] TargetSetting in TargetSettings)
-                {
-                    var Target = UEBuildTarget.CreateTarget(TargetSetting);	// @todo ubtmake: Optimization: Ideally we don't have to do this in the assembler.  It is pretty slow to scan all rules files.  Maybe we can skip the slow parts internally?  It has to scan for *.Target.cs files, and also Load the rules assembly up. :(
-                    if ((Target == null) && (UEBuildConfiguration.bCleanProject))
-                    {
-                        continue;
-                    }
-                    Targets.Add(Target);
-                }
+				if( BuildConfiguration.bPrintPerformanceInfo )
+				{ 
+					var RunUBTInitTime = (DateTime.UtcNow - RunUBTInitStartTime).TotalSeconds;
+					Log.TraceInformation( "RunUBT initialization took " + RunUBTInitTime + "s" );
+				}
+
+				{ 
+					var TargetInitStartTime = DateTime.UtcNow;
+
+					foreach (string[] TargetSetting in TargetSettings)
+					{
+						var Target = UEBuildTarget.CreateTarget(TargetSetting);	// @todo ubtmake: Optimization: Ideally we don't have to do this in the assembler.  It is pretty slow to scan all rules files.  Maybe we can skip the slow parts internally?  It has to scan for *.Target.cs files, and also Load the rules assembly up. :(
+						if ((Target == null) && (UEBuildConfiguration.bCleanProject))
+						{
+							continue;
+						}
+						Targets.Add(Target);
+					}
+
+					if( BuildConfiguration.bPrintPerformanceInfo )
+					{ 
+						var TargetInitTime = (DateTime.UtcNow - TargetInitStartTime).TotalSeconds;
+						Log.TraceInformation( "Target init took " + TargetInitTime + "s" );
+					}
+				}
 
                 UBTMakefile UBTMakefile = null;
 
@@ -1575,6 +1604,7 @@ namespace UnrealBuildTool
                                     if( TargetUObjectModules.Count > 0 )
                                     {
                                         // Execute the header tool
+										// @todo ubtmake urgent: With a game project, I had a couple of failed builds, then suddenly all generated.cpp's were out of date, and failed to compile (needs repro)
                                         string ModuleInfoFileName = Path.GetFullPath( Path.Combine( Target.ProjectIntermediateDirectory, "UnrealHeaderTool.manifest" ) );
                                         ECompilationResult UHTResult = ECompilationResult.OtherCompilationError;
                                         if (!ExternalExecution.ExecuteHeaderToolIfNecessary(Target, GlobalCompileEnvironment:null, UObjectModules:TargetUObjectModules, ModuleInfoFileName:ModuleInfoFileName, UHTResult:ref UHTResult))
@@ -2126,11 +2156,19 @@ namespace UnrealBuildTool
             { 
                 try
                 {
+					var LoadUBTMakefileStartTime = DateTime.UtcNow;
+
                     using (FileStream Stream = new FileStream(UBTMakefileItem.AbsolutePath, FileMode.Open, FileAccess.Read))
                     {	
                         BinaryFormatter Formatter = new BinaryFormatter();
                         LoadedUBTMakefile = Formatter.Deserialize(Stream) as UBTMakefile;
                     }
+
+					if( BuildConfiguration.bPrintPerformanceInfo )
+					{ 
+						var LoadUBTMakefileTime = (DateTime.UtcNow - LoadUBTMakefileStartTime).TotalSeconds;
+						Log.TraceInformation( "LoadUBTMakefile took " + LoadUBTMakefileTime + "s" );
+					}
                 }
                 catch (Exception Ex)
                 {
