Generate an aggregate UAT project file, which adds references to all the other automation projects. Allows rebuilding and debugging directly from the IDE.

diff --git a/Engine/Source/Programs/AutomationTool/Program.cs b/Engine/Source/Programs/AutomationTool/Program.cs
--- a/Engine/Source/Programs/AutomationTool/Program.cs
+++ b/Engine/Source/Programs/AutomationTool/Program.cs
@@ -67,14 +67,14 @@ namespace AutomationTool
 		}
 	};
 
-	class Program
+	public class Program
 	{
 		// This needs to be static, otherwise SetConsoleCtrlHandler will result in a crash on exit.
 		static ProcessManager.CtrlHandlerDelegate ProgramCtrlHandler = new ProcessManager.CtrlHandlerDelegate(CtrlHandler);
 		static public int ReturnCode = 0;
 
 		[STAThread]
-		static int Main()
+		public static int Main()
 		{
 			var CommandLine = SharedUtils.ParseCommandLine();
 			HostPlatform.Initialize();
diff --git a/Engine/Source/Programs/AutomationToolStub/App.config b/Engine/Source/Programs/AutomationToolStub/App.config
new file mode 100644
--- /dev/null
+++ b/Engine/Source/Programs/AutomationToolStub/App.config
@@ -0,0 +1,6 @@
+﻿<?xml version="1.0" encoding="utf-8" ?>
+<configuration>
+    <startup> 
+        <supportedRuntime version="v4.0" sku=".NETFramework,Version=v4.5" />
+    </startup>
+</configuration>
\ No newline at end of file
diff --git a/Engine/Source/Programs/AutomationToolStub/AutomationToolStubTemplate.csproj b/Engine/Source/Programs/AutomationToolStub/AutomationToolStubTemplate.csproj
new file mode 100644
--- /dev/null
+++ b/Engine/Source/Programs/AutomationToolStub/AutomationToolStubTemplate.csproj
@@ -0,0 +1,70 @@
+﻿<?xml version="1.0" encoding="utf-8"?>
+<Project ToolsVersion="12.0" DefaultTargets="Build" xmlns="http://schemas.microsoft.com/developer/msbuild/2003">
+  <Import Project="$(MSBuildExtensionsPath)\$(MSBuildToolsVersion)\Microsoft.Common.props" Condition="Exists('$(MSBuildExtensionsPath)\$(MSBuildToolsVersion)\Microsoft.Common.props')" />
+  <PropertyGroup>
+    <Configuration Condition=" '$(Configuration)' == '' ">Debug</Configuration>
+    <Platform Condition=" '$(Platform)' == '' ">AnyCPU</Platform>
+    <ProjectGuid>{5C41B176-09E5-44D7-8B25-277854234CA3}</ProjectGuid>
+    <OutputType>Exe</OutputType>
+    <AppDesignerFolder>Properties</AppDesignerFolder>
+    <RootNamespace>AutomationToolStub</RootNamespace>
+    <AssemblyName>AutomationToolStub</AssemblyName>
+    <TargetFrameworkVersion>v4.5</TargetFrameworkVersion>
+    <FileAlignment>512</FileAlignment>
+  </PropertyGroup>
+  <PropertyGroup Condition=" '$(Configuration)|$(Platform)' == 'Debug|AnyCPU' ">
+    <PlatformTarget>AnyCPU</PlatformTarget>
+    <DebugSymbols>true</DebugSymbols>
+    <DebugType>full</DebugType>
+    <Optimize>false</Optimize>
+    <OutputPath>..\..\..\Binaries\DotNET\</OutputPath>
+    <DefineConstants>DEBUG;TRACE</DefineConstants>
+    <ErrorReport>prompt</ErrorReport>
+    <WarningLevel>4</WarningLevel>
+  </PropertyGroup>
+  <PropertyGroup Condition=" '$(Configuration)|$(Platform)' == 'Development|AnyCPU' ">
+    <PlatformTarget>AnyCPU</PlatformTarget>
+    <DebugType>pdbonly</DebugType>
+    <Optimize>true</Optimize>
+    <OutputPath>..\..\..\Binaries\DotNET\</OutputPath>
+    <DefineConstants>TRACE</DefineConstants>
+    <ErrorReport>prompt</ErrorReport>
+    <WarningLevel>4</WarningLevel>
+  </PropertyGroup>
+  <ItemGroup>
+    <Reference Include="System" />
+    <Reference Include="System.Core" />
+    <Reference Include="System.Xml.Linq" />
+    <Reference Include="System.Data.DataSetExtensions" />
+    <Reference Include="Microsoft.CSharp" />
+    <Reference Include="System.Data" />
+    <Reference Include="System.Xml" />
+  </ItemGroup>
+  <ItemGroup>
+    <Compile Include="..\DotNETCommon\MetaData.cs">
+      <Link>Properties\MetaData.cs</Link>
+    </Compile>
+    <Compile Include="Program.cs" />
+    <Compile Include="Properties\AssemblyInfo.cs">
+      <Link>Properties\AssemblyInfo.cs</Link>
+    </Compile>
+  </ItemGroup>
+  <ItemGroup>
+    <None Include="App.config" />
+  </ItemGroup>
+  <ItemGroup>
+    <ProjectReference Include="..\AutomationTool\AutomationTool.csproj">
+      <Project>{a35b0261-d868-475e-a098-eb14ce6dc89f}</Project>
+      <Name>AutomationTool</Name>
+      <Private>false</Private>
+    </ProjectReference>
+  </ItemGroup>
+  <Import Project="$(MSBuildToolsPath)\Microsoft.CSharp.targets" />
+  <!-- To modify your build process, add your task inside one of the targets below and uncomment it. 
+       Other similar extension points exist, see Microsoft.Common.targets.
+  <Target Name="BeforeBuild">
+  </Target>
+  <Target Name="AfterBuild">
+  </Target>
+  -->
+</Project>
\ No newline at end of file
diff --git a/Engine/Source/Programs/AutomationToolStub/Program.cs b/Engine/Source/Programs/AutomationToolStub/Program.cs
new file mode 100644
--- /dev/null
+++ b/Engine/Source/Programs/AutomationToolStub/Program.cs
@@ -0,0 +1,28 @@
+﻿using System;
+using System.Collections.Generic;
+using System.Linq;
+using System.Text;
+using System.Threading.Tasks;
+using System.IO;
+using System.Reflection;
+using AutomationTool;
+
+namespace AutomationToolStub
+{
+	class Program
+	{
+		// This program can be used to build and debug AutomationTool and any automation scripts from the IDE.
+		//
+		// AutomationTool's default behavior is to find and compile assemblies at runtime, which is inconvenient for diagnosing compile errors and causes sharing violations if 
+		// the debugger has already loaded the assemblies.
+		//
+		// UnrealBuildTool creates the AutomationHost project from Engine/Source/Programs/AutomationHostTemplate.csproj while generating project files, and adds explicit references
+		// to the other automation modules in the solution so that the IDE can compile them normally. The -nocompile is set on UAT to prevent it compiling them a second time.
+
+		static void Main(string[] args)
+		{
+			AutomationTool.GlobalCommandLine.NoCompile.Set();
+			AutomationTool.Program.Main();
+		}
+	}
+}
diff --git a/Engine/Source/Programs/AutomationToolStub/Properties/AssemblyInfo.cs b/Engine/Source/Programs/AutomationToolStub/Properties/AssemblyInfo.cs
new file mode 100644
--- /dev/null
+++ b/Engine/Source/Programs/AutomationToolStub/Properties/AssemblyInfo.cs
@@ -0,0 +1,13 @@
+﻿using System.Reflection;
+using System.Runtime.CompilerServices;
+using System.Runtime.InteropServices;
+
+// General Information about an assembly is controlled through the following 
+// set of attributes. Change these attribute values to modify the information
+// associated with an assembly.
+[assembly: AssemblyTitle("AutomationHost")]
+[assembly: AssemblyDescription("")]
+[assembly: AssemblyConfiguration("")]
+
+// The following GUID is for the ID of the typelib if this project is exposed to COM
+[assembly: Guid("03557680-bb4e-4e77-9170-877d4795f468")]
diff --git a/Engine/Source/Programs/UnrealBuildTool/System/Project.cs b/Engine/Source/Programs/UnrealBuildTool/System/Project.cs
--- a/Engine/Source/Programs/UnrealBuildTool/System/Project.cs
+++ b/Engine/Source/Programs/UnrealBuildTool/System/Project.cs
@@ -171,6 +171,14 @@ namespace UnrealBuildTool
 			set;
 		}
 
+		/// Whether this project should be built for all solution targets
+		public bool ShouldBuildForAllSolutionTargets
+		{
+			get;
+			set;
+		}
+
+
 
 
 		/// All of the targets in this project.  All non-stub projects must have at least one target.
@@ -434,6 +442,12 @@ namespace UnrealBuildTool
 			return Result;
 		}
 
+		/** Visualizer for the debugger */
+		public override string ToString()
+		{
+			return RelativeProjectFilePath;
+		}
+
 		/// Map of file paths to files in the project.
         private readonly Dictionary<string, SourceFile> SourceFileMap = new Dictionary<string, SourceFile>(StringComparer.InvariantCultureIgnoreCase);
 
diff --git a/Engine/Source/Programs/UnrealBuildTool/System/ProjectFileGenerator.cs b/Engine/Source/Programs/UnrealBuildTool/System/ProjectFileGenerator.cs
--- a/Engine/Source/Programs/UnrealBuildTool/System/ProjectFileGenerator.cs
+++ b/Engine/Source/Programs/UnrealBuildTool/System/ProjectFileGenerator.cs
@@ -8,6 +8,9 @@ using System.Reflection;
 using Microsoft.Win32;
 using System.Linq;
 using System.Diagnostics;
+using System.Xml;
+using System.Xml.Linq;
+using System.Xml.XPath;
 
 namespace UnrealBuildTool
 {
@@ -214,7 +217,7 @@ namespace UnrealBuildTool
 		/// <summary>
 		/// Adds all .automation.csproj files to the solution.
 		/// </summary>
-		void AddAutomationModules(MasterProjectFolder ProgramsFolder)
+		List<VCSharpProjectFile> AddAutomationModules(MasterProjectFolder ProgramsFolder)
 		{
 			var Folder = ProgramsFolder.AddSubFolder("Automation");
 			var AllGameFolders = UEBuildTarget.DiscoverAllGameFolders();
@@ -229,7 +232,8 @@ namespace UnrealBuildTool
 			}
 
 			// Find all the automation modules .csproj files to add
-			var ModuleFiles = RulesCompiler.FindAllRulesSourceFiles(RulesCompiler.RulesFileType.AutomationModule, BuildFolders);
+			List<VCSharpProjectFile> ModuleProjects = new List<VCSharpProjectFile>();
+			List<string> ModuleFiles = RulesCompiler.FindAllRulesSourceFiles(RulesCompiler.RulesFileType.AutomationModule, BuildFolders);
 			foreach (var ProjectFile in ModuleFiles)
 			{
 				FileInfo Info = new FileInfo(Path.Combine(ProjectFile));
@@ -237,11 +241,94 @@ namespace UnrealBuildTool
 				{
 					var RelativeFileName = Utils.MakePathRelativeTo(ProjectFile, MasterProjectRelativePath);
 					var Project = new VCSharpProjectFile(RelativeFileName);
+					Project.ShouldBuildForAllSolutionTargets = true;
 					AddExistingProjectFile(Project);
+					ModuleProjects.Add(Project);
 
 					Folder.ChildProjects.Add( Project );
 				}
 			}
+			return ModuleProjects;
+		}
+
+		/// <summary>
+		/// Generates a C# project which references all the given automation modules and can be used to launch UAT directly from the IDE.
+		/// </summary>
+		private VCSharpProjectFile GenerateAutomationToolStubProject(List<VCSharpProjectFile> AutomationModules)
+		{
+			string InputProjectDir = Path.GetFullPath("Programs\\AutomationToolStub");
+			string OutputProjectDir = IntermediateProjectFilesPath;
+
+			// Load the template aggregate project file
+			XDocument Document = XDocument.Load(Path.Combine(InputProjectDir, "AutomationToolStubTemplate.csproj"));
+			XNamespace Namespace = Document.Root.Name.Namespace;
+
+			// Update the GUID for this project. Using a fixed GUID shouldn't be a problem, but it should be different from the source.
+			XElement SharedPropertyGroup = Document.Root.Elements().First(x => x.Name.LocalName == "PropertyGroup" && x.Attribute("Condition") == null);
+			if(SharedPropertyGroup != null)
+			{
+				SharedPropertyGroup.SetElementValue(Namespace + "ProjectGuid", "{5C41B176-09E5-44D7-8B25-277854234CA4}");
+			}
+
+			// Create a namespace manager that can deal with msbuild files
+			XmlNamespaceManager NamespaceManager = new XmlNamespaceManager(new NameTable());
+			NamespaceManager.AddNamespace("msbuild", Namespace.NamespaceName);
+
+			// Fixup the output directory
+			foreach(XElement PropertyGroup in Document.Root.Elements(Namespace + "PropertyGroup"))
+			{
+				XElement OutputPath = PropertyGroup.Element(Namespace + "OutputPath");
+				if(OutputPath != null)
+				{
+					OutputPath.Value = Utils.MakePathRelativeTo(Path.Combine(InputProjectDir, OutputPath.Value), OutputProjectDir);
+				}
+			}
+
+			// Fixup all the source files
+			foreach(XElement ItemGroup in Document.Root.Elements(Namespace + "ItemGroup"))
+			{
+				foreach(XElement SourceElement in ItemGroup.Elements().Where(x => x.Name.LocalName == "Compile" || x.Name.LocalName == "None" || x.Name.LocalName == "ProjectReference"))
+				{
+					XAttribute Include = SourceElement.Attribute("Include");
+					if(Include != null)
+					{
+						Include.Value = Utils.MakePathRelativeTo(Path.Combine(InputProjectDir, Include.Value), OutputProjectDir);
+					}
+				}
+			}
+
+			// Add references to all the other AutomationTool modules
+			XElement ProjectReferences = new XElement(Namespace + "ItemGroup");
+			foreach(VCSharpProjectFile AutomationModule in AutomationModules)
+			{
+				string ModuleProjectPath = Path.GetFullPath(Path.Combine(MasterProjectRelativePath, AutomationModule.RelativeProjectFilePath));
+				XDocument ModuleDocument = XDocument.Load(ModuleProjectPath);
+				if(ModuleDocument != null)
+				{
+					XElement ModulePropertyGroup = ModuleDocument.Root.Elements().First(x => x.Name.LocalName == "PropertyGroup" && x.Attribute("Condition") == null);
+					XElement ModuleName = ModulePropertyGroup.Element(Namespace + "AssemblyName");
+					XElement ModuleGuid = ModulePropertyGroup.Element(Namespace + "ProjectGuid");
+					if(ModuleName != null && ModuleGuid != null)
+					{
+						XElement ProjectReference = new XElement(Namespace + "ProjectReference");
+						ProjectReference.SetAttributeValue("Include", Utils.MakePathRelativeTo(ModuleProjectPath, OutputProjectDir, false));
+						ProjectReference.Add(new XElement(Namespace + "Name", ModuleName.Value));
+						ProjectReference.Add(new XElement(Namespace + "Project", ModuleGuid.Value));
+						ProjectReference.Add(new XElement(Namespace + "Private", "false"));
+						ProjectReferences.Add(ProjectReference);
+					}
+				}
+			}
+			Document.Root.Add(ProjectReferences);
+
+			// Save the project
+			string Filename = Path.Combine(OutputProjectDir, "AutomationToolStub.csproj");
+			Document.Save(Filename);
+
+			// Add the project file to the list of projects to build
+			VCSharpProjectFile NewProjectFile = new VCSharpProjectFile(Utils.MakePathRelativeTo(Filename, MasterProjectRelativePath));
+			AddExistingProjectFile(NewProjectFile);
+			return NewProjectFile;
 		}
 
 		/// <summary>
@@ -478,19 +565,22 @@ namespace UnrealBuildTool
 				if( ( !bGeneratingGameProjectFiles && !bGeneratingRocketProjectFiles ) || bAlwaysIncludeEngineModules )
 				{
 					// Add AutomationTool to the master project
-					ProgramsFolder.ChildProjects.Add(AddSimpleCSharpProject("AutomationTool"));
+					ProgramsFolder.ChildProjects.Add(AddSimpleCSharpProject("AutomationTool", true));
 
 					// Add UnrealAutomationTool (launcher) to the master project
 					ProgramsFolder.ChildProjects.Add(AddSimpleCSharpProject("AutomationToolLauncher"));
 
-					// Add automation.csproj files to the master project.
-					AddAutomationModules(ProgramsFolder);
+					// Add automation.csproj files to the master project
+					List<VCSharpProjectFile> AutomationModules = AddAutomationModules(ProgramsFolder);
+
+					// Create the automation tool stub project
+					ProgramsFolder.ChildProjects.Add(GenerateAutomationToolStubProject(AutomationModules));
 
 					// Add Distill to the master project
 					ProgramsFolder.ChildProjects.Add( AddSimpleCSharpProject( "Distill" ) );
 
 					// Add DotNETUtilities to the master project
-					ProgramsFolder.ChildProjects.Add( AddSimpleCSharpProject(  "DotNETCommon/DotNETUtilities" ) );
+					ProgramsFolder.ChildProjects.Add( AddSimpleCSharpProject(  "DotNETCommon/DotNETUtilities", true ) );
 
 					// Add all of the IOS C# projects
 					AddIOSProjects( ProgramsFolder );
@@ -1111,9 +1201,9 @@ namespace UnrealBuildTool
 		/// </summary>
 		/// <param name="ProjectName">Name of project file to add</param>
 		/// <returns>ProjectFile if the operation was successful, otherwise null.</returns>
-		private ProjectFile AddSimpleCSharpProject(string ProjectName)
+		private VCSharpProjectFile AddSimpleCSharpProject(string ProjectName, bool bShouldBuildForAllSolutionTargets = false)
 		{
-			ProjectFile Project = null;
+			VCSharpProjectFile Project = null;
 
 			var ProjectFileName = Path.Combine( EngineRelativePath, "Source", "Programs", ProjectName, Path.GetFileName( ProjectName ) + ".csproj" );
 			FileInfo Info = new FileInfo( ProjectFileName );
@@ -1121,6 +1211,7 @@ namespace UnrealBuildTool
 			{
 				var FileNameRelativeToMasterProject = Utils.MakePathRelativeTo( ProjectFileName, MasterProjectRelativePath );
 				Project = new VCSharpProjectFile(FileNameRelativeToMasterProject);
+				Project.ShouldBuildForAllSolutionTargets = bShouldBuildForAllSolutionTargets;
 				AddExistingProjectFile(Project);
 			}
 			else
@@ -1208,7 +1299,7 @@ namespace UnrealBuildTool
 			if( ProjectFolderInfo.Exists )
 			{
 				Folder.ChildProjects.Add( AddSimpleCSharpProject( "IOS/iPhonePackager" ) );
-				Folder.ChildProjects.Add( AddSimpleCSharpProject( "IOS/DeploymentInterface" ) );
+				Folder.ChildProjects.Add( AddSimpleCSharpProject( "IOS/DeploymentInterface", true ) ); // Build by default; needed for UAT.
 				Folder.ChildProjects.Add( AddSimpleCSharpProject( "IOS/DeploymentServer" ) );
 				Folder.ChildProjects.Add( AddSimpleCSharpProject( "IOS/MobileDeviceInterface" ) );
 			}
diff --git a/Engine/Source/Programs/UnrealBuildTool/System/VCProjectFileGenerator.cs b/Engine/Source/Programs/UnrealBuildTool/System/VCProjectFileGenerator.cs
--- a/Engine/Source/Programs/UnrealBuildTool/System/VCProjectFileGenerator.cs
+++ b/Engine/Source/Programs/UnrealBuildTool/System/VCProjectFileGenerator.cs
@@ -575,6 +575,7 @@ namespace UnrealBuildTool
 											{
 												// UBT gets a pass because it is a dependency of every single configuration combination
 												if( CurProject != UBTProject &&
+													!CurProject.ShouldBuildForAllSolutionTargets &&
 													TargetConfigurationName != TargetRules.TargetType.Game.ToString() )
 												{
 													// Can't build non-generated project in configurations except for the default (Game)
