Fixed AutomationToolStub project generation on Mono

diff --git a/Engine/Source/Programs/UnrealBuildTool/System/ProjectFileGenerator.cs b/Engine/Source/Programs/UnrealBuildTool/System/ProjectFileGenerator.cs
--- a/Engine/Source/Programs/UnrealBuildTool/System/ProjectFileGenerator.cs
+++ b/Engine/Source/Programs/UnrealBuildTool/System/ProjectFileGenerator.cs
@@ -260,7 +260,7 @@ namespace UnrealBuildTool
 			string OutputProjectDir = IntermediateProjectFilesPath;
 
 			// Load the template aggregate project file
-			XDocument Document = XDocument.Load(Path.Combine(InputProjectDir, "AutomationToolStubTemplate.csproj"));
+			XDocument Document = XDocument.Load(Utils.CleanDirectorySeparators(Path.Combine(InputProjectDir, "AutomationToolStubTemplate.csproj")));
 			XNamespace Namespace = Document.Root.Name.Namespace;
 
 			// Update the GUID for this project. Using a fixed GUID shouldn't be a problem, but it should be different from the source.
