UnrealBuildTool: Even faster build iteration with bUseExperimentalFastBuildIteration
- Shaved over *half a second* of the best case iterative build times with UnrealBuildTool
- When assembling a build from an existing UBTMakefile, we now load Target information from the makefile instead of from the C# assembly

diff --git a/Engine/Source/Programs/UnrealBuildTool/Configuration/UEBuildClient.cs b/Engine/Source/Programs/UnrealBuildTool/Configuration/UEBuildClient.cs
--- a/Engine/Source/Programs/UnrealBuildTool/Configuration/UEBuildClient.cs
+++ b/Engine/Source/Programs/UnrealBuildTool/Configuration/UEBuildClient.cs
@@ -7,6 +7,7 @@ using System.Text;
 
 namespace UnrealBuildTool
 {
+	[Serializable]
     public class UEBuildClient : UEBuildTarget
 	{
 		public UEBuildClient(
diff --git a/Engine/Source/Programs/UnrealBuildTool/Configuration/UEBuildEditor.cs b/Engine/Source/Programs/UnrealBuildTool/Configuration/UEBuildEditor.cs
--- a/Engine/Source/Programs/UnrealBuildTool/Configuration/UEBuildEditor.cs
+++ b/Engine/Source/Programs/UnrealBuildTool/Configuration/UEBuildEditor.cs
@@ -7,6 +7,7 @@ using System.Text;
 
 namespace UnrealBuildTool
 {
+	[Serializable]
 	public class UEBuildEditor : UEBuildTarget
 	{
 		public UEBuildEditor(
diff --git a/Engine/Source/Programs/UnrealBuildTool/Configuration/UEBuildGame.cs b/Engine/Source/Programs/UnrealBuildTool/Configuration/UEBuildGame.cs
--- a/Engine/Source/Programs/UnrealBuildTool/Configuration/UEBuildGame.cs
+++ b/Engine/Source/Programs/UnrealBuildTool/Configuration/UEBuildGame.cs
@@ -7,6 +7,7 @@ using System.Text;
 
 namespace UnrealBuildTool
 {
+	[Serializable]
 	public class UEBuildGame : UEBuildTarget
 	{
 		public UEBuildGame(
diff --git a/Engine/Source/Programs/UnrealBuildTool/Configuration/UEBuildModule.cs b/Engine/Source/Programs/UnrealBuildTool/Configuration/UEBuildModule.cs
--- a/Engine/Source/Programs/UnrealBuildTool/Configuration/UEBuildModule.cs
+++ b/Engine/Source/Programs/UnrealBuildTool/Configuration/UEBuildModule.cs
@@ -1950,7 +1950,7 @@ namespace UnrealBuildTool
 								// Is this a plugin module?
 								var PluginInfo = Plugins.GetPluginInfoForModule( DependencyName );
 
-								string OutputFilePath = Target.MakeBinaryPath(DependencyModule.Name, Target.GetAppName() + "-" + DependencyModule.Name, UEBuildBinaryType.DynamicLinkLibrary, Target.Rules.Type, bIsRocketModule, PluginInfo, "");
+								string OutputFilePath = Target.MakeBinaryPath(DependencyModule.Name, Target.GetAppName() + "-" + DependencyModule.Name, UEBuildBinaryType.DynamicLinkLibrary, Target.TargetType, bIsRocketModule, PluginInfo, "");
 
 								// If it's an engine module, output intermediates to the engine intermediates directory. 
 								string IntermediateDirectory = Binary.Config.IntermediateDirectory;
@@ -2022,7 +2022,7 @@ namespace UnrealBuildTool
 		public static string GetGeneratedCodeDirectoryForModule(UEBuildTarget Target, string ModuleDirectory, string ModuleName)
 		{
 			string BaseDirectory = null;
-			if ((Target.ShouldCompileMonolithic() || Target.Rules.Type == TargetRules.TargetType.Program) && 
+			if ((Target.ShouldCompileMonolithic() || Target.TargetType == TargetRules.TargetType.Program) && 
 				(!UnrealBuildTool.BuildingRocket()) &&
 				(!UnrealBuildTool.RunningRocket() || Utils.IsFileUnderDirectory(ModuleDirectory, UnrealBuildTool.GetUProjectPath())))
 			{
diff --git a/Engine/Source/Programs/UnrealBuildTool/Configuration/UEBuildServer.cs b/Engine/Source/Programs/UnrealBuildTool/Configuration/UEBuildServer.cs
--- a/Engine/Source/Programs/UnrealBuildTool/Configuration/UEBuildServer.cs
+++ b/Engine/Source/Programs/UnrealBuildTool/Configuration/UEBuildServer.cs
@@ -7,6 +7,7 @@ using System.Text;
 
 namespace UnrealBuildTool
 {
+	[Serializable]
 	public class UEBuildServer : UEBuildTarget
 	{
 		public UEBuildServer(
diff --git a/Engine/Source/Programs/UnrealBuildTool/Configuration/UEBuildTarget.cs b/Engine/Source/Programs/UnrealBuildTool/Configuration/UEBuildTarget.cs
--- a/Engine/Source/Programs/UnrealBuildTool/Configuration/UEBuildTarget.cs
+++ b/Engine/Source/Programs/UnrealBuildTool/Configuration/UEBuildTarget.cs
@@ -92,8 +92,6 @@ namespace UnrealBuildTool
 		}
 	};
 
-	delegate UEBuildTarget TargetCreationDelegate();
-
 	/// <summary>
 	/// A container for a binary files (dll, exe) with its associated debug info.
 	/// </summary>
@@ -154,6 +152,26 @@ namespace UnrealBuildTool
 		public readonly string OnlyModuleSuffix;
 	}
 
+
+	/// <summary>
+	/// Describes all of the information needed to initialize a UEBuildTarget object
+	/// </summary>
+	public class TargetDescriptor
+	{
+		public string TargetName;
+		public UnrealTargetPlatform Platform;
+		public UnrealTargetConfiguration Configuration;
+		public List<string> AdditionalDefinitions;
+		public bool bIsEditorRecompile;
+		public string RemoteRoot;
+		public List<OnlyModule> OnlyModules;
+	}
+
+
+	/**
+	 * A target that can be built
+	 */
+	[Serializable]
 	public class UEBuildTarget
 	{
 		public string GetAppName()
@@ -185,24 +203,26 @@ namespace UnrealBuildTool
 			throw new BuildException("CPPTargetPlatformToUnrealTargetPlatform: Unknown CPPTargetPlatform {0}", InCPPPlatform.ToString());
 		}
 
-		public static UEBuildTarget CreateTarget(string[] SourceArguments)
+
+		public static List<TargetDescriptor> ParseTargetCommandLine(string[] SourceArguments )
 		{
-			TargetCreationDelegate TargetCreationDelegate = null;
+			var Targets = new List<TargetDescriptor>();
+
+			string TargetName = null;
 			var AdditionalDefinitions = new List<string>();
 			var Platform = UnrealTargetPlatform.Unknown;
 			var Configuration = UnrealTargetConfiguration.Unknown;
 			string RemoteRoot = null;
-
 			var OnlyModules = new List<OnlyModule>();
 
+			// If true, the recompile was launched by the editor.
+			bool bIsEditorRecompile = false;
+
+
 			// Combine the two arrays of arguments
 			List<string> Arguments = new List<string>(SourceArguments.Length);
 			Arguments.AddRange(SourceArguments);
 
-			// If true, the recompile was launched by the editor.
-			// This means we actually have to 
-			bool bIsEditorRecompile = false;
-
 			List<string> PossibleTargetNames = new List<string>();
 			for (int ArgumentIndex = 0; ArgumentIndex < Arguments.Count; ArgumentIndex++)
 			{
@@ -355,19 +375,13 @@ namespace UnrealBuildTool
 				throw new BuildException("Couldn't determine configuration name.");
 			}
 
-			if (TargetCreationDelegate != null)
-			{
-				return TargetCreationDelegate();
-			}
-
 			if (PossibleTargetNames.Count > 0)
 			{
 				// We have possible targets!
-				UEBuildTarget BuildTarget = null;
 				foreach (string PossibleTargetName in PossibleTargetNames)
 				{
 					// If running Rocket, the PossibleTargetName could contain a path
-					var TargetName = PossibleTargetName;
+					TargetName = PossibleTargetName;
 
 					// If a project file was not specified see if we can find one
 					string CheckProjectFile = UProjectInfo.GetProjectForTarget(TargetName);
@@ -396,45 +410,72 @@ namespace UnrealBuildTool
 						}
 					}
 
-					if( !ProjectFileGenerator.bGenerateProjectFiles )
-					{
-						// Configure the rules compiler
-						string PossibleAssemblyName = TargetName;
-						if (bIsEditorRecompile == true)
+					Targets.Add( new TargetDescriptor()
 						{
-							PossibleAssemblyName += "_EditorRecompile";
-						}
+							TargetName = TargetName,
+							Platform = Platform,
+							Configuration = Configuration,
+							AdditionalDefinitions = AdditionalDefinitions,
+							bIsEditorRecompile = bIsEditorRecompile,
+							RemoteRoot = RemoteRoot,
+							OnlyModules = OnlyModules
+						} );
+					break;
+				}
+			}
+			if( Targets.Count == 0 )
+			{
+				throw new BuildException("No target name was specified on the command-line.");
+			}
+			return Targets;
+		}
 
-						// Scan the disk to find source files for all known targets and modules, and generate "in memory" project
-						// file data that will be used to determine what to build
-						RulesCompiler.SetAssemblyNameAndGameFolders( PossibleAssemblyName, UEBuildTarget.DiscoverAllGameFolders() );
-					}
+		public static UEBuildTarget CreateTarget( TargetDescriptor Desc )
+		{
+			string TargetName = Desc.TargetName;
+			List<string> AdditionalDefinitions = Desc.AdditionalDefinitions;
+			UnrealTargetPlatform Platform = Desc.Platform;
+			UnrealTargetConfiguration Configuration = Desc.Configuration;
+			string RemoteRoot = Desc.RemoteRoot;
+			List<OnlyModule> OnlyModules = Desc.OnlyModules;
+			bool bIsEditorRecompile = Desc.bIsEditorRecompile;
+
+			UEBuildTarget BuildTarget = null;
+			if( !ProjectFileGenerator.bGenerateProjectFiles )
+			{
+				// Configure the rules compiler
+				string PossibleAssemblyName = TargetName;
+				if (bIsEditorRecompile == true)
+				{
+					PossibleAssemblyName += "_EditorRecompile";
+				}
 
-					// Try getting it from the RulesCompiler
-					UEBuildTarget Target = RulesCompiler.CreateTarget(
-						TargetName:TargetName, 
-						Target:new TargetInfo(Platform, Configuration),
-						InAdditionalDefinitions:AdditionalDefinitions, 
-						InRemoteRoot:RemoteRoot, 
-						InOnlyModules:OnlyModules, 
-						bInEditorRecompile:bIsEditorRecompile);
-					if (Target == null)
-					{
-						if (UEBuildConfiguration.bCleanProject)
-						{
-							return null;
-						}
-						throw new BuildException( "Couldn't find target name {0}.", TargetName );
-					}
-					else
-					{
-						BuildTarget = Target;
-						break;
-					}
+				// Scan the disk to find source files for all known targets and modules, and generate "in memory" project
+				// file data that will be used to determine what to build
+				RulesCompiler.SetAssemblyNameAndGameFolders( PossibleAssemblyName, UEBuildTarget.DiscoverAllGameFolders() );
+			}
+
+			// Try getting it from the RulesCompiler
+			UEBuildTarget Target = RulesCompiler.CreateTarget(
+				TargetName:TargetName, 
+				Target:new TargetInfo(Platform, Configuration),
+				InAdditionalDefinitions:AdditionalDefinitions, 
+				InRemoteRoot:RemoteRoot, 
+				InOnlyModules:OnlyModules, 
+				bInEditorRecompile:bIsEditorRecompile);
+			if (Target == null)
+			{
+				if (UEBuildConfiguration.bCleanProject)
+				{
+					return null;
 				}
-				return BuildTarget;
+				throw new BuildException( "Couldn't find target name {0}.", TargetName );
+			}
+			else
+			{
+				BuildTarget = Target;
 			}
-			throw new BuildException("No target name was specified on the command-line.");
+			return BuildTarget;
 		}
 
 		/// Parses only the target platform and configuration from the specified command-line argument list
@@ -567,8 +608,19 @@ namespace UnrealBuildTool
 		}
 
 		/** The target rules */
+		[NonSerialized]
 		public TargetRules Rules = null;
 
+		/** Type of target, or null if undetermined (such as in the case of a synthetic target with no TargetRules) */
+		public TargetRules.TargetType TargetType
+		{
+			get
+			{
+				return TargetTypeOrNull.Value;
+			}
+		}
+		private readonly TargetRules.TargetType? TargetTypeOrNull;
+
 		/** The name of the application the target is part of. */
 		public string AppName;
 
@@ -597,30 +649,37 @@ namespace UnrealBuildTool
 		public string RemoteRoot;
 
 		/** The C++ environment that all the environments used to compile UE-based modules are derived from. */
+		[NonSerialized]
 		public CPPEnvironment GlobalCompileEnvironment = new CPPEnvironment();
 
 		/** The link environment all binary link environments are derived from. */
+		[NonSerialized]
 		public LinkEnvironment GlobalLinkEnvironment = new LinkEnvironment();
 
 		/** All plugins enabled for this target */
+		[NonSerialized]
 		public List<PluginInfo> EnabledPlugins = new List<PluginInfo>();
 
 		/** All application binaries; may include binaries not built by this target. */
+		[NonSerialized]
 		public List<UEBuildBinary> AppBinaries = new List<UEBuildBinary>();
 
 		/** Extra engine module names to either include in the binary (monolithic) or create side-by-side DLLs for (modular) */
+		[NonSerialized]
 		public List<string> ExtraModuleNames = new List<string>();
 
 		/** True if re-compiling this target from the editor */
 		public bool bEditorRecompile;
 
 		/** If building only a specific set of modules, these are the modules to build */
+		[NonSerialized]
 		protected List<OnlyModule> OnlyModules = new List<OnlyModule>();
 
 		/** true if target should be compiled in monolithic mode, false if not */
 		protected bool bCompileMonolithic = false;
 
 		/** Used to keep track of all modules by name. */
+		[NonSerialized]
 		private Dictionary<string, UEBuildModule> Modules = new Dictionary<string, UEBuildModule>(StringComparer.InvariantCultureIgnoreCase);
 
 		/// <summary>
@@ -707,10 +766,10 @@ namespace UnrealBuildTool
 
 			OnlyModules = InOnlyModules;
 
-			bool bIsRocketGame = (InRulesObject != null) ? (InRulesObject.Type == TargetRules.TargetType.RocketGame) : false;
+			TargetTypeOrNull = (Rules != null) ? Rules.Type : (TargetRules.TargetType?)null;
+			bool bIsRocketGame = (InRulesObject != null) ? (TargetType == TargetRules.TargetType.RocketGame) : false;
 
 			// Construct the output path based on configuration, platform, game if not specified.
-			TargetRules.TargetType? TargetType = (Rules != null) ? Rules.Type : (TargetRules.TargetType?)null;
             OutputPath = Path.GetFullPath(MakeBinaryPath("", AppName, UEBuildBinaryType.Executable, TargetType, bIsRocketGame, null, InAppName, Configuration == UnrealTargetConfiguration.Shipping ? Rules.ForceNameAsForDevelopment() : false, Rules.ExeBinariesSubFolder));
 
 			if (bCompileMonolithic && TargetRules.IsGameType(InRulesObject.Type))
@@ -863,10 +922,9 @@ namespace UnrealBuildTool
 		/// <param name="Manifest">Manifest</param>
 		protected void CleanTarget(List<UEBuildBinary> Binaries, CPPTargetPlatform Platform, FileManifest Manifest)
 		{
-			if (Rules != null)
 			{
 				var TargetFilename = RulesCompiler.GetTargetFilename(GameName);
-				var LocalTargetName = (Rules.Type == TargetRules.TargetType.Program) ? AppName : GameName;
+				var LocalTargetName = (TargetType == TargetRules.TargetType.Program) ? AppName : GameName;
 
 				Log.TraceVerbose("Cleaning target {0} - AppName {1}", LocalTargetName, AppName);
 				Log.TraceVerbose("\tTargetFilename {0}", TargetFilename);
@@ -1006,10 +1064,17 @@ namespace UnrealBuildTool
 				// Delete the UBT makefile
 				{
 					// Figure out what to call our action graph based on everything we're building
-					var Targets = new List<UEBuildTarget>();
-					Targets.Add( this );	// @todo ubtmake: Only supports cleaning one target at a time :(
+					var TargetDescs = new List<TargetDescriptor>();
+					
+					// @todo ubtmake: Only supports cleaning one target at a time :(
+					TargetDescs.Add( new TargetDescriptor
+						{
+							TargetName = GetTargetName(),
+							Platform = this.Platform,
+							Configuration = this.Configuration
+						} );
 
-					var UBTMakefilePath = UnrealBuildTool.GetUBTMakefilePath( Targets );
+					var UBTMakefilePath = UnrealBuildTool.GetUBTMakefilePath( TargetDescs );
 					if (File.Exists(UBTMakefilePath))
 					{
 						Log.TraceVerbose("\tDeleting " + UBTMakefilePath);
@@ -1034,10 +1099,6 @@ namespace UnrealBuildTool
 					CleanUnrealHeaderTool();
 				}
 			}
-			else
-			{
-				Log.TraceVerbose("Cannot clean target with no Rules object: {0}", GameName);
-			}
 		}
 
 		/** Generates a list of external files which are required to build this target */
@@ -1264,7 +1325,7 @@ namespace UnrealBuildTool
 
                 }
 
-				if ( (TargetRules.IsAGame(Rules.Type) || (Rules.Type == TargetRules.TargetType.Server)) 
+				if ( (TargetRules.IsAGame(TargetType) || (TargetType == TargetRules.TargetType.Server)) 
 					&& IsCurrentPlatform)
 				{
 					// The hardcoded engine directory needs to be a relative path to match the normal EngineDir format. Not doing so breaks the network file system (TTP#315861).
@@ -1289,7 +1350,7 @@ namespace UnrealBuildTool
 				}
 
 				// Generate static libraries for monolithic games in Rocket
-				if ((UnrealBuildTool.BuildingRocket() || UnrealBuildTool.RunningRocket()) && TargetRules.IsAGame(Rules.Type))
+				if ((UnrealBuildTool.BuildingRocket() || UnrealBuildTool.RunningRocket()) && TargetRules.IsAGame(TargetType))
 				{
 					List<UEBuildModule> Modules = ExecutableBinary.GetAllDependencyModules(true, false);
 					foreach (UEBuildModuleCPP Module in Modules.OfType<UEBuildModuleCPP>())
@@ -1297,7 +1358,7 @@ namespace UnrealBuildTool
 						if(Utils.IsFileUnderDirectory(Module.ModuleDirectory, BuildConfiguration.RelativeEnginePath) && Module.Binary == ExecutableBinary)
 						{
 							UnrealTargetConfiguration LibraryConfiguration = (Configuration == UnrealTargetConfiguration.DebugGame)? UnrealTargetConfiguration.Development : Configuration;
-							Module.RedistStaticLibraryPath = MakeBinaryPath("", "UE4Game-Redist-" + Module.Name, Platform, LibraryConfiguration, UEBuildBinaryType.StaticLibrary, Rules.Type, false, null, AppName);
+							Module.RedistStaticLibraryPath = MakeBinaryPath("", "UE4Game-Redist-" + Module.Name, Platform, LibraryConfiguration, UEBuildBinaryType.StaticLibrary, TargetType, false, null, AppName);
 							Module.bBuildingRedistStaticLibrary = UnrealBuildTool.BuildingRocket();
                             if (Module.bBuildingRedistStaticLibrary)
                             {
@@ -1410,7 +1471,7 @@ namespace UnrealBuildTool
 				}
 			}
 
-			if (ShouldCompileMonolithic() && !ProjectFileGenerator.bGenerateProjectFiles && Rules != null && Rules.Type != TargetRules.TargetType.Program)
+			if (ShouldCompileMonolithic() && !ProjectFileGenerator.bGenerateProjectFiles && Rules != null && TargetType != TargetRules.TargetType.Program)
 			{
 				// All non-program monolithic binaries implicitly depend on all static plugin libraries so they are always linked appropriately
 				// In order to do this, we create a new module here with a cpp file we emit that invokes an empty function in each library.
@@ -1435,7 +1496,7 @@ namespace UnrealBuildTool
 		{
 			string EULAViolationWarning = null;
 
-			if (Rules.Type != TargetRules.TargetType.Editor && Rules.Type != TargetRules.TargetType.Program &&
+			if (TargetType != TargetRules.TargetType.Editor && TargetType != TargetRules.TargetType.Program &&
 				BuildConfiguration.bCheckLicenseViolations)
 			{
 				var RedistributionErrorMessageBuilder = new StringBuilder();
@@ -1712,7 +1773,7 @@ namespace UnrealBuildTool
 			UEBuildBinaryCPP BinaryCPP = ExecutableBinary as UEBuildBinaryCPP;
 			if (BinaryCPP != null)
 			{
-				TargetInfo CurrentTarget = new TargetInfo(Platform, Configuration, (Rules != null) ? Rules.Type : (TargetRules.TargetType?)null);
+				TargetInfo CurrentTarget = new TargetInfo(Platform, Configuration, (Rules != null) ? TargetType : (TargetRules.TargetType?)null);
 				foreach (var TargetModuleName in BinaryCPP.ModuleNames)
 				{
 					string UnusedFilename;
@@ -1986,7 +2047,7 @@ namespace UnrealBuildTool
 			string OutputFilePath;
 			if ((UnrealBuildTool.BuildingRocket() || UnrealBuildTool.RunningRocket()) && bCompileMonolithic)
 			{
-				OutputFilePath = MakeBinaryPath(Module.Name, Module.Name, BinaryType, Rules.Type, bIsRocketModule, Plugin, AppName);
+				OutputFilePath = MakeBinaryPath(Module.Name, Module.Name, BinaryType, TargetType, bIsRocketModule, Plugin, AppName);
                 if (UnrealBuildTool.BuildingRocket())
                 {
                     SpecialRocketLibFilesThatAreBuildProducts.Add(OutputFilePath);
@@ -1994,7 +2055,7 @@ namespace UnrealBuildTool
 			}
 			else
 			{
-				OutputFilePath = MakeBinaryPath(Module.Name, GetAppName() + "-" + Module.Name, BinaryType, Rules.Type, bIsRocketModule, Plugin, AppName);
+				OutputFilePath = MakeBinaryPath(Module.Name, GetAppName() + "-" + Module.Name, BinaryType, TargetType, bIsRocketModule, Plugin, AppName);
 			}
 
 			// Try to determine if we have the rules file
@@ -2052,7 +2113,7 @@ namespace UnrealBuildTool
 					bool bIsRocketModule = RulesCompiler.IsRocketProjectModule(ModuleName);
 
 					// Create a DLL binary for this module
-					string OutputFilePath = MakeBinaryPath(ModuleName, GetAppName() + "-" + ModuleName, UEBuildBinaryType.DynamicLinkLibrary, Rules.Type, bIsRocketModule, null, AppName);
+					string OutputFilePath = MakeBinaryPath(ModuleName, GetAppName() + "-" + ModuleName, UEBuildBinaryType.DynamicLinkLibrary, TargetType, bIsRocketModule, null, AppName);
 					UEBuildBinaryConfiguration Config = new UEBuildBinaryConfiguration( InType: UEBuildBinaryType.DynamicLinkLibrary,
 																						InOutputFilePath: OutputFilePath,
 																						InIntermediateDirectory: RulesCompiler.IsGameModule(ModuleName)? ProjectIntermediateDirectory : EngineIntermediateDirectory,
@@ -2206,7 +2267,7 @@ namespace UnrealBuildTool
 						return UEBuildConfiguration.bBuildEditor;
 
 					case PluginInfo.PluginModuleType.Program:
-						return Rules.Type == TargetRules.TargetType.Program;
+						return TargetType == TargetRules.TargetType.Program;
 				}
 			}
 			return false;
@@ -2217,7 +2278,7 @@ namespace UnrealBuildTool
 		{
 			var BuildPlatform = UEBuildPlatform.GetBuildPlatform(Platform);
 			List<string> PlatformExtraModules = new List<string>();
-			BuildPlatform.GetExtraModules(new TargetInfo(Platform, Configuration, Rules.Type), this, ref PlatformExtraModules);
+			BuildPlatform.GetExtraModules(new TargetInfo(Platform, Configuration, TargetType), this, ref PlatformExtraModules);
 			ExtraModuleNames.AddRange(PlatformExtraModules);			
 		}
 
@@ -2235,7 +2296,7 @@ namespace UnrealBuildTool
 				}
 
 				// Build the enabled plugin list
-				if (ShouldCompileMonolithic() || Rules.Type == TargetRules.TargetType.Program)
+				if (ShouldCompileMonolithic() || TargetType == TargetRules.TargetType.Program)
 				{
 					var FilterPluginNames = new List<string>(Rules.AdditionalPlugins);
 
@@ -2246,7 +2307,7 @@ namespace UnrealBuildTool
 					}
 
 					// Update the plugin list for game targets
-					if(Rules.Type != TargetRules.TargetType.Program && UnrealBuildTool.HasUProjectFile())
+					if(TargetType != TargetRules.TargetType.Program && UnrealBuildTool.HasUProjectFile())
 					{
 						// Enable all the game specific plugins by default
 						FilterPluginNames.AddRange(ValidPlugins.Where(x => x.LoadedFrom == PluginInfo.LoadedFromType.GameProject).Select(x => x.Name));
@@ -2271,7 +2332,7 @@ namespace UnrealBuildTool
 				List<UEBuildBinaryConfiguration> RulesBuildBinaryConfigurations = new List<UEBuildBinaryConfiguration>();
 				List<string> RulesExtraModuleNames = new List<string>();
 				Rules.SetupBinaries(
-					new TargetInfo(Platform, Configuration, Rules.Type),
+					new TargetInfo(Platform, Configuration, TargetType),
 					ref RulesBuildBinaryConfigurations,
 					ref RulesExtraModuleNames
 					);
@@ -2317,7 +2378,7 @@ namespace UnrealBuildTool
 				LinkEnvironmentConfiguration RulesLinkEnvConfig = GlobalLinkEnvironment.Config;
 				CPPEnvironmentConfiguration RulesCPPEnvConfig = GlobalCompileEnvironment.Config;
 				SetupDefaultGlobalEnvironment(
-						new TargetInfo(Platform, Configuration, Rules.Type),
+						new TargetInfo(Platform, Configuration, TargetType),
 						ref RulesLinkEnvConfig,
 						ref RulesCPPEnvConfig
 						);
@@ -2339,13 +2400,13 @@ namespace UnrealBuildTool
 				// This will potentially cause problems due to them generating a game-agnostic exe.
 				// RocketEditor is a special case.
 				// 
-				bool bAllowGameOverride = !TargetRules.IsGameType(Rules.Type);
+				bool bAllowGameOverride = !TargetRules.IsGameType(TargetType);
 				if (bAllowGameOverride ||
 					(Rules.ToString() == "UE4EditorTarget" && UnrealBuildTool.BuildingRocket()) ||	// @todo Rocket: Hard coded target name hack
 					Rules.ShouldCompileMonolithic(Platform, Configuration))
 				{
 					Rules.SetupGlobalEnvironment(
-						new TargetInfo(Platform, Configuration, Rules.Type),
+						new TargetInfo(Platform, Configuration, TargetType),
 						ref RulesLinkEnvConfig,
 						ref RulesCPPEnvConfig
 						);
@@ -2354,7 +2415,7 @@ namespace UnrealBuildTool
 					GlobalLinkEnvironment.Config = RulesLinkEnvConfig;
 					GlobalCompileEnvironment.Config = RulesCPPEnvConfig;
 				}
-				GlobalCompileEnvironment.Config.Definitions.Add(String.Format("IS_PROGRAM={0}", Rules.Type == TargetRules.TargetType.Program ? "1" : "0"));
+				GlobalCompileEnvironment.Config.Definitions.Add(String.Format("IS_PROGRAM={0}", TargetType == TargetRules.TargetType.Program ? "1" : "0"));
 			}
 
 			// Validate UE configuration - needs to happen before setting any environment mojo and after argument parsing.
@@ -2581,7 +2642,7 @@ namespace UnrealBuildTool
 			}
 			else
 			{
-				return new TargetInfo(Platform, Configuration, Rules.Type);
+				return new TargetInfo(Platform, Configuration, TargetType);
 			}
 		}
 
diff --git a/Engine/Source/Programs/UnrealBuildTool/IOS/IOSToolChain.cs b/Engine/Source/Programs/UnrealBuildTool/IOS/IOSToolChain.cs
--- a/Engine/Source/Programs/UnrealBuildTool/IOS/IOSToolChain.cs
+++ b/Engine/Source/Programs/UnrealBuildTool/IOS/IOSToolChain.cs
@@ -962,7 +962,7 @@ namespace UnrealBuildTool
 		{
 			base.PostBuildSync(Target);
 
-			string AppName = Target.Rules.Type == TargetRules.TargetType.Game ? Target.GameName : Target.AppName;
+			string AppName = Target.TargetType == TargetRules.TargetType.Game ? Target.GameName : Target.AppName;
 
 			if (BuildHostPlatform.Current.Platform == UnrealTargetPlatform.Mac)
 			{
@@ -1100,7 +1100,7 @@ namespace UnrealBuildTool
 				}
 
 				// Generate static libraries for monolithic games in Rocket
-				if ((UnrealBuildTool.BuildingRocket() || UnrealBuildTool.RunningRocket()) && TargetRules.IsAGame(Target.Rules.Type))
+				if ((UnrealBuildTool.BuildingRocket() || UnrealBuildTool.RunningRocket()) && TargetRules.IsAGame(Target.TargetType))
 				{
 					List<UEBuildModule> Modules = Target.AppBinaries[0].GetAllDependencyModules(true, false);
 					foreach (UEBuildModuleCPP Module in Modules.OfType<UEBuildModuleCPP>())
diff --git a/Engine/Source/Programs/UnrealBuildTool/IOS/UEDeployIOS.cs b/Engine/Source/Programs/UnrealBuildTool/IOS/UEDeployIOS.cs
--- a/Engine/Source/Programs/UnrealBuildTool/IOS/UEDeployIOS.cs
+++ b/Engine/Source/Programs/UnrealBuildTool/IOS/UEDeployIOS.cs
@@ -298,9 +298,9 @@ namespace UnrealBuildTool.IOS
 					try
 					{
 						string BinaryDir = Path.GetDirectoryName(InTarget.OutputPath) + "\\";
-						if (BinaryDir.EndsWith(InTarget.AppName + "\\Binaries\\IOS\\") && InTarget.Rules.Type != TargetRules.TargetType.Game)
+						if (BinaryDir.EndsWith(InTarget.AppName + "\\Binaries\\IOS\\") && InTarget.TargetType != TargetRules.TargetType.Game)
 						{
-							BinaryDir = BinaryDir.Replace(InTarget.Rules.Type.ToString(), "Game");
+							BinaryDir = BinaryDir.Replace(InTarget.TargetType.ToString(), "Game");
 						}
 
 						// Get the app bundle's name
diff --git a/Engine/Source/Programs/UnrealBuildTool/Mac/UEDeployMac.cs b/Engine/Source/Programs/UnrealBuildTool/Mac/UEDeployMac.cs
--- a/Engine/Source/Programs/UnrealBuildTool/Mac/UEDeployMac.cs
+++ b/Engine/Source/Programs/UnrealBuildTool/Mac/UEDeployMac.cs
@@ -96,9 +96,9 @@ namespace UnrealBuildTool.IOS
 					try
 					{
 						string BinaryDir = Path.GetDirectoryName(InTarget.OutputPath) + "\\";
-						if (BinaryDir.EndsWith(InTarget.AppName + "\\Binaries\\Mac\\") && InTarget.Rules.Type != TargetRules.TargetType.Game)
+						if (BinaryDir.EndsWith(InTarget.AppName + "\\Binaries\\Mac\\") && InTarget.TargetType != TargetRules.TargetType.Game)
 						{
-							BinaryDir = BinaryDir.Replace(InTarget.Rules.Type.ToString(), "Game");
+							BinaryDir = BinaryDir.Replace(InTarget.TargetType.ToString(), "Game");
 						}
 
 						string RemoteBinariesDir = Toolchain.ConvertPath( BinaryDir );
diff --git a/Engine/Source/Programs/UnrealBuildTool/System/ActionHistory.cs b/Engine/Source/Programs/UnrealBuildTool/System/ActionHistory.cs
--- a/Engine/Source/Programs/UnrealBuildTool/System/ActionHistory.cs
+++ b/Engine/Source/Programs/UnrealBuildTool/System/ActionHistory.cs
@@ -82,7 +82,7 @@ namespace UnrealBuildTool
 		public static string GeneratePathForTarget(UEBuildTarget Target)
 		{
 			string Folder = null;
-			if (Target.ShouldCompileMonolithic() || Target.Rules.Type == TargetRules.TargetType.Program)
+			if (Target.ShouldCompileMonolithic() || Target.TargetType == TargetRules.TargetType.Program)
 			{
 				// Monolithic configs and programs have their Action History stored in their respective project folders
 				// or under engine intermediate folder + program name folder
diff --git a/Engine/Source/Programs/UnrealBuildTool/System/ExternalExecution.cs b/Engine/Source/Programs/UnrealBuildTool/System/ExternalExecution.cs
--- a/Engine/Source/Programs/UnrealBuildTool/System/ExternalExecution.cs
+++ b/Engine/Source/Programs/UnrealBuildTool/System/ExternalExecution.cs
@@ -88,7 +88,7 @@ namespace UnrealBuildTool
 
 		public UHTManifest(UEBuildTarget Target, string InRootLocalPath, string InRootBuildPath, IEnumerable<UHTModuleInfo> ModuleInfo)
 		{
-			IsGameTarget  = TargetRules.IsGameType(Target.Rules.Type);
+			IsGameTarget  = TargetRules.IsGameType(Target.TargetType);
 			RootLocalPath = InRootLocalPath;
 			RootBuildPath = InRootBuildPath;
 			TargetName    = Target.GetTargetName();
@@ -98,7 +98,7 @@ namespace UnrealBuildTool
 				ModuleType				 = Info.ModuleType,
 				BaseDirectory            = Info.ModuleDirectory,
 				IncludeBase              = Info.ModuleDirectory,
-				OutputDirectory          = UEBuildModuleCPP.GetGeneratedCodeDirectoryForModule(Target, Info.ModuleDirectory, Info.ModuleName),
+				OutputDirectory          = Path.GetDirectoryName( Info.GeneratedCPPFilenameBase ),
 				ClassesHeaders           = Info.PublicUObjectClassesHeaders.Select((Header) => Header.AbsolutePath).ToList(),
 				PublicHeaders            = Info.PublicUObjectHeaders       .Select((Header) => Header.AbsolutePath).ToList(),
 				PrivateHeaders           = Info.PrivateUObjectHeaders      .Select((Header) => Header.AbsolutePath).ToList(),
@@ -270,7 +270,7 @@ namespace UnrealBuildTool
 		/// Gets the timestamp of CoreUObject.generated.cpp file.
 		/// </summary>
 		/// <returns>Last write time of CoreUObject.generated.cpp or DateTime.MaxValue if it doesn't exist.</returns>
-		private static DateTime GetCoreGeneratedTimestamp(UEBuildTarget Target, string ModuleName, string ModuleDirectory)
+		private static DateTime GetCoreGeneratedTimestamp(string ModuleName, string ModuleGeneratedCodeDirectory)
 		{			
 			DateTime Timestamp;
 			if( UnrealBuildTool.RunningRocket() )
@@ -280,7 +280,7 @@ namespace UnrealBuildTool
 			}
 			else
 			{
-				string CoreGeneratedFilename = Path.Combine(UEBuildModuleCPP.GetGeneratedCodeDirectoryForModule(Target, ModuleDirectory, ModuleName), ModuleName + ".generated.cpp");
+				string CoreGeneratedFilename = Path.Combine(ModuleGeneratedCodeDirectory, ModuleName + ".generated.cpp");
 				if (File.Exists(CoreGeneratedFilename))
 				{
 					Timestamp = new FileInfo(CoreGeneratedFilename).LastWriteTime;
@@ -301,7 +301,7 @@ namespace UnrealBuildTool
 		 * 
 		 * @return					True if the code files are out of date
 		 * */
-		private static bool AreGeneratedCodeFilesOutOfDate(UEBuildTarget Target, List<UHTModuleInfo> UObjectModules)
+		private static bool AreGeneratedCodeFilesOutOfDate(List<UHTModuleInfo> UObjectModules)
 		{
 			bool bIsOutOfDate = false;
 
@@ -317,7 +317,7 @@ namespace UnrealBuildTool
 				{
 					if( Module.ModuleName.Equals( "CoreUObject", StringComparison.InvariantCultureIgnoreCase ) )
 					{
-						CoreGeneratedTimestamp = GetCoreGeneratedTimestamp(Target, Module.ModuleName, Module.ModuleDirectory);
+						CoreGeneratedTimestamp = GetCoreGeneratedTimestamp(Module.ModuleName, Path.GetDirectoryName( Module.GeneratedCPPFilenameBase ));
 						break;
 					}
 				}
@@ -342,7 +342,7 @@ namespace UnrealBuildTool
 				}
 
 				// Make sure we have an existing folder for generated code.  If not, then we definitely need to generate code!
-				var GeneratedCodeDirectory = UEBuildModuleCPP.GetGeneratedCodeDirectoryForModule(Target, Module.ModuleDirectory, Module.ModuleName);
+				var GeneratedCodeDirectory = Path.GetDirectoryName( Module.GeneratedCPPFilenameBase );
 				var TestDirectory = (FileSystemInfo)new DirectoryInfo(GeneratedCodeDirectory);
 				if( TestDirectory.Exists )
 				{
@@ -422,11 +422,11 @@ namespace UnrealBuildTool
 		}
 
 		/** Updates the intermediate include directory timestamps of all the passed in UObject modules */
-		private static void UpdateDirectoryTimestamps(UEBuildTarget Target, List<UHTModuleInfo> UObjectModules)
+		private static void UpdateDirectoryTimestamps(List<UHTModuleInfo> UObjectModules)
 		{
 			foreach( var Module in UObjectModules )
 			{
-				string GeneratedCodeDirectory = UEBuildModuleCPP.GetGeneratedCodeDirectoryForModule(Target, Module.ModuleDirectory, Module.ModuleName);
+				string GeneratedCodeDirectory = Path.GetDirectoryName( Module.GeneratedCPPFilenameBase );
 				var GeneratedCodeDirectoryInfo = new DirectoryInfo( GeneratedCodeDirectory );
 
 				try
@@ -530,7 +530,7 @@ namespace UnrealBuildTool
 
 
 				// ensure the headers are up to date
-				bool bUHTNeedsToRun = (UEBuildConfiguration.bForceHeaderGeneration == true || AreGeneratedCodeFilesOutOfDate(Target, UObjectModules));
+				bool bUHTNeedsToRun = (UEBuildConfiguration.bForceHeaderGeneration == true || AreGeneratedCodeFilesOutOfDate(UObjectModules));
 				if( bUHTNeedsToRun )
 				{
 					// Since code files are definitely out of date, we'll now finish computing information about the UObject modules for UHT.  We
@@ -660,11 +660,11 @@ namespace UnrealBuildTool
 					// generated headers include other generated headers using absolute paths which in case of building remotely are already
 					// the remote machine absolute paths. Because of that parsing headers will not result in finding all includes properly.
 					// @todo ubtmake: Need to figure out what this does in the assembler case, and whether we need to run it
-					ToolChain.PostCodeGeneration(Target, Manifest);
+					ToolChain.PostCodeGeneration(Manifest);
 				}
 
 				// touch the directories
-				UpdateDirectoryTimestamps(Target, UObjectModules);
+				UpdateDirectoryTimestamps(UObjectModules);
 
 				Progress.Write(3, 3);
 			}
diff --git a/Engine/Source/Programs/UnrealBuildTool/System/RulesCompiler.cs b/Engine/Source/Programs/UnrealBuildTool/System/RulesCompiler.cs
--- a/Engine/Source/Programs/UnrealBuildTool/System/RulesCompiler.cs
+++ b/Engine/Source/Programs/UnrealBuildTool/System/RulesCompiler.cs
@@ -314,6 +314,7 @@ namespace UnrealBuildTool
 	public abstract class TargetRules
 	{
 		/// Type of target
+		[Serializable]
 		public enum TargetType
 		{
 			/// Cooked monolithic game executable (GameName.exe).  Also used for a game-agnostic engine executable (UE4Game.exe or RocketGame.exe)
diff --git a/Engine/Source/Programs/UnrealBuildTool/System/UnrealBuildTool.cs b/Engine/Source/Programs/UnrealBuildTool/System/UnrealBuildTool.cs
--- a/Engine/Source/Programs/UnrealBuildTool/System/UnrealBuildTool.cs
+++ b/Engine/Source/Programs/UnrealBuildTool/System/UnrealBuildTool.cs
@@ -1092,11 +1092,12 @@ namespace UnrealBuildTool
                                 {
                                     // We need to be able to identify the Target.Type we can derive it from the Arguments.
                                     BuildConfiguration.bFlushBuildDirOnRemoteMac = false;
-                                    UEBuildTarget CheckTarget = UEBuildTarget.CreateTarget(Arguments);
+									var TargetDescs = UEBuildTarget.ParseTargetCommandLine( Arguments );
+			                        UEBuildTarget CheckTarget = UEBuildTarget.CreateTarget( TargetDescs[0] );	// @todo ubtmake: This may not work in assembler only mode.  We don't want to be loading target rules assemblies here either.
                                     CheckTarget.SetupGlobalEnvironment();
-                                    if ((CheckTarget.Rules.Type == TargetRules.TargetType.Game) ||
-                                        (CheckTarget.Rules.Type == TargetRules.TargetType.Server) ||
-                                        (CheckTarget.Rules.Type == TargetRules.TargetType.Client))
+                                    if ((CheckTarget.TargetType == TargetRules.TargetType.Game) ||
+                                        (CheckTarget.TargetType == TargetRules.TargetType.Server) ||
+                                        (CheckTarget.TargetType == TargetRules.TargetType.Client))
                                     {
                                         CheckTarget.AppName = CheckTarget.GameName;
                                     }
@@ -1261,7 +1262,7 @@ namespace UnrealBuildTool
             }
         }
 
-        // @todo: Ideally get rid of RunUBT() and all of the Clear/Reset stuff!
+        
         public static ECompilationResult RunUBT(string[] Arguments)
         {
             bool bSuccess = true;
@@ -1294,7 +1295,6 @@ namespace UnrealBuildTool
 
             BuildPlatform.ResetBuildConfiguration(ResetPlatform, ResetConfiguration);
 
-            var Targets = new List<UEBuildTarget>();
             string ExecutorName = "Unknown";
             ECompilationResult BuildResult = ECompilationResult.Succeeded;
 
@@ -1328,28 +1328,22 @@ namespace UnrealBuildTool
 					Log.TraceInformation( "RunUBT initialization took " + RunUBTInitTime + "s" );
 				}
 
+	            var TargetDescs = new List<TargetDescriptor>();
 				{ 
-					var TargetInitStartTime = DateTime.UtcNow;
+					var TargetDescstStartTime = DateTime.UtcNow;
 
 					foreach (string[] TargetSetting in TargetSettings)
 					{
-						var Target = UEBuildTarget.CreateTarget(TargetSetting);	// @todo ubtmake: Optimization: Ideally we don't have to do this in the assembler.  It is pretty slow to scan all rules files.  Maybe we can skip the slow parts internally?  It has to scan for *.Target.cs files, and also Load the rules assembly up. :(
-						if ((Target == null) && (UEBuildConfiguration.bCleanProject))
-						{
-							continue;
-						}
-						Targets.Add(Target);
+						TargetDescs.AddRange( UEBuildTarget.ParseTargetCommandLine( TargetSetting ) );
 					}
 
 					if( BuildConfiguration.bPrintPerformanceInfo )
 					{ 
-						var TargetInitTime = (DateTime.UtcNow - TargetInitStartTime).TotalSeconds;
-						Log.TraceInformation( "Target init took " + TargetInitTime + "s" );
+						var TargetDescsTime = (DateTime.UtcNow - TargetDescstStartTime).TotalSeconds;
+						Log.TraceInformation( "Target descriptors took " + TargetDescsTime + "s" );
 					}
 				}
 
-                UBTMakefile UBTMakefile = null;
-
 
                 // If UBT is being used to build something different than it did last time
                 if( !ProjectFileGenerator.bGenerateProjectFiles )
@@ -1361,7 +1355,7 @@ namespace UnrealBuildTool
                         // UHT.  In that (very common) case we definitely don't want to have to rebuild our cache from scratch.
                         bool bMustAssumeSameTargets = false;
 
-                        if( Targets[0].GetTargetName().Equals( "UnrealHeaderTool", StringComparison.InvariantCultureIgnoreCase ) )
+                        if( TargetDescs[0].TargetName.Equals( "UnrealHeaderTool", StringComparison.InvariantCultureIgnoreCase ) )
                         {
                             int NoMutexArgumentIndex;
                             if (Utils.ParseCommandLineFlag(Arguments, "-NoMutex", out NoMutexArgumentIndex))
@@ -1377,7 +1371,7 @@ namespace UnrealBuildTool
                         }
                         else
                         {
-                            string TargetCollectionName = MakeTargetCollectionName( Targets );
+                            string TargetCollectionName = MakeTargetCollectionName( TargetDescs );
                             string LastBuiltTargetsFilePath = Path.Combine( BuildConfiguration.BaseIntermediatePath, "LastBuiltTargets.txt" );
                             if( File.Exists( LastBuiltTargetsFilePath ) && Utils.ReadAllText( LastBuiltTargetsFilePath ) == TargetCollectionName )
                             {
@@ -1406,6 +1400,7 @@ namespace UnrealBuildTool
             
 
 
+                UBTMakefile UBTMakefile = null;
                 { 
                     // If we're generating project files, then go ahead and wipe out the existing UBTMakefile for every target, to make sure that
                     // it gets a full dependency scan next time.
@@ -1413,7 +1408,7 @@ namespace UnrealBuildTool
                     if( ProjectFileGenerator.bGenerateProjectFiles )	// @todo ubtmake: This is only hit when generating IntelliSense for project files.  Probably should be done right inside ProjectFileGenerator.bat
                     {													// @todo ubtmake: Won't catch multi-target cases as GPF always builds one target at a time for Intellisense
                         // Delete the UBTMakefile
-                        var UBTMakefilePath = UnrealBuildTool.GetUBTMakefilePath( Targets );
+                        var UBTMakefilePath = UnrealBuildTool.GetUBTMakefilePath( TargetDescs );
                         if (File.Exists(UBTMakefilePath))
                         {
                             UEBuildTarget.CleanFile(UBTMakefilePath);
@@ -1435,7 +1430,7 @@ namespace UnrealBuildTool
 
                         // Try to load the UBTMakefile.  It will only be loaded if it has valid content and is not determined to be out of date.    
 						string ReasonNotLoaded;
-                        UBTMakefile = LoadUBTMakefile( Targets, out ReasonNotLoaded );
+                        UBTMakefile = LoadUBTMakefile( TargetDescs, out ReasonNotLoaded );
 
                         if( UBTMakefile == null )
                         { 
@@ -1443,7 +1438,7 @@ namespace UnrealBuildTool
                             // a 'gather' and 'assemble' in the same run.  This will take a while longer, but subsequent runs will be fast!
                             UnrealBuildTool.bIsGatheringBuild_Unsafe = true;
 
-                            Log.TraceInformation( "Creating makefile for {0}{1} ({2})", Targets[0].GetTargetName(), Targets.Count > 1 ? ( " (and " + ( Targets.Count - 1 ).ToString() + " more)" ) : "", ReasonNotLoaded );
+                            Log.TraceInformation( "Creating makefile for {0}{1} ({2})", TargetDescs[0].TargetName, TargetDescs.Count > 1 ? ( " (and " + ( TargetDescs.Count - 1 ).ToString() + " more)" ) : "", ReasonNotLoaded );
                         }
                     }
 
@@ -1453,6 +1448,34 @@ namespace UnrealBuildTool
                 }
 
 
+				List<UEBuildTarget> Targets;
+				if( UBTMakefile != null && !IsGatheringBuild && IsAssemblingBuild )
+				{
+					// If we've loaded a makefile, then we can fill target information from this file!
+					Targets = UBTMakefile.Targets;
+				}
+				else
+				{ 
+					var TargetInitStartTime = DateTime.UtcNow;
+
+					Targets = new List<UEBuildTarget>();
+					foreach( var TargetDesc in TargetDescs )
+					{
+						var Target = UEBuildTarget.CreateTarget( TargetDesc );
+						if ((Target == null) && (UEBuildConfiguration.bCleanProject))
+						{
+							continue;
+						}
+						Targets.Add(Target);
+					}
+
+					if( BuildConfiguration.bPrintPerformanceInfo )
+					{ 
+						var TargetInitTime = (DateTime.UtcNow - TargetInitStartTime).TotalSeconds;
+						Log.TraceInformation( "Target init took " + TargetInitTime + "s" );
+					}
+				}
+
                 // Build action lists for all passed in targets.
                 var OutputItemsForAllTargets = new List<FileItem>();
                 var TargetNameToUObjectModules = new Dictionary<string, List<UHTModuleInfo>>( StringComparer.InvariantCultureIgnoreCase );
@@ -1501,7 +1524,7 @@ namespace UnrealBuildTool
                         if ( (BuildConfiguration.bXGEExport && UEBuildConfiguration.bGenerateManifest) || (!ProjectFileGenerator.bGenerateProjectFiles && !UEBuildConfiguration.bGenerateManifest && !UEBuildConfiguration.bCleanProject))
                         {
                             // We don't currently support building game targets in rocket.
-                            if (UnrealBuildTool.RunningRocket() && Target.Rules != null && Target.Rules.Type == TargetRules.TargetType.Game)
+                            if (UnrealBuildTool.RunningRocket() && Target.Rules != null && Target.TargetType == TargetRules.TargetType.Game)
                             {
                                 throw new BuildException(
                                     "You currently can not build a game target in Rocket.\nTry again in a future release.\nFor now, build and run your editor project."
@@ -1537,7 +1560,7 @@ namespace UnrealBuildTool
                             "CleanTarget", UEBuildConfiguration.bCleanProject.ToString(),
                             "Monolithic", Target.ShouldCompileMonolithic().ToString(),
                             "CreateDebugInfo", Target.IsCreatingDebugInfo().ToString(),
-                            "TargetType", Target.Rules.Type.ToString(),
+                            "TargetType", Target.TargetType.ToString(),
                             "TargetCreateTimeSec", TargetBuildTime.ToString("0.00")
                             );
                     }
@@ -1571,6 +1594,7 @@ namespace UnrealBuildTool
                             UBTMakefile.EnvironmentVariables.Add( Tuple.Create( (string)EnvironmentVariable.Key, (string)EnvironmentVariable.Value ) );
                         }
                         UBTMakefile.TargetNameToUObjectModules = TargetNameToUObjectModules;
+						UBTMakefile.Targets = Targets;
 
 						if( BuildConfiguration.bUseExperimentalFastBuildIteration )
 						{ 
@@ -1578,7 +1602,7 @@ namespace UnrealBuildTool
 							// to assemble the build.  Even if we are configured to assemble the build in this same invocation, we want to save out the
 							// Makefile so that it can be used on subsequent 'assemble only' runs, for the fastest possible iteration times
 							// @todo ubtmake: Optimization: We could make 'gather + assemble' mode slightly faster by saving this while busy compiling (on our worker thread)
-							SaveUBTMakefile( Targets, UBTMakefile );
+							SaveUBTMakefile( TargetDescs, UBTMakefile );
 						}
                     }
 
@@ -1621,7 +1645,6 @@ namespace UnrealBuildTool
                         if( BuildResult == ECompilationResult.Succeeded )
                         { 
                             // Make sure any old DLL files from in-engine recompiles aren't lying around.  Must be called after the action graph is finalized.
-                            // @todo ubtmake: Will we even use traditional hot reload DLLs anymore in the future?  
                             ActionGraph.DeleteStaleHotReloadDLLs();
 
                             // Plan the actions to execute for the build.
@@ -1768,7 +1791,7 @@ namespace UnrealBuildTool
 		private static bool ShouldDoHotReload(UEBuildTarget Target)
 		{
 			bool bIsRunning = false;
-			if (!ProjectFileGenerator.bGenerateProjectFiles && !UEBuildConfiguration.bGenerateManifest && Target.Rules.Type == TargetRules.TargetType.Editor)
+			if (!ProjectFileGenerator.bGenerateProjectFiles && !UEBuildConfiguration.bGenerateManifest && Target.TargetType == TargetRules.TargetType.Editor)
 			{
 				var EditorProcessFilename = UEBuildTarget.MakeBinaryPath("Launcher", "UE4Editor", Target.Platform, Target.Configuration, UEBuildBinaryType.Executable, null, false, null, "UE4Editor");
 				var EditorProcessName = Path.GetFileNameWithoutExtension(EditorProcessFilename);
@@ -2009,6 +2032,9 @@ namespace UnrealBuildTool
             /** Maps each target to a list of UObject module info structures */
             public Dictionary<string,List<UHTModuleInfo>> TargetNameToUObjectModules;
 
+			/** List of targets being built */
+			public List<UEBuildTarget> Targets;
+
 
             /** @return	 True if this makefile's contents look valid.  Called after loading the file to make sure it is legit. */
             public bool IsValidMakefile()
@@ -2017,7 +2043,8 @@ namespace UnrealBuildTool
                     AllActions != null && AllActions.Count > 0 &&
                     PrerequisiteActions != null && PrerequisiteActions.Length > 0 &&
                     EnvironmentVariables != null &&
-                    TargetNameToUObjectModules != null && TargetNameToUObjectModules.Count > 0;
+                    TargetNameToUObjectModules != null && TargetNameToUObjectModules.Count > 0 &&
+					Targets != null && Targets.Count > 0;
             }
         }
 
@@ -2025,8 +2052,8 @@ namespace UnrealBuildTool
         /// <summary>
         /// Saves a UBTMakefile to disk
         /// </summary>
-        /// <param name="Targets">List of targets.  Order is not important</param>
-        static void SaveUBTMakefile( List<UEBuildTarget> Targets, UBTMakefile UBTMakefile )
+        /// <param name="TargetDescs">List of targets.  Order is not important</param>
+        static void SaveUBTMakefile( List<TargetDescriptor> TargetDescs, UBTMakefile UBTMakefile )
         {
             if( !UBTMakefile.IsValidMakefile() )
             {
@@ -2035,7 +2062,7 @@ namespace UnrealBuildTool
 
             var TimerStartTime = DateTime.UtcNow;
 
-            var UBTMakefileItem = FileItem.GetItemByFullPath( GetUBTMakefilePath( Targets ) );
+            var UBTMakefileItem = FileItem.GetItemByFullPath( GetUBTMakefilePath( TargetDescs ) );
 
             // @todo ubtmake: Optimization: The UBTMakefile saved for game projects is upwards of 9 MB.  We should try to shrink its content if possible
             // @todo ubtmake: Optimization: C# Serialization may be too slow for these big Makefiles.  Loading these files often shows up as the slower part of the assembling phase.
@@ -2066,12 +2093,12 @@ namespace UnrealBuildTool
         /// <summary>
         /// Loads a UBTMakefile from disk
         /// </summary>
-        /// <param name="Targets">List of targets.  Order is not important</param>
+        /// <param name="TargetDescs">List of targets.  Order is not important</param>
 		/// <param name="ReasonNotLoaded">If the function returns null, this string will contain the reason why</param>
 		/// <returns>The loaded makefile, or null if it failed for some reason.  On failure, the 'ReasonNotLoaded' variable will contain information about why</returns>
-        static UBTMakefile LoadUBTMakefile( List<UEBuildTarget> Targets, out string ReasonNotLoaded )
+        static UBTMakefile LoadUBTMakefile( List<TargetDescriptor> TargetDescs, out string ReasonNotLoaded )
         {
-            var UBTMakefileItem = FileItem.GetItemByFullPath( GetUBTMakefilePath( Targets ) );
+            var UBTMakefileItem = FileItem.GetItemByFullPath( GetUBTMakefilePath( TargetDescs ) );
 			ReasonNotLoaded = null;
 
             // Check the directory timestamp on the project files directory.  If the user has generated project files more
@@ -2193,11 +2220,11 @@ namespace UnrealBuildTool
         /// </summary>
         /// <param name="Targets">List of targets.  Order is not important</param>
         /// <returns>UBTMakefile path</returns>
-        public static string GetUBTMakefilePath( List<UEBuildTarget> Targets )
+        public static string GetUBTMakefilePath( List<TargetDescriptor> TargetDescs )
         {
             string UBTMakefilePath;
 
-            if( Targets.Count == 1 )
+            if( TargetDescs.Count == 1 )
             {
                 // If there's only one target, just save the UBTMakefile in the target's build intermediate directory
                 // under a folder for that target (and platform/config combo.)
@@ -2207,13 +2234,13 @@ namespace UnrealBuildTool
                     PlatformIntermediatePath = Path.Combine(UnrealBuildTool.GetUProjectPath(), BuildConfiguration.PlatformIntermediateFolder);
                 }
 
-                UBTMakefilePath = Path.Combine(PlatformIntermediatePath, Targets[0].GetTargetName(), Targets[0].GetTargetInfo().Configuration.ToString(), "Makefile.ubt" );
+                UBTMakefilePath = Path.Combine(PlatformIntermediatePath, TargetDescs[0].TargetName, TargetDescs[0].Configuration.ToString(), "Makefile.ubt" );
             }
             else
             {
                 // For Makefiles that contain multiple targets, we'll make up a file name that contains all of the targets, their
                 // configurations and platforms, and save it into the base intermediate folder
-                var TargetCollectionName = MakeTargetCollectionName( Targets );
+                var TargetCollectionName = MakeTargetCollectionName( TargetDescs );
 
                 string ProjectIntermediatePath = BuildConfiguration.BaseIntermediatePath;
                 if (UnrealBuildTool.HasUProjectFile())
@@ -2232,18 +2259,18 @@ namespace UnrealBuildTool
         /// <summary>
         /// Makes up a name for a set of targets that we can use for file or directory names
         /// </summary>
-        /// <param name="Targets">List of targets.  Order is not important</param>
+        /// <param name="TargetDescs">List of targets.  Order is not important</param>
         /// <returns>The name to use</returns>
-        private static string MakeTargetCollectionName( List<UEBuildTarget> Targets )
+        private static string MakeTargetCollectionName( List<TargetDescriptor> TargetDescs )
         {
-            if( Targets.Count == 0 )
+            if( TargetDescs.Count == 0 )
             {
                 throw new BuildException( "Expecting at least one Target to be passed to MakeTargetCollectionName" );
             }
 
-            var SortedTargets = new List<UEBuildTarget>();
-            SortedTargets.AddRange( Targets );
-            SortedTargets.Sort( ( x, y ) => { return x.GetTargetName().CompareTo( y.GetTargetName() ); } );
+            var SortedTargets = new List<TargetDescriptor>();
+            SortedTargets.AddRange( TargetDescs );
+            SortedTargets.Sort( ( x, y ) => { return x.TargetName.CompareTo( y.TargetName ); } );
 
             // Figure out what to call our action graph based on everything we're building
             var TargetCollectionName = new StringBuilder();
@@ -2253,7 +2280,9 @@ namespace UnrealBuildTool
                 {
                     TargetCollectionName.Append( "_" );
                 }
-                TargetCollectionName.Append( Target.GetTargetName() + "-" + Target.GetTargetInfo().Platform.ToString() + "-" + Target.GetTargetInfo().Configuration.ToString() );
+
+				// @todo ubtmake: Should we also have the platform Architecture in this string?
+                TargetCollectionName.Append( Target.TargetName + "-" + Target.Platform.ToString() + "-" + Target.Configuration.ToString() );
             }
 
             return TargetCollectionName.ToString();
diff --git a/Engine/Source/Programs/UnrealBuildTool/ToolChain/RemoteToolChain.cs b/Engine/Source/Programs/UnrealBuildTool/ToolChain/RemoteToolChain.cs
--- a/Engine/Source/Programs/UnrealBuildTool/ToolChain/RemoteToolChain.cs
+++ b/Engine/Source/Programs/UnrealBuildTool/ToolChain/RemoteToolChain.cs
@@ -462,7 +462,7 @@ namespace UnrealBuildTool
             }
         }
 
-		public override void PostCodeGeneration(UEBuildTarget Target, UHTManifest Manifest)
+		public override void PostCodeGeneration(UHTManifest Manifest)
 		{
 			if (BuildHostPlatform.Current.Platform != UnrealTargetPlatform.Mac)
 			{
@@ -478,7 +478,7 @@ namespace UnrealBuildTool
 					// @todo uht: Ideally would only copy exactly the files emitted by UnrealHeaderTool, rather than scanning directory (could copy stale files; not a big deal though)
 					try
 					{
-						var GeneratedCodeDirectory = UEBuildModuleCPP.GetGeneratedCodeDirectoryForModule(Target, UObjectModule.BaseDirectory, UObjectModule.Name);
+						var GeneratedCodeDirectory = Path.GetDirectoryName( UObjectModule.GeneratedCPPFilenameBase );
 						var GeneratedCodeFiles     = Directory.GetFiles( GeneratedCodeDirectory, "*", SearchOption.AllDirectories );
 						foreach( var GeneratedCodeFile in GeneratedCodeFiles )
 						{
diff --git a/Engine/Source/Programs/UnrealBuildTool/ToolChain/UEToolChain.cs b/Engine/Source/Programs/UnrealBuildTool/ToolChain/UEToolChain.cs
--- a/Engine/Source/Programs/UnrealBuildTool/ToolChain/UEToolChain.cs
+++ b/Engine/Source/Programs/UnrealBuildTool/ToolChain/UEToolChain.cs
@@ -27,8 +27,8 @@ namespace UnrealBuildTool
 		/// <summary>
 		/// Called immediately after UnrealHeaderTool is executed to generated code for all UObjects modules.  Only is called if UnrealHeaderTool was actually run in this session.
 		/// </summary>
-		/// <param name="UObjectModules">List of UObject modules we generated code for.</param>
-		void PostCodeGeneration(UEBuildTarget Target, UHTManifest Manifest);
+		/// <param name="Manifest">List of UObject modules we generated code for.</param>
+		void PostCodeGeneration(UHTManifest Manifest);
 		
 		void PreBuildSync();
 		
@@ -133,8 +133,8 @@ namespace UnrealBuildTool
 		/// <summary>
 		/// Called immediately after UnrealHeaderTool is executed to generated code for all UObjects modules.  Only is called if UnrealHeaderTool was actually run in this session.
 		/// </summary>
-		/// <param name="UObjectModules">List of UObject modules we generated code for.</param>
-		public virtual void PostCodeGeneration(UEBuildTarget Target,  UHTManifest Manifest)
+		/// <param name="Manifest">List of UObject modules we generated code for.</param>
+		public virtual void PostCodeGeneration(UHTManifest Manifest)
 		{
 		}
 
diff --git a/Engine/Source/Programs/UnrealBuildTool/Windows/UEBuildWindows.cs b/Engine/Source/Programs/UnrealBuildTool/Windows/UEBuildWindows.cs
--- a/Engine/Source/Programs/UnrealBuildTool/Windows/UEBuildWindows.cs
+++ b/Engine/Source/Programs/UnrealBuildTool/Windows/UEBuildWindows.cs
@@ -428,7 +428,7 @@ namespace UnrealBuildTool
                 // that is set by the modules the target includes to allow for easier tracking.
                 // Alternatively, if VSAccessor is modified to not require ATL than we should always exclude the libraries.
                 if (InBuildTarget.ShouldCompileMonolithic() && (InBuildTarget.Rules != null) &&
-                    (TargetRules.IsGameType(InBuildTarget.Rules.Type)) && (TargetRules.IsEditorType(InBuildTarget.Rules.Type) == false))
+                    (TargetRules.IsGameType(InBuildTarget.TargetType)) && (TargetRules.IsEditorType(InBuildTarget.TargetType) == false))
                 {
                     InBuildTarget.GlobalLinkEnvironment.Config.ExcludedLibraries.Add("atl");
                     InBuildTarget.GlobalLinkEnvironment.Config.ExcludedLibraries.Add("atls");
