Fix generating project files for the first time on a machine.

diff --git a/Engine/Source/Programs/UnrealBuildTool/System/ProjectFileGenerator.cs b/Engine/Source/Programs/UnrealBuildTool/System/ProjectFileGenerator.cs
--- a/Engine/Source/Programs/UnrealBuildTool/System/ProjectFileGenerator.cs
+++ b/Engine/Source/Programs/UnrealBuildTool/System/ProjectFileGenerator.cs
@@ -258,6 +258,7 @@ namespace UnrealBuildTool
 		{
 			string InputProjectDir = Path.GetFullPath("Programs\\AutomationToolStub");
 			string OutputProjectDir = IntermediateProjectFilesPath;
+			Directory.CreateDirectory(OutputProjectDir);
 
 			// Load the template aggregate project file
 			XDocument Document = XDocument.Load(Utils.CleanDirectorySeparators(Path.Combine(InputProjectDir, "AutomationToolStubTemplate.csproj")));
