Making sure all dependencies of UAT build with the same config.

#change AutomationTool, AutomationToolLauncher and DotNETUtilities will always compile in Development in the sln to avoid situations where one of UAT dependencies is built with diffrerent config which results in UAT failing to compile scripts in Debug due to outdated dependencies of UAT.

diff --git a/Engine/Source/Programs/AutomationTool/AutomationTool.csproj b/Engine/Source/Programs/AutomationTool/AutomationTool.csproj
--- a/Engine/Source/Programs/AutomationTool/AutomationTool.csproj
+++ b/Engine/Source/Programs/AutomationTool/AutomationTool.csproj
@@ -138,7 +138,7 @@
     <ProjectReference Include="..\UnrealBuildTool\UnrealBuildTool.csproj">
       <Project>{fd7c5e1a-cfe4-4fd5-a525-1eb1599a39ac}</Project>
       <Name>UnrealBuildTool</Name>
-      <Private>True</Private>
+      <Private>False</Private>
     </ProjectReference>
   </ItemGroup>
   <ItemGroup>
diff --git a/Engine/Source/Programs/UnrealBuildTool/System/ProjectFileGenerator.cs b/Engine/Source/Programs/UnrealBuildTool/System/ProjectFileGenerator.cs
--- a/Engine/Source/Programs/UnrealBuildTool/System/ProjectFileGenerator.cs
+++ b/Engine/Source/Programs/UnrealBuildTool/System/ProjectFileGenerator.cs
@@ -566,10 +566,10 @@ namespace UnrealBuildTool
 				if( ( !bGeneratingGameProjectFiles && !bGeneratingRocketProjectFiles ) || bAlwaysIncludeEngineModules )
 				{
 					// Add AutomationTool to the master project
-					ProgramsFolder.ChildProjects.Add(AddSimpleCSharpProject("AutomationTool", true));
+					ProgramsFolder.ChildProjects.Add(AddSimpleCSharpProject("AutomationTool", bShouldBuildForAllSolutionTargets: true, bForceDevelopmentConfiguration: true));
 
 					// Add UnrealAutomationTool (launcher) to the master project
-					ProgramsFolder.ChildProjects.Add(AddSimpleCSharpProject("AutomationToolLauncher"));
+					ProgramsFolder.ChildProjects.Add(AddSimpleCSharpProject("AutomationToolLauncher", bShouldBuildForAllSolutionTargets: true, bForceDevelopmentConfiguration: true));
 
 					// Add automation.csproj files to the master project
 					List<VCSharpProjectFile> AutomationModules = AddAutomationModules(ProgramsFolder);
@@ -578,10 +578,10 @@ namespace UnrealBuildTool
 					ProgramsFolder.ChildProjects.Add(GenerateAutomationToolStubProject(AutomationModules));
 
 					// Add Distill to the master project
-					ProgramsFolder.ChildProjects.Add( AddSimpleCSharpProject( "Distill" ) );
+					ProgramsFolder.ChildProjects.Add(AddSimpleCSharpProject("Distill"));
 
 					// Add DotNETUtilities to the master project
-					ProgramsFolder.ChildProjects.Add( AddSimpleCSharpProject(  "DotNETCommon/DotNETUtilities", true ) );
+					ProgramsFolder.ChildProjects.Add(AddSimpleCSharpProject("DotNETCommon/DotNETUtilities", bShouldBuildForAllSolutionTargets: true, bForceDevelopmentConfiguration: true));
 
 					// Add all of the IOS C# projects
 					AddIOSProjects( ProgramsFolder );
@@ -1186,6 +1186,7 @@ namespace UnrealBuildTool
 		{
 			var ProjectFileName = Utils.MakePathRelativeTo( Path.Combine( Path.Combine( EngineRelativePath, "Source" ), "Programs", "UnrealBuildTool", "UnrealBuildTool.csproj" ), MasterProjectRelativePath );
 			var UnrealBuildToolProject = new VCSharpProjectFile( ProjectFileName );
+			UnrealBuildToolProject.ShouldBuildForAllSolutionTargets = true;
 
 			// Store it off as we need it when generating target projects.
 			UBTProject = UnrealBuildToolProject;
@@ -1202,7 +1203,7 @@ namespace UnrealBuildTool
 		/// </summary>
 		/// <param name="ProjectName">Name of project file to add</param>
 		/// <returns>ProjectFile if the operation was successful, otherwise null.</returns>
-		private VCSharpProjectFile AddSimpleCSharpProject(string ProjectName, bool bShouldBuildForAllSolutionTargets = false)
+		private VCSharpProjectFile AddSimpleCSharpProject(string ProjectName, bool bShouldBuildForAllSolutionTargets = false, bool bForceDevelopmentConfiguration = false)
 		{
 			VCSharpProjectFile Project = null;
 
@@ -1213,7 +1214,7 @@ namespace UnrealBuildTool
 				var FileNameRelativeToMasterProject = Utils.MakePathRelativeTo( ProjectFileName, MasterProjectRelativePath );
 				Project = new VCSharpProjectFile(FileNameRelativeToMasterProject);
 				Project.ShouldBuildForAllSolutionTargets = bShouldBuildForAllSolutionTargets;
-				AddExistingProjectFile(Project);
+				AddExistingProjectFile(Project, bForceDevelopmentConfiguration: bForceDevelopmentConfiguration);
 			}
 			else
 			{
