UnrealBuildTool: Added more performance diagnostic logging

diff --git a/Engine/Source/Programs/UnrealBuildTool/System/ActionGraph.cs b/Engine/Source/Programs/UnrealBuildTool/System/ActionGraph.cs
--- a/Engine/Source/Programs/UnrealBuildTool/System/ActionGraph.cs
+++ b/Engine/Source/Programs/UnrealBuildTool/System/ActionGraph.cs
@@ -217,6 +217,8 @@ namespace UnrealBuildTool
 		/** Builds a list of actions that need to be executed to produce the specified output items. */
 		public static List<Action> GetActionsToExecute(Action[] PrerequisiteActions, List<UEBuildTarget> Targets, out Dictionary<UEBuildTarget, List<FileItem>> TargetToOutdatedPrerequisitesMap)
 		{
+			var CheckOutdatednessStartTime = DateTime.UtcNow;
+
 			// Build a set of all actions needed for this target.
 			var IsActionOutdatedMap = new Dictionary<Action, bool>();
 			foreach (var Action in PrerequisiteActions)
@@ -265,6 +267,12 @@ namespace UnrealBuildTool
 				}
 			}
 
+			if( BuildConfiguration.bPrintPerformanceInfo )
+			{ 
+				var CheckOutdatednessTime = (DateTime.UtcNow - CheckOutdatednessStartTime).TotalSeconds;
+				Log.TraceInformation( "Checking outdatedness took " + CheckOutdatednessTime + "s" );
+			}
+
 			return ActionsToExecute;
 		}
 
@@ -425,6 +433,8 @@ namespace UnrealBuildTool
 		/** Finds and deletes stale hot reload DLLs. */
 		public static void DeleteStaleHotReloadDLLs()
 		{
+			var DeleteStartTime = DateTime.UtcNow;
+
 			foreach (Action BuildAction in AllActions)
 			{
 				if (BuildAction.ActionType == ActionType.Link)
@@ -485,6 +495,12 @@ namespace UnrealBuildTool
 					}
 				}
 			}
+
+			if( BuildConfiguration.bPrintPerformanceInfo )
+			{ 
+				var DeleteTime = (DateTime.UtcNow - DeleteStartTime).TotalSeconds;
+				Log.TraceInformation( "Deleting stale hot reload DLLs took " + DeleteTime + "s" );
+			}
 		}
 
 		/**
@@ -934,10 +950,18 @@ namespace UnrealBuildTool
 		 */
 		static void GatherAllOutdatedActions(UEBuildTarget Target, ActionHistory ActionHistory, ref Dictionary<Action,bool> OutdatedActions, Dictionary<UEBuildTarget, List<FileItem>> TargetToOutdatedPrerequisitesMap )
 		{
+			var CheckOutdatednessStartTime = DateTime.UtcNow;
+
 			foreach (var Action in AllActions)
 			{
 				IsActionOutdated(Target, Action, ref OutdatedActions, ActionHistory, TargetToOutdatedPrerequisitesMap);
 			}
+	
+			if( BuildConfiguration.bPrintPerformanceInfo )
+			{ 
+				var CheckOutdatednessTime = (DateTime.UtcNow - CheckOutdatednessStartTime).TotalSeconds;
+				Log.TraceInformation( "Checking actions for " + Target.GetTargetName() + " took " + CheckOutdatednessTime + "s" );
+			}
 		}
 
 		/**
