Force UAT scripts to be built in Development config (just like UAT).

Disabled UAT code optimizations (shouldn't affect performance but should improve debugging in Development config)

diff --git a/Engine/Source/Programs/AutomationTool/AutomationTool.csproj b/Engine/Source/Programs/AutomationTool/AutomationTool.csproj
--- a/Engine/Source/Programs/AutomationTool/AutomationTool.csproj
+++ b/Engine/Source/Programs/AutomationTool/AutomationTool.csproj
@@ -67,7 +67,7 @@
   <PropertyGroup Condition=" '$(Configuration)|$(Platform)' == 'Development|AnyCPU' ">
     <OutputPath>..\..\..\Binaries\DotNET\</OutputPath>
     <DefineConstants>TRACE</DefineConstants>
-    <Optimize>true</Optimize>
+    <Optimize>false</Optimize>
     <DebugType>pdbonly</DebugType>
     <PlatformTarget>AnyCPU</PlatformTarget>
     <ErrorReport>prompt</ErrorReport>
diff --git a/Engine/Source/Programs/UnrealBuildTool/System/ProjectFileGenerator.cs b/Engine/Source/Programs/UnrealBuildTool/System/ProjectFileGenerator.cs
--- a/Engine/Source/Programs/UnrealBuildTool/System/ProjectFileGenerator.cs
+++ b/Engine/Source/Programs/UnrealBuildTool/System/ProjectFileGenerator.cs
@@ -242,7 +242,7 @@ namespace UnrealBuildTool
 					var RelativeFileName = Utils.MakePathRelativeTo(ProjectFile, MasterProjectRelativePath);
 					var Project = new VCSharpProjectFile(RelativeFileName);
 					Project.ShouldBuildForAllSolutionTargets = true;
-					AddExistingProjectFile(Project);
+					AddExistingProjectFile(Project, bForceDevelopmentConfiguration: true);
 					ModuleProjects.Add(Project);
 
 					Folder.ChildProjects.Add( Project );
