Output an error message if an XML configuration file is badly formatted. Response to https://answers.unrealengine.com/questions/90667/error-on-editor-launch-unrealbuildtool.html.

diff --git a/Engine/Source/Programs/UnrealBuildTool/Utilities/XmlConfigLoader.cs b/Engine/Source/Programs/UnrealBuildTool/Utilities/XmlConfigLoader.cs
--- a/Engine/Source/Programs/UnrealBuildTool/Utilities/XmlConfigLoader.cs
+++ b/Engine/Source/Programs/UnrealBuildTool/Utilities/XmlConfigLoader.cs
@@ -533,7 +533,14 @@ namespace UnrealBuildTool
 					continue;
 				}
 
-				Load(PossibleConfigLocation.FSLocation);
+				try
+				{
+					Load(PossibleConfigLocation.FSLocation);
+				}
+				catch(Exception Ex)
+				{
+					Console.WriteLine("Problem parsing {0}:\n\t{1}", PossibleConfigLocation.FSLocation, Ex.ToString());
+				}
 			}
 		}
 
