Tidy up output formatting for XML config errors.

diff --git a/Engine/Source/Programs/UnrealBuildTool/Utilities/XmlConfigLoader.cs b/Engine/Source/Programs/UnrealBuildTool/Utilities/XmlConfigLoader.cs
--- a/Engine/Source/Programs/UnrealBuildTool/Utilities/XmlConfigLoader.cs
+++ b/Engine/Source/Programs/UnrealBuildTool/Utilities/XmlConfigLoader.cs
@@ -539,7 +539,7 @@ namespace UnrealBuildTool
 				}
 				catch(Exception Ex)
 				{
-					Console.WriteLine("Problem parsing {0}:\n\t{1}", PossibleConfigLocation.FSLocation, Ex.ToString());
+					Console.WriteLine("Problem parsing {0}:\n   {1}", PossibleConfigLocation.FSLocation, Ex.ToString());
 				}
 			}
 		}
