Remove AutomationToolStub and related code from UBT.

diff --git a/Engine/Source/Programs/AutomationToolStub/App.config b/Engine/Source/Programs/AutomationToolStub/App.config
deleted file mode 100644
diff --git a/Engine/Source/Programs/AutomationToolStub/AutomationToolStubTemplate.csproj b/Engine/Source/Programs/AutomationToolStub/AutomationToolStubTemplate.csproj
deleted file mode 100644
diff --git a/Engine/Source/Programs/AutomationToolStub/Program.cs b/Engine/Source/Programs/AutomationToolStub/Program.cs
deleted file mode 100644
diff --git a/Engine/Source/Programs/AutomationToolStub/Properties/AssemblyInfo.cs b/Engine/Source/Programs/AutomationToolStub/Properties/AssemblyInfo.cs
deleted file mode 100644
diff --git a/Engine/Source/Programs/UnrealBuildTool/System/ProjectFileGenerator.cs b/Engine/Source/Programs/UnrealBuildTool/System/ProjectFileGenerator.cs
--- a/Engine/Source/Programs/UnrealBuildTool/System/ProjectFileGenerator.cs
+++ b/Engine/Source/Programs/UnrealBuildTool/System/ProjectFileGenerator.cs
@@ -8,9 +8,6 @@ using System.Reflection;
 using Microsoft.Win32;
 using System.Linq;
 using System.Diagnostics;
-using System.Xml;
-using System.Xml.Linq;
-using System.Xml.XPath;
 
 namespace UnrealBuildTool
 {
@@ -217,7 +214,7 @@ namespace UnrealBuildTool
 		/// <summary>
 		/// Adds all .automation.csproj files to the solution.
 		/// </summary>
-		List<VCSharpProjectFile> AddAutomationModules(MasterProjectFolder ProgramsFolder)
+		void AddAutomationModules(MasterProjectFolder ProgramsFolder)
 		{
 			var Folder = ProgramsFolder.AddSubFolder("Automation");
 			var AllGameFolders = UEBuildTarget.DiscoverAllGameFolders();
@@ -232,8 +229,7 @@ namespace UnrealBuildTool
 			}
 
 			// Find all the automation modules .csproj files to add
-			List<VCSharpProjectFile> ModuleProjects = new List<VCSharpProjectFile>();
-			List<string> ModuleFiles = RulesCompiler.FindAllRulesSourceFiles(RulesCompiler.RulesFileType.AutomationModule, BuildFolders);
+			var ModuleFiles = RulesCompiler.FindAllRulesSourceFiles(RulesCompiler.RulesFileType.AutomationModule, BuildFolders);
 			foreach (var ProjectFile in ModuleFiles)
 			{
 				FileInfo Info = new FileInfo(Path.Combine(ProjectFile));
@@ -243,93 +239,10 @@ namespace UnrealBuildTool
 					var Project = new VCSharpProjectFile(RelativeFileName);
 					Project.ShouldBuildForAllSolutionTargets = true;
 					AddExistingProjectFile(Project, bForceDevelopmentConfiguration: true);
-					ModuleProjects.Add(Project);
 
 					Folder.ChildProjects.Add( Project );
 				}
 			}
-			return ModuleProjects;
-		}
-
-		/// <summary>
-		/// Generates a C# project which references all the given automation modules and can be used to launch UAT directly from the IDE.
-		/// </summary>
-		private VCSharpProjectFile GenerateAutomationToolStubProject(List<VCSharpProjectFile> AutomationModules)
-		{
-			string InputProjectDir = Path.GetFullPath("Programs\\AutomationToolStub");
-			string OutputProjectDir = IntermediateProjectFilesPath;
-			Directory.CreateDirectory(OutputProjectDir);
-
-			// Load the template aggregate project file
-			XDocument Document = XDocument.Load(Utils.CleanDirectorySeparators(Path.Combine(InputProjectDir, "AutomationToolStubTemplate.csproj")));
-			XNamespace Namespace = Document.Root.Name.Namespace;
-
-			// Update the GUID for this project. Using a fixed GUID shouldn't be a problem, but it should be different from the source.
-			XElement SharedPropertyGroup = Document.Root.Elements().First(x => x.Name.LocalName == "PropertyGroup" && x.Attribute("Condition") == null);
-			if(SharedPropertyGroup != null)
-			{
-				SharedPropertyGroup.SetElementValue(Namespace + "ProjectGuid", "{5C41B176-09E5-44D7-8B25-277854234CA4}");
-			}
-
-			// Create a namespace manager that can deal with msbuild files
-			XmlNamespaceManager NamespaceManager = new XmlNamespaceManager(new NameTable());
-			NamespaceManager.AddNamespace("msbuild", Namespace.NamespaceName);
-
-			// Fixup the output directory
-			foreach(XElement PropertyGroup in Document.Root.Elements(Namespace + "PropertyGroup"))
-			{
-				XElement OutputPath = PropertyGroup.Element(Namespace + "OutputPath");
-				if(OutputPath != null)
-				{
-					OutputPath.Value = Utils.MakePathRelativeTo(Path.Combine(InputProjectDir, OutputPath.Value), OutputProjectDir);
-				}
-			}
-
-			// Fixup all the source files
-			foreach(XElement ItemGroup in Document.Root.Elements(Namespace + "ItemGroup"))
-			{
-				foreach(XElement SourceElement in ItemGroup.Elements().Where(x => x.Name.LocalName == "Compile" || x.Name.LocalName == "None" || x.Name.LocalName == "ProjectReference"))
-				{
-					XAttribute Include = SourceElement.Attribute("Include");
-					if(Include != null)
-					{
-						Include.Value = Utils.MakePathRelativeTo(Path.Combine(InputProjectDir, Include.Value), OutputProjectDir);
-					}
-				}
-			}
-
-			// Add references to all the other AutomationTool modules
-			XElement ProjectReferences = new XElement(Namespace + "ItemGroup");
-			foreach(VCSharpProjectFile AutomationModule in AutomationModules)
-			{
-				string ModuleProjectPath = Path.GetFullPath(Path.Combine(MasterProjectRelativePath, AutomationModule.RelativeProjectFilePath));
-				XDocument ModuleDocument = XDocument.Load(ModuleProjectPath);
-				if(ModuleDocument != null)
-				{
-					XElement ModulePropertyGroup = ModuleDocument.Root.Elements().First(x => x.Name.LocalName == "PropertyGroup" && x.Attribute("Condition") == null);
-					XElement ModuleName = ModulePropertyGroup.Element(Namespace + "AssemblyName");
-					XElement ModuleGuid = ModulePropertyGroup.Element(Namespace + "ProjectGuid");
-					if(ModuleName != null && ModuleGuid != null)
-					{
-						XElement ProjectReference = new XElement(Namespace + "ProjectReference");
-						ProjectReference.SetAttributeValue("Include", Utils.MakePathRelativeTo(ModuleProjectPath, OutputProjectDir, false));
-						ProjectReference.Add(new XElement(Namespace + "Name", ModuleName.Value));
-						ProjectReference.Add(new XElement(Namespace + "Project", ModuleGuid.Value));
-						ProjectReference.Add(new XElement(Namespace + "Private", "false"));
-						ProjectReferences.Add(ProjectReference);
-					}
-				}
-			}
-			Document.Root.Add(ProjectReferences);
-
-			// Save the project
-			string Filename = Path.Combine(OutputProjectDir, "AutomationToolStub.csproj");
-			Document.Save(Filename);
-
-			// Add the project file to the list of projects to build
-			VCSharpProjectFile NewProjectFile = new VCSharpProjectFile(Utils.MakePathRelativeTo(Filename, MasterProjectRelativePath));
-			AddExistingProjectFile(NewProjectFile);
-			return NewProjectFile;
 		}
 
 		/// <summary>
@@ -572,10 +485,7 @@ namespace UnrealBuildTool
 					ProgramsFolder.ChildProjects.Add(AddSimpleCSharpProject("AutomationToolLauncher", bShouldBuildForAllSolutionTargets: true, bForceDevelopmentConfiguration: true));
 
 					// Add automation.csproj files to the master project
-					List<VCSharpProjectFile> AutomationModules = AddAutomationModules(ProgramsFolder);
-
-					// Create the automation tool stub project
-					ProgramsFolder.ChildProjects.Add(GenerateAutomationToolStubProject(AutomationModules));
+					AddAutomationModules(ProgramsFolder);
 
 					// Add Distill to the master project
 					ProgramsFolder.ChildProjects.Add(AddSimpleCSharpProject("Distill"));
