- Changed UBT to allow for multiple OutputPaths in a given UE binary (this is needed by Android to output multiple .so files per UBT run, for armv7, x86, etc)

diff --git a/Engine/Source/Programs/UnrealBuildTool/Android/AndroidToolChain.cs b/Engine/Source/Programs/UnrealBuildTool/Android/AndroidToolChain.cs
--- a/Engine/Source/Programs/UnrealBuildTool/Android/AndroidToolChain.cs
+++ b/Engine/Source/Programs/UnrealBuildTool/Android/AndroidToolChain.cs
@@ -469,15 +469,27 @@ namespace UnrealBuildTool
 
 		static bool ShouldSkipLib(string Lib, string Arch)
 		{
+			// reject any libs we outright don't want to link with
 			foreach (var LibName in LibrariesToSkip[Arch])
 			{
-				// if there's a directory in the path with a bad architecture name, reject it
 				if (LibName == Lib)
 				{
 					return true;
 				}
 			}
 
+			// if another architecture is in the filename, reject it
+			foreach (string ArchName in Arches)
+			{
+				if (ArchName != Arch)
+				{
+					if (Path.GetFileNameWithoutExtension(Lib).EndsWith(ArchName))
+					{
+						return true;
+					}
+				}
+			}
+
 			// if nothing was found, we are okay
 			return false;
 		}
@@ -685,12 +697,31 @@ namespace UnrealBuildTool
 			return Path.Combine(Path.GetDirectoryName(Pathname), Path.GetFileNameWithoutExtension(Pathname) + Arch + Path.GetExtension(Pathname));
 		}
 
+		static public string RemoveArchName(string Pathname)
+		{
+			// remove all architecture names
+			foreach (string Arch in Arches)
+			{
+				Pathname = Path.Combine(Path.GetDirectoryName(Pathname), Path.GetFileName(Pathname).Replace(Arch, ""));
+			}
+			return Pathname;
+		}
+
 		public override FileItem[] LinkAllFiles(LinkEnvironment LinkEnvironment, bool bBuildImportLibraryOnly)
 		{
 			List<FileItem> Outputs = new List<FileItem>();
 
-			foreach (string Arch in Arches)
+			for (int ArchIndex = 0; ArchIndex < Arches.Length; ArchIndex++)
 			{
+				string Arch = Arches[ArchIndex];
+
+				// Android will have an array of outputs
+				if (LinkEnvironment.Config.OutputFilePaths.Length < ArchIndex || 
+					!Path.GetFileNameWithoutExtension(LinkEnvironment.Config.OutputFilePaths[ArchIndex]).EndsWith(Arch))
+				{
+					throw new BuildException("The OutputFilePaths array didn't match the Arches array in AndroidToolChain.LinkAllFiles");
+				}
+
 				// Create an action that invokes the linker.
 				Action LinkAction = new Action(ActionType.Link);
 				LinkAction.WorkingDirectory = Path.GetFullPath(".");
@@ -713,15 +744,7 @@ namespace UnrealBuildTool
 
 				// Add the output file as a production of the link action.
 				FileItem OutputFile;
-				// libraries don't have the platform name in it, so inline the name
-				if (LinkEnvironment.Config.bIsBuildingLibrary)
-				{
-					OutputFile = FileItem.GetItemByPath(InlineArchName(LinkEnvironment.Config.OutputFilePath, Arch));
-				}
-				else
-				{
-					OutputFile = FileItem.GetItemByPath(LinkEnvironment.Config.OutputFilePath);
-				}
+				OutputFile = FileItem.GetItemByPath(LinkEnvironment.Config.OutputFilePaths[ArchIndex]);
 				Outputs.Add(OutputFile);
 				LinkAction.ProducedItems.Add(OutputFile);
 				LinkAction.StatusDescription = string.Format("{0}", Path.GetFileName(OutputFile.AbsolutePath));
@@ -781,23 +804,18 @@ namespace UnrealBuildTool
 					LinkAction.CommandArguments += string.Format(" -Wl,--start-group");
 					foreach (string AdditionalLibrary in LinkEnvironment.Config.AdditionalLibraries)
 					{
-						if (String.IsNullOrEmpty(Path.GetDirectoryName(AdditionalLibrary)))
+						if (!ShouldSkipLib(AdditionalLibrary, Arch))
 						{
-							if (!ShouldSkipLib(AdditionalLibrary, Arch))
+							if (String.IsNullOrEmpty(Path.GetDirectoryName(AdditionalLibrary)))
 							{
 								LinkAction.CommandArguments += string.Format(" \"-l{0}\"", AdditionalLibrary);
 							}
-						}
-						else
-						{
-							// full pathed libs are compiled by us, so we depend on linking them (look for extra architecture name if needed)
-							string LibPath = AdditionalLibrary;
-							if (!File.Exists(AdditionalLibrary))
+							else
 							{
-								LibPath = InlineArchName(AdditionalLibrary, Arch);
+								// full pathed libs are compiled by us, so we depend on linking them
+								LinkAction.CommandArguments += string.Format(" \"{0}\"", AdditionalLibrary);
+								LinkAction.PrerequisiteItems.Add(FileItem.GetItemByPath(AdditionalLibrary));
 							}
-							LinkAction.CommandArguments += string.Format(" \"{0}\"", LibPath);
-							LinkAction.PrerequisiteItems.Add(FileItem.GetItemByPath(LibPath));
 						}
 					}
 					LinkAction.CommandArguments += string.Format(" -Wl,--end-group");
@@ -814,6 +832,13 @@ namespace UnrealBuildTool
 			return Outputs.ToArray();
 		}
 
+		public override void AddFilesToManifest(ref FileManifest Manifest, UEBuildBinary Binary)
+		{
+			// the binary will have all of the .so's in the output files, we need to trim down to the shared apk (which is what needs to go into the manifest)
+			string ApkFile = Path.ChangeExtension(RemoveArchName(Binary.Config.OutputFilePaths[0]), ".apk");
+			Manifest.AddFileName(ApkFile);
+		}
+
 		public override void CompileCSharpProject(CSharpEnvironment CompileEnvironment, string ProjectFileName, string DestinationFile)
 		{
             throw new BuildException("Android cannot compile C# files");
diff --git a/Engine/Source/Programs/UnrealBuildTool/Android/UEBuildAndroid.cs b/Engine/Source/Programs/UnrealBuildTool/Android/UEBuildAndroid.cs
--- a/Engine/Source/Programs/UnrealBuildTool/Android/UEBuildAndroid.cs
+++ b/Engine/Source/Programs/UnrealBuildTool/Android/UEBuildAndroid.cs
@@ -153,7 +153,8 @@ namespace UnrealBuildTool
 		public override bool ShouldUsePCHFiles(CPPTargetPlatform Platform, CPPTargetConfiguration Configuration)
 		{
 			// if we compile multiple architectures at once, PCH doesn't work right, so only use PCH when building for one platform
-			return AndroidToolChain.GetAllArchitectures().Length == 1;
+			// @todo android: Handle PCH files in the multiple architecture case
+			return false;// AndroidToolChain.GetAllArchitectures().Length == 1;
 		}
 
 		public override string GetDebugInfoExtension(UEBuildBinaryType InBinaryType)
@@ -349,47 +350,18 @@ namespace UnrealBuildTool
 			};
 		}
 
-		public override void SetupBinaries(UEBuildTarget InBuildTarget)
+		public override string[] FinalizeBinaryPaths(string BinaryName)
 		{
-			// for "fat" binaries, we need to expand the .so's across the platforms
 			string[] Architectures = AndroidToolChain.GetAllArchitectures();
 
-			// all .so's have architecture name, the manifest won't by default, so expand across the binaries
-			List<UEBuildBinary> NewBinaries = new List<UEBuildBinary>();
+			// make multiple output binaries
+			List<string> AllBinaries = new List<string>();
 			foreach (string Architecture in Architectures)
 			{
-				foreach (var Binary in InBuildTarget.AppBinaries)
-				{
-					UEBuildBinaryConfiguration NewConfig = new UEBuildBinaryConfiguration(
-																	InType: Binary.Config.Type,
-																	InOutputFilePath: AndroidToolChain.InlineArchName(Binary.Config.OutputFilePath, Architecture),
-																	InIntermediateDirectory: Binary.Config.IntermediateDirectory,
-																	bInCreateImportLibrarySeparately: Binary.Config.bCreateImportLibrarySeparately,
-																	bInAllowExports: Binary.Config.bAllowExports,
-																	InModuleNames: Binary.Config.ModuleNames);
-
-					NewBinaries.Add(new UEBuildBinaryCPP(InBuildTarget, NewConfig));
-				}
-			}
-
-
-			// add the .apk to the binaries (UBT always generates the apk without an extension with all .so's inside)
-			foreach (var Binary in InBuildTarget.AppBinaries)
-			{
-				// make a binary that just points to the .stub of this executable
-				UEBuildBinaryConfiguration NewConfig = new UEBuildBinaryConfiguration(
-																InType: Binary.Config.Type,
-																InOutputFilePath: Path.ChangeExtension(Binary.Config.OutputFilePath, ".apk"),
-																InIntermediateDirectory: Binary.Config.IntermediateDirectory,
-																bInCreateImportLibrarySeparately: Binary.Config.bCreateImportLibrarySeparately,
-																bInAllowExports: Binary.Config.bAllowExports,
-																InModuleNames: Binary.Config.ModuleNames);
-
-				NewBinaries.Add(new UEStubDummyBinary(InBuildTarget, NewConfig));
+				AllBinaries.Add(AndroidToolChain.InlineArchName(BinaryName, Architecture));
 			}
 
-			// replace the old binaries with the new
-			InBuildTarget.AppBinaries = NewBinaries;
+			return AllBinaries.ToArray();
 		}
 	}
 }
diff --git a/Engine/Source/Programs/UnrealBuildTool/Android/UEDeployAndroid.cs b/Engine/Source/Programs/UnrealBuildTool/Android/UEDeployAndroid.cs
--- a/Engine/Source/Programs/UnrealBuildTool/Android/UEDeployAndroid.cs
+++ b/Engine/Source/Programs/UnrealBuildTool/Android/UEDeployAndroid.cs
@@ -269,7 +269,18 @@ namespace UnrealBuildTool.Android
 				case "-armv7": return "armeabi-v7a";
 				case "-x86": return "x86";
 
-				default: throw new BuildException("Unknown architecture {0}", UE4Arch);
+				default: throw new BuildException("Unknown UE4 architecture {0}", UE4Arch);
+			}
+		}
+
+		public static string GetUE4Arch(string NDKArch)
+		{
+			switch (NDKArch)
+			{
+				case "armeabi-v7a": return "-armv7";
+				case "x86":			return "-x86";
+
+				default: throw new BuildException("Unknown NDK architecture '{0}'", NDKArch);
 			}
 		}
 
@@ -715,8 +726,10 @@ namespace UnrealBuildTool.Android
 
 		public override bool PrepTargetForDeployment(UEBuildTarget InTarget)
 		{
+			// we need to strip architecture from any of the output paths
+			string BaseSoName = AndroidToolChain.RemoveArchName(InTarget.OutputPaths[0]);
 			// this always makes a merged .apk since for debugging, there's no way to know which one to run
-			MakeApk(InTarget.AppName, InTarget.ProjectDirectory, InTarget.OutputPath, BuildConfiguration.RelativeEnginePath, bForDistribution:false, CookFlavor:"", bMakeSeparateApks:false);
+			MakeApk(InTarget.AppName, InTarget.ProjectDirectory, BaseSoName, BuildConfiguration.RelativeEnginePath, bForDistribution:false, CookFlavor:"", bMakeSeparateApks:false);
 			return true;
 		}
 
diff --git a/Engine/Source/Programs/UnrealBuildTool/Configuration/UEBuildBinary.cs b/Engine/Source/Programs/UnrealBuildTool/Configuration/UEBuildBinary.cs
--- a/Engine/Source/Programs/UnrealBuildTool/Configuration/UEBuildBinary.cs
+++ b/Engine/Source/Programs/UnrealBuildTool/Configuration/UEBuildBinary.cs
@@ -37,12 +37,42 @@ namespace UnrealBuildTool
 		/// <summary>
 		/// The output file path. This must be set before a binary can be built using it.
 		/// </summary>
-		public string OutputFilePath;
+		public string[] OutputFilePaths;
+
+		/// <summary>
+		/// Returns the OutputFilePath is there is only one entry in OutputFilePaths
+		/// </summary>
+		public string OutputFilePath
+		{
+			get
+			{
+				if (OutputFilePaths.Length != 1)
+				{
+					throw new BuildException("Attempted to use UEBuildBinaryConfiguration.OutputFilePath property, but there are multiple (or no) OutputFilePaths. You need to handle multiple in the code that called this (size = {0})", OutputFilePaths.Length);
+				}
+				return OutputFilePaths[0];
+			}
+		}
 
 		/// <summary>
 		/// Original output filepath. This is the original binary name before hot-reload suffix has been appended to it.
 		/// </summary>
-		public string OriginalOutputFilePath;
+		public string[] OriginalOutputFilePaths;
+
+		/// <summary>
+		/// Returns the OutputFilePath is there is only one entry in OutputFilePaths
+		/// </summary>
+		public string OriginalOutputFilePath
+		{
+			get
+			{
+				if (OutputFilePaths.Length != 1)
+				{
+					throw new BuildException("Attempted to use UEBuildBinaryConfiguration.OutputFilePath property, but there are multiple (or no) OutputFilePaths. You need to handle multiple in the code that called this (size = {0})", OutputFilePaths.Length);
+				}
+				return OutputFilePaths[0];
+			}
+		}
 
 		/// <summary>
 		/// The intermediate directory for this binary. Modules should create separate intermediate directories below this. Must be set before a binary can be built using it.
@@ -116,7 +146,7 @@ namespace UnrealBuildTool
 		/// <param name="InModuleNames"></param>
 		public UEBuildBinaryConfiguration(
 				UEBuildBinaryType InType,
-				string InOutputFilePath = null,
+				string[] InOutputFilePaths = null,
 				string InIntermediateDirectory = null,
 				bool bInAllowExports = false,
 				bool bInCreateImportLibrarySeparately = false,
@@ -132,7 +162,7 @@ namespace UnrealBuildTool
 			)
 		{
 			Type = InType;
-			OutputFilePath = InOutputFilePath;
+			OutputFilePaths = InOutputFilePaths;
 			IntermediateDirectory = InIntermediateDirectory;
 			bAllowExports = bInAllowExports;
 			bCreateImportLibrarySeparately = bInCreateImportLibrarySeparately;
@@ -350,7 +380,10 @@ namespace UnrealBuildTool
 				if (Module != null && (Target.Rules == null || Target.Rules.bOutputToEngineBinaries == false))
 				{
 					// Fix up the binary path if this is module specifies an alternate output directory
-					Config.OutputFilePath = Module.FixupOutputPath(Config.OutputFilePath);
+					for (int Index = 0; Index < Config.OutputFilePaths.Length; Index++ )
+					{
+						Config.OutputFilePaths[Index] = Module.FixupOutputPath(Config.OutputFilePaths[Index]);
+					}
 				}
 			}
 		}
@@ -485,10 +518,11 @@ namespace UnrealBuildTool
 			// this for each file when generating their dependencies.
 			BinaryCompileEnvironment.bHackHeaderGenerator = (Target.GetAppName() == "UnrealHeaderTool");
 
+			// @todo: This should be in some Windows code somewhere...
 			// Set the original file name macro; used in PCLaunch.rc to set the binary metadata fields.
-			var OriginalFilename = !String.IsNullOrEmpty(Config.OriginalOutputFilePath) ?
-				Path.GetFileName(Config.OriginalOutputFilePath) :
-				Path.GetFileName(Config.OutputFilePath);
+			var OriginalFilename = (Config.OriginalOutputFilePaths != null) ?
+				Path.GetFileName(Config.OriginalOutputFilePaths[0]) :
+				Path.GetFileName(Config.OutputFilePaths[0]);
 			BinaryCompileEnvironment.Config.Definitions.Add("ORIGINAL_FILE_NAME=\"" + OriginalFilename + "\"");
 
 			foreach (var ModuleName in ModuleNames)
@@ -535,7 +569,7 @@ namespace UnrealBuildTool
 			}
 
 			// Set the link output file.
-			BinaryLinkEnvironment.Config.OutputFilePath = Config.OutputFilePath;
+			BinaryLinkEnvironment.Config.OutputFilePaths = Config.OutputFilePaths;
 
 			// Set whether the link is allowed to have exports.
 			BinaryLinkEnvironment.Config.bHasExports = Config.bAllowExports;
@@ -544,7 +578,7 @@ namespace UnrealBuildTool
 			BinaryLinkEnvironment.Config.IntermediateDirectory = Config.IntermediateDirectory;
 
 			// Put the non-executable output files (PDB, import library, etc) in the same directory as the production
-			BinaryLinkEnvironment.Config.OutputDirectory = Path.GetDirectoryName(Config.OutputFilePath);
+			BinaryLinkEnvironment.Config.OutputDirectory = Path.GetDirectoryName(Config.OutputFilePaths[0]);
 
 			// Determine the type of binary we're linking.
 			switch (Config.Type)
@@ -602,7 +636,10 @@ namespace UnrealBuildTool
 					var ConsoleAppLinkEvironment = BinaryLinkEnvironment.DeepCopy();
 					ConsoleAppLinkEvironment.Config.bIsBuildingConsoleApplication = true;
 					ConsoleAppLinkEvironment.Config.WindowsEntryPointOverride = "WinMainCRTStartup";		// For WinMain() instead of "main()" for Launch module
-					ConsoleAppLinkEvironment.Config.OutputFilePath = GetAdditionalConsoleAppPath(ConsoleAppLinkEvironment.Config.OutputFilePath);
+					for (int Index = 0; Index < Config.OutputFilePaths.Length; Index++)
+					{
+						ConsoleAppLinkEvironment.Config.OutputFilePaths[Index] = GetAdditionalConsoleAppPath(ConsoleAppLinkEvironment.Config.OutputFilePaths[Index]);
+					}
 
 					// Link the console app executable
 					OutputFiles.AddRange(ConsoleAppLinkEvironment.LinkExecutable(false));
@@ -624,18 +661,21 @@ namespace UnrealBuildTool
 		/// <param name="DependentLinkEnvironment">The link environment of the dependency</param>
 		public override void SetupDependentLinkEnvironment(ref LinkEnvironment DependentLinkEnvironment)
 		{
-			string LibraryFileName;
-			if (Config.Type == UEBuildBinaryType.StaticLibrary
-                || DependentLinkEnvironment.Config.Target.Platform == CPPTargetPlatform.Mac
-                || DependentLinkEnvironment.Config.Target.Platform == CPPTargetPlatform.Linux)
+			foreach (string OutputFilePath in Config.OutputFilePaths)
 			{
-				LibraryFileName = Config.OutputFilePath;
-			}
-			else
-			{
-				LibraryFileName = Path.Combine(Config.IntermediateDirectory, Path.GetFileNameWithoutExtension(Config.OutputFilePath) + ".lib");
+				string LibraryFileName;
+				if (Config.Type == UEBuildBinaryType.StaticLibrary
+					|| DependentLinkEnvironment.Config.Target.Platform == CPPTargetPlatform.Mac
+					|| DependentLinkEnvironment.Config.Target.Platform == CPPTargetPlatform.Linux)
+				{
+					LibraryFileName = OutputFilePath;
+				}
+				else
+				{
+					LibraryFileName = Path.Combine(Config.IntermediateDirectory, Path.GetFileNameWithoutExtension(OutputFilePath) + ".lib");
+				}
+				DependentLinkEnvironment.Config.AdditionalLibraries.Add(LibraryFileName);
 			}
-			DependentLinkEnvironment.Config.AdditionalLibraries.Add(LibraryFileName);
 		}
 
 		/// <summary>
diff --git a/Engine/Source/Programs/UnrealBuildTool/Configuration/UEBuildClient.cs b/Engine/Source/Programs/UnrealBuildTool/Configuration/UEBuildClient.cs
--- a/Engine/Source/Programs/UnrealBuildTool/Configuration/UEBuildClient.cs
+++ b/Engine/Source/Programs/UnrealBuildTool/Configuration/UEBuildClient.cs
@@ -44,7 +44,10 @@ namespace UnrealBuildTool
                     // We want the output to go into the <GAME>\Binaries folder
                     if (InRulesObject.bOutputToEngineBinaries == false)
                     {
-                        OutputPath = OutputPath.Replace("Engine\\Binaries", InGameName + "\\Binaries");
+						for (int Index = 0; Index < OutputPaths.Length; Index++ )
+						{
+							OutputPaths[Index] = OutputPaths[Index].Replace("Engine\\Binaries", InGameName + "\\Binaries");
+						}
                     }
                 }
             }
@@ -57,13 +60,13 @@ namespace UnrealBuildTool
             {
                 // Make the game executable.
                 UEBuildBinaryConfiguration Config = new UEBuildBinaryConfiguration(InType: UEBuildBinaryType.Executable,
-                                                                                    InOutputFilePath: OutputPath,
+                                                                                    InOutputFilePaths: OutputPaths,
 																					InIntermediateDirectory: EngineIntermediateDirectory,
                                                                                     bInCreateImportLibrarySeparately: (ShouldCompileMonolithic() ? false : true),
                                                                                     bInAllowExports: !ShouldCompileMonolithic(),
                                                                                     InModuleNames: new List<string>() { "Launch" });
 
-                AppBinaries.Add(new UEBuildBinaryCPP(this, Config));
+				AppBinaries.Add(new UEBuildBinaryCPP(this, Config));
             }
 
             // Add the other modules that we want to compile along with the executable.  These aren't necessarily
diff --git a/Engine/Source/Programs/UnrealBuildTool/Configuration/UEBuildEditor.cs b/Engine/Source/Programs/UnrealBuildTool/Configuration/UEBuildEditor.cs
--- a/Engine/Source/Programs/UnrealBuildTool/Configuration/UEBuildEditor.cs
+++ b/Engine/Source/Programs/UnrealBuildTool/Configuration/UEBuildEditor.cs
@@ -52,7 +52,7 @@ namespace UnrealBuildTool
 			{
 				// Make the editor executable.
 				UEBuildBinaryConfiguration Config = new UEBuildBinaryConfiguration( InType: UEBuildBinaryType.Executable,
-																					InOutputFilePath: OutputPath,
+																					InOutputFilePaths: OutputPaths,
 																					InIntermediateDirectory: EngineIntermediateDirectory,
 																					bInCreateImportLibrarySeparately: (ShouldCompileMonolithic() ? false : true),
 																					bInAllowExports:!ShouldCompileMonolithic(),
diff --git a/Engine/Source/Programs/UnrealBuildTool/Configuration/UEBuildGame.cs b/Engine/Source/Programs/UnrealBuildTool/Configuration/UEBuildGame.cs
--- a/Engine/Source/Programs/UnrealBuildTool/Configuration/UEBuildGame.cs
+++ b/Engine/Source/Programs/UnrealBuildTool/Configuration/UEBuildGame.cs
@@ -44,7 +44,10 @@ namespace UnrealBuildTool
 					// We want the output to go into the <GAME>\Binaries folder
 					if (InRulesObject.bOutputToEngineBinaries == false)
 					{
-						OutputPath = OutputPath.Replace("Engine\\Binaries", InGameName + "\\Binaries");
+						for (int Index = 0; Index < OutputPaths.Length; Index++)
+						{
+							OutputPaths[Index] = OutputPaths[Index].Replace("Engine\\Binaries", InGameName + "\\Binaries");
+						}
 					}
 				}
 			}
@@ -72,7 +75,7 @@ namespace UnrealBuildTool
 			{
 				// Make the game executable.
 				UEBuildBinaryConfiguration Config = new UEBuildBinaryConfiguration( InType: UEBuildBinaryType.Executable,
-																					InOutputFilePath: OutputPath,
+																					InOutputFilePaths: OutputPaths,
 																					InIntermediateDirectory: EngineIntermediateDirectory,
 																					bInCreateImportLibrarySeparately: (ShouldCompileMonolithic() ? false : true),
 																					bInAllowExports:!ShouldCompileMonolithic(),
@@ -87,9 +90,6 @@ namespace UnrealBuildTool
 				// Modules should properly identify the 'extra modules' they need now.
 				// There should be nothing here!
 			}
-
-			// Allow the platform to setup binaries
-			UEBuildPlatform.GetBuildPlatform(Platform).SetupBinaries(this);
 		}
 
 		public override void SetupDefaultGlobalEnvironment(
diff --git a/Engine/Source/Programs/UnrealBuildTool/Configuration/UEBuildModule.cs b/Engine/Source/Programs/UnrealBuildTool/Configuration/UEBuildModule.cs
--- a/Engine/Source/Programs/UnrealBuildTool/Configuration/UEBuildModule.cs
+++ b/Engine/Source/Programs/UnrealBuildTool/Configuration/UEBuildModule.cs
@@ -544,6 +544,9 @@ namespace UnrealBuildTool
 				// If this module is being built into a DLL or EXE, set up an IMPORTS or EXPORTS definition for it.
 				if (Binary != null)
 				{
+					string BinaryPath = Binary.Config.OutputFilePaths[0];
+					string SourceBinaryPath = SourceBinary.Config.OutputFilePaths[0];
+
 					if (ProjectFileGenerator.bGenerateProjectFiles || (Binary.Config.Type == UEBuildBinaryType.StaticLibrary))
 					{
 						// When generating IntelliSense files, never add dllimport/dllexport specifiers as it
@@ -554,12 +557,12 @@ namespace UnrealBuildTool
 					{
 						if( Binary.Config.bAllowExports )
 						{
-							Log.TraceVerbose( "{0}: Exporting {1} from {2}", Path.GetFileNameWithoutExtension( SourceBinary.Config.OutputFilePath ), Name, Path.GetFileNameWithoutExtension( Binary.Config.OutputFilePath ) );
+							Log.TraceVerbose( "{0}: Exporting {1} from {2}", Path.GetFileNameWithoutExtension( SourceBinaryPath ), Name, Path.GetFileNameWithoutExtension( BinaryPath ) );
 							Definitions.Add( Name.ToUpperInvariant() + "_API=DLLEXPORT" );
 						}
 						else
 						{
-							Log.TraceVerbose( "{0}: Not importing/exporting {1} (binary: {2})", Path.GetFileNameWithoutExtension( SourceBinary.Config.OutputFilePath ), Name, Path.GetFileNameWithoutExtension( Binary.Config.OutputFilePath ) );
+							Log.TraceVerbose( "{0}: Not importing/exporting {1} (binary: {2})", Path.GetFileNameWithoutExtension( SourceBinaryPath ), Name, Path.GetFileNameWithoutExtension( BinaryPath ) );
 							Definitions.Add( Name.ToUpperInvariant() + "_API=" );
 						}
 					}
@@ -570,17 +573,17 @@ namespace UnrealBuildTool
 						//		this, we need to suppress warnings at compile time.
 						if( bIncludePathsOnly )
 						{
-							Log.TraceVerbose( "{0}: Include paths only for {1} (binary: {2})", Path.GetFileNameWithoutExtension( SourceBinary.Config.OutputFilePath ), Name, Path.GetFileNameWithoutExtension( Binary.Config.OutputFilePath ) );
+							Log.TraceVerbose( "{0}: Include paths only for {1} (binary: {2})", Path.GetFileNameWithoutExtension( SourceBinaryPath ), Name, Path.GetFileNameWithoutExtension( BinaryPath ) );
 							Definitions.Add( Name.ToUpperInvariant() + "_API=" );
 						}
 						else if (Binary.Config.bAllowExports)
 						{
-							Log.TraceVerbose( "{0}: Importing {1} from {2}", Path.GetFileNameWithoutExtension( SourceBinary.Config.OutputFilePath ), Name, Path.GetFileNameWithoutExtension( Binary.Config.OutputFilePath ) );
+							Log.TraceVerbose( "{0}: Importing {1} from {2}", Path.GetFileNameWithoutExtension( SourceBinaryPath ), Name, Path.GetFileNameWithoutExtension( BinaryPath ) );
 							Definitions.Add( Name.ToUpperInvariant() + "_API=DLLIMPORT" );
 						}
 						else
 						{
-							Log.TraceVerbose("{0}: Not importing/exporting {1} (binary: {2})", Path.GetFileNameWithoutExtension(SourceBinary.Config.OutputFilePath), Name, Path.GetFileNameWithoutExtension(Binary.Config.OutputFilePath));
+							Log.TraceVerbose( "{0}: Not importing/exporting {1} (binary: {2})", Path.GetFileNameWithoutExtension( SourceBinaryPath ), Name, Path.GetFileNameWithoutExtension( BinaryPath ));
 							Definitions.Add( Name.ToUpperInvariant() + "_API=" );
 						}
 					}
@@ -961,7 +964,7 @@ namespace UnrealBuildTool
 		public bool bEnableExceptions = false;
 
 		/** Path to this module's redist static library */
-		public string RedistStaticLibraryPath = null;
+		public string[] RedistStaticLibraryPaths = null;
 
 		/** Whether we're building the redist static library (as well as using it). */
 		public bool bBuildingRedistStaticLibrary = false;
@@ -1111,7 +1114,7 @@ namespace UnrealBuildTool
 				return LinkInputFiles;
 			}
 
-			if(RedistStaticLibraryPath != null && !bBuildingRedistStaticLibrary)
+			if(RedistStaticLibraryPaths != null && !bBuildingRedistStaticLibrary)
 			{
 				// The redist static library will be added in SetupPrivateLinkEnvironment
 				return LinkInputFiles;
@@ -1480,7 +1483,7 @@ namespace UnrealBuildTool
 			LinkInputFiles.AddRange(CPPCompileEnvironment.CompileFiles( Target, SourceFilesToBuild.MMFiles, Name).ObjectFiles);
 
 			// If we're building Rocket, generate a static library for this module
-			if(RedistStaticLibraryPath != null)
+			if(RedistStaticLibraryPaths != null)
 			{
 				// Create a link environment for it
 				LinkEnvironment RedistLinkEnvironment = new LinkEnvironment();
@@ -1489,7 +1492,7 @@ namespace UnrealBuildTool
 				RedistLinkEnvironment.Config.bIsBuildingDLL        = false;
 				RedistLinkEnvironment.Config.bIsBuildingLibrary    = true;
 				RedistLinkEnvironment.Config.IntermediateDirectory = Binary.Config.IntermediateDirectory;
-				RedistLinkEnvironment.Config.OutputFilePath        = RedistStaticLibraryPath;
+				RedistLinkEnvironment.Config.OutputFilePaths        = RedistStaticLibraryPaths;
 
 				// Replace the items built so far with the library
 				RedistLinkEnvironment.LinkExecutable(false);
@@ -1776,9 +1779,12 @@ namespace UnrealBuildTool
 		{
 			base.SetupPrivateLinkEnvironment(ref LinkEnvironment,ref BinaryDependencies,ref VisitedModules);
 
-			if(RedistStaticLibraryPath != null)
+			if(RedistStaticLibraryPaths != null)
 			{
-				LinkEnvironment.Config.AdditionalLibraries.Add(Path.GetFullPath(RedistStaticLibraryPath));
+				foreach (string LibVar in RedistStaticLibraryPaths)
+				{
+					LinkEnvironment.Config.AdditionalLibraries.Add(Path.GetFullPath(LibVar));
+				}
 			}
 		}
 
@@ -1950,7 +1956,7 @@ namespace UnrealBuildTool
 								// Is this a plugin module?
 								var PluginInfo = Plugins.GetPluginInfoForModule( DependencyName );
 
-								string OutputFilePath = Target.MakeBinaryPath(DependencyModule.Name, Target.GetAppName() + "-" + DependencyModule.Name, UEBuildBinaryType.DynamicLinkLibrary, Target.TargetType, bIsRocketModule, PluginInfo, "");
+								string[] OutputFilePaths = Target.MakeBinaryPaths(DependencyModule.Name, Target.GetAppName() + "-" + DependencyModule.Name, UEBuildBinaryType.DynamicLinkLibrary, Target.TargetType, bIsRocketModule, PluginInfo, "");
 
 								// If it's an engine module, output intermediates to the engine intermediates directory. 
 								string IntermediateDirectory = Binary.Config.IntermediateDirectory;
@@ -1961,7 +1967,7 @@ namespace UnrealBuildTool
 
 								// When using modular linkage, unbound modules will be linked into their own DLL files
 								UEBuildBinaryConfiguration Config = new UEBuildBinaryConfiguration( InType: UEBuildBinaryType.DynamicLinkLibrary,
-																									InOutputFilePath: OutputFilePath,
+																									InOutputFilePaths: OutputFilePaths,
 																									InIntermediateDirectory: IntermediateDirectory,
 																									bInAllowExports: true,
 																									InModuleNames: new List<string> { DependencyModule.Name },
@@ -1971,8 +1977,11 @@ namespace UnrealBuildTool
 																									bInCompileMonolithic: Target.ShouldCompileMonolithic() );
 
 								// Fix up the binary path if this is module specifies an alternate output directory
-								Config.OutputFilePath = DependencyModule.FixupOutputPath(Config.OutputFilePath);
-
+								for (int Index = 0; Index < Config.OutputFilePaths.Length; Index++)
+								{
+									Config.OutputFilePaths[Index] = DependencyModule.FixupOutputPath(Config.OutputFilePaths[Index]);
+								}
+	
 								BinaryToBindTo = new UEBuildBinaryCPP( Target, Config );
 							}
 
diff --git a/Engine/Source/Programs/UnrealBuildTool/Configuration/UEBuildPlatform.cs b/Engine/Source/Programs/UnrealBuildTool/Configuration/UEBuildPlatform.cs
--- a/Engine/Source/Programs/UnrealBuildTool/Configuration/UEBuildPlatform.cs
+++ b/Engine/Source/Programs/UnrealBuildTool/Configuration/UEBuildPlatform.cs
@@ -171,6 +171,12 @@ namespace UnrealBuildTool
 		 */
 		string ApplyArchitectureName(string BinaryName);
 
+		/**
+		 * For platforms that need to output multiple files per binary (ie Android "fat" binaries)
+		 * this will emit multiple paths. By default, it simply makes an array from the input
+		 */
+		string[] FinalizeBinaryPaths(string BinaryName);
+
 		/**
 		 * Setup the configuration environment for building
 		 * 
@@ -644,6 +650,16 @@ namespace UnrealBuildTool
 			return BinaryName + GetActiveArchitecture();
 		}
 
+		/**
+		 * For platforms that need to output multiple files per binary (ie Android "fat" binaries)
+		 * this will emit multiple paths. By default, it simply makes an array from the input
+		 */
+		public virtual string[] FinalizeBinaryPaths(string BinaryName)
+		{
+			List<string> TempList = new List<string>() { BinaryName };
+			return TempList.ToArray();
+		}
+
 		/**
 		 *	Setup the configuration environment for building
 		 *	
@@ -657,7 +673,7 @@ namespace UnrealBuildTool
 			//@todo SAS: Add a true Debug mode!
 			if (UnrealBuildTool.RunningRocket())
 			{
-				if (Utils.IsFileUnderDirectory(InBuildTarget.OutputPath, UnrealBuildTool.GetUProjectPath()))
+				if (Utils.IsFileUnderDirectory(InBuildTarget.OutputPaths[0], UnrealBuildTool.GetUProjectPath()))
 				{
 					if (CheckConfig == UnrealTargetConfiguration.Debug)
 					{
diff --git a/Engine/Source/Programs/UnrealBuildTool/Configuration/UEBuildServer.cs b/Engine/Source/Programs/UnrealBuildTool/Configuration/UEBuildServer.cs
--- a/Engine/Source/Programs/UnrealBuildTool/Configuration/UEBuildServer.cs
+++ b/Engine/Source/Programs/UnrealBuildTool/Configuration/UEBuildServer.cs
@@ -50,7 +50,7 @@ namespace UnrealBuildTool
 			{
 				// Make the game executable.
 				UEBuildBinaryConfiguration Config = new UEBuildBinaryConfiguration( InType: UEBuildBinaryType.Executable,
-																					InOutputFilePath: OutputPath,
+																					InOutputFilePaths: OutputPaths,
 																					InIntermediateDirectory: EngineIntermediateDirectory,
 																					bInCreateImportLibrarySeparately: (ShouldCompileMonolithic() ? false : true),
 																					bInAllowExports:!ShouldCompileMonolithic(),
diff --git a/Engine/Source/Programs/UnrealBuildTool/Configuration/UEBuildTarget.cs b/Engine/Source/Programs/UnrealBuildTool/Configuration/UEBuildTarget.cs
--- a/Engine/Source/Programs/UnrealBuildTool/Configuration/UEBuildTarget.cs
+++ b/Engine/Source/Programs/UnrealBuildTool/Configuration/UEBuildTarget.cs
@@ -642,8 +642,21 @@ namespace UnrealBuildTool
 		/** Directory for engine intermediates. For an agnostic editor/game executable, this will be under the engine directory. For monolithic executables this will be the same as the project intermediate directory. */
 		public string EngineIntermediateDirectory;
 
-		/** Output path of final executable. */
-		public string OutputPath;
+		/** Output paths of final executable. */
+		public string[] OutputPaths;
+
+		/** Returns the OutputPath is there is only one entry in OutputPaths */
+		public string OutputPath
+		{
+			get
+			{
+				if (OutputPaths.Length != 1)
+				{
+					throw new BuildException("Attempted to use UEBuildTarget.OutputPath property, but there are multiple (or no) OutputPaths. You need to handle multiple in the code that called this (size = {0})", OutputPaths.Length);
+				}
+				return OutputPaths[0];
+			}
+		}
 
 		/** Remote path of the binary if it is to be synced with CookerSync */
 		public string RemoteRoot;
@@ -770,7 +783,11 @@ namespace UnrealBuildTool
 			bool bIsRocketGame = (InRulesObject != null) ? (TargetType == TargetRules.TargetType.RocketGame) : false;
 
 			// Construct the output path based on configuration, platform, game if not specified.
-            OutputPath = Path.GetFullPath(MakeBinaryPath("", AppName, UEBuildBinaryType.Executable, TargetType, bIsRocketGame, null, InAppName, Configuration == UnrealTargetConfiguration.Shipping ? Rules.ForceNameAsForDevelopment() : false, Rules.ExeBinariesSubFolder));
+            OutputPaths = MakeBinaryPaths("", AppName, UEBuildBinaryType.Executable, TargetType, bIsRocketGame, null, InAppName, Configuration == UnrealTargetConfiguration.Shipping ? Rules.ForceNameAsForDevelopment() : false, Rules.ExeBinariesSubFolder);
+			for (int Index = 0; Index < OutputPaths.Length; Index++)
+			{
+				OutputPaths[Index] = Path.GetFullPath(OutputPaths[Index]);
+			}
 
 			if (bCompileMonolithic && TargetRules.IsGameType(InRulesObject.Type))
 			{
@@ -781,7 +798,10 @@ namespace UnrealBuildTool
 					// We want the output to go into the <GAME>\Binaries folder
 					if (UnrealBuildTool.HasUProjectFile() == false)
 					{
-						OutputPath = OutputPath.Replace(Path.Combine("Engine", "Binaries"), Path.Combine(InGameName,"Binaries"));
+						for (int Index = 0; Index < OutputPaths.Length; Index++)
+						{
+							OutputPaths[Index] = OutputPaths[Index].Replace(Path.Combine("Engine", "Binaries"), Path.Combine(InGameName, "Binaries"));
+						}
 					}
 					else
 					{
@@ -794,7 +814,10 @@ namespace UnrealBuildTool
 							UProjectPath = Path.GetDirectoryName(FullPath);
 						}
 						string ProjectPath = Path.GetFullPath(Path.Combine(UProjectPath, "Binaries"));
-						OutputPath = OutputPath.Replace(EnginePath, ProjectPath);
+						for (int Index = 0; Index < OutputPaths.Length; Index++)
+						{
+							OutputPaths[Index] = OutputPaths[Index].Replace(EnginePath, ProjectPath);
+						}
 					}
 				}
 			}
@@ -1259,12 +1282,19 @@ namespace UnrealBuildTool
                 string DebugInfoExtension = BuildPlatform.GetDebugInfoExtension(Binary.Config.Type);
 
 				// Create and add the binary and associated debug info
-				Manifest.AddBinaryNames(Binary.Config.OutputFilePath, DebugInfoExtension);
-				if (Binary.Config.Type == UEBuildBinaryType.Executable && 
+				foreach (string OutputFilePath in Binary.Config.OutputFilePaths)
+				{
+					Manifest.AddBinaryNames(OutputFilePath, DebugInfoExtension);
+				}
+
+				if (Binary.Config.Type == UEBuildBinaryType.Executable &&
 					  GlobalLinkEnvironment.Config.CanProduceAdditionalConsoleApp &&
 					  UEBuildConfiguration.bBuildEditor)
 				{
-					Manifest.AddBinaryNames(UEBuildBinary.GetAdditionalConsoleAppPath(Binary.Config.OutputFilePath), DebugInfoExtension);
+					foreach (string OutputFilePath in Binary.Config.OutputFilePaths)
+					{
+						Manifest.AddBinaryNames(UEBuildBinary.GetAdditionalConsoleAppPath(OutputFilePath), DebugInfoExtension);
+					}
 				}
 
                 ToolChain.AddFilesToManifest(ref Manifest,Binary);
@@ -1358,11 +1388,11 @@ namespace UnrealBuildTool
 						if(Utils.IsFileUnderDirectory(Module.ModuleDirectory, BuildConfiguration.RelativeEnginePath) && Module.Binary == ExecutableBinary)
 						{
 							UnrealTargetConfiguration LibraryConfiguration = (Configuration == UnrealTargetConfiguration.DebugGame)? UnrealTargetConfiguration.Development : Configuration;
-							Module.RedistStaticLibraryPath = MakeBinaryPath("", "UE4Game-Redist-" + Module.Name, Platform, LibraryConfiguration, UEBuildBinaryType.StaticLibrary, TargetType, false, null, AppName);
+							Module.RedistStaticLibraryPaths = MakeBinaryPaths("", "UE4Game-Redist-" + Module.Name, Platform, LibraryConfiguration, UEBuildBinaryType.StaticLibrary, TargetType, false, null, AppName);
 							Module.bBuildingRedistStaticLibrary = UnrealBuildTool.BuildingRocket();
                             if (Module.bBuildingRedistStaticLibrary)
                             {
-                                SpecialRocketLibFilesThatAreBuildProducts.Add(Module.RedistStaticLibraryPath);
+                                SpecialRocketLibFilesThatAreBuildProducts.AddRange(Module.RedistStaticLibraryPaths);
                             }
 						}
 					}
@@ -1555,6 +1585,9 @@ namespace UnrealBuildTool
 				SpecialRocketLibFilesThatAreBuildProducts.AddRange(AddPlugin(Plugin));
 			}
 
+			// Allow the platform to setup binaries/plugins/modules
+			UEBuildPlatform.GetBuildPlatform(Platform).SetupBinaries(this);
+
 			// Describe what's being built.
 			Log.TraceVerbose("Building {0} - {1} - {2} - {3}", AppName, GameName, Platform, Configuration);
 
@@ -1716,8 +1749,11 @@ namespace UnrealBuildTool
 
 					if (!String.IsNullOrEmpty(FoundOnlyModule.OnlyModuleSuffix))
 					{
-						DLLBinary.Config.OriginalOutputFilePath = DLLBinary.Config.OutputFilePath;
-						DLLBinary.Config.OutputFilePath = AddModuleFilenameSuffix(FoundOnlyModule.OnlyModuleName, DLLBinary.Config.OutputFilePath, FoundOnlyModule.OnlyModuleSuffix);
+						DLLBinary.Config.OriginalOutputFilePaths = DLLBinary.Config.OutputFilePaths;
+						for (int Index = 0; Index < DLLBinary.Config.OutputFilePaths.Length; Index++ )
+						{
+							DLLBinary.Config.OutputFilePaths[Index] = AddModuleFilenameSuffix(FoundOnlyModule.OnlyModuleName, DLLBinary.Config.OutputFilePaths[Index], FoundOnlyModule.OnlyModuleSuffix);
+						}
 					}
 				}
 			}
@@ -1742,8 +1778,11 @@ namespace UnrealBuildTool
 				if (GameModules != null && GameModules.Count > 0)
 				{
 					var UniqueSuffix = (new Random((int)(DateTime.Now.Ticks % Int32.MaxValue)).Next(10000)).ToString();
-					DLLBinary.Config.OriginalOutputFilePath = DLLBinary.Config.OutputFilePath;
-					DLLBinary.Config.OutputFilePath = AddModuleFilenameSuffix(GameModules[0].Name, DLLBinary.Config.OutputFilePath, UniqueSuffix);
+					DLLBinary.Config.OriginalOutputFilePaths = DLLBinary.Config.OutputFilePaths;
+					for (int Index = 0; Index < DLLBinary.Config.OutputFilePaths.Length; Index++)
+					{
+						DLLBinary.Config.OutputFilePaths[Index] = AddModuleFilenameSuffix(GameModules[0].Name, DLLBinary.Config.OutputFilePaths[Index], UniqueSuffix);
+					}
 
 					FilteredBinaries.Add(DLLBinary);
 					AnyBinariesAdded = true;
@@ -2044,18 +2083,18 @@ namespace UnrealBuildTool
 			UEBuildBinaryType BinaryType = bCompileMonolithic ? UEBuildBinaryType.StaticLibrary : UEBuildBinaryType.DynamicLinkLibrary;
 
 			// Get the output path. Don't prefix the app name for Rocket
-			string OutputFilePath;
+			string[] OutputFilePaths;
 			if ((UnrealBuildTool.BuildingRocket() || UnrealBuildTool.RunningRocket()) && bCompileMonolithic)
 			{
-				OutputFilePath = MakeBinaryPath(Module.Name, Module.Name, BinaryType, TargetType, bIsRocketModule, Plugin, AppName);
+				OutputFilePaths = MakeBinaryPaths(Module.Name, Module.Name, BinaryType, TargetType, bIsRocketModule, Plugin, AppName);
                 if (UnrealBuildTool.BuildingRocket())
                 {
-                    SpecialRocketLibFilesThatAreBuildProducts.Add(OutputFilePath);
+                    SpecialRocketLibFilesThatAreBuildProducts.AddRange(OutputFilePaths);
                 }
 			}
 			else
 			{
-				OutputFilePath = MakeBinaryPath(Module.Name, GetAppName() + "-" + Module.Name, BinaryType, TargetType, bIsRocketModule, Plugin, AppName);
+				OutputFilePaths = MakeBinaryPaths(Module.Name, GetAppName() + "-" + Module.Name, BinaryType, TargetType, bIsRocketModule, Plugin, AppName);
 			}
 
 			// Try to determine if we have the rules file
@@ -2084,7 +2123,7 @@ namespace UnrealBuildTool
 
 			// Create the binary
 			UEBuildBinaryConfiguration Config = new UEBuildBinaryConfiguration( InType:                  BinaryType,
-																				InOutputFilePath:        OutputFilePath,
+																				InOutputFilePaths:       OutputFilePaths,
 																				InIntermediateDirectory: PluginIntermediateBuildPath,
 																				bInAllowExports:         true,
 																				bInAllowCompilation:     bShouldBeBuiltFromSource,
@@ -2113,15 +2152,15 @@ namespace UnrealBuildTool
 					bool bIsRocketModule = RulesCompiler.IsRocketProjectModule(ModuleName);
 
 					// Create a DLL binary for this module
-					string OutputFilePath = MakeBinaryPath(ModuleName, GetAppName() + "-" + ModuleName, UEBuildBinaryType.DynamicLinkLibrary, TargetType, bIsRocketModule, null, AppName);
+					string[] OutputFilePaths = MakeBinaryPaths(ModuleName, GetAppName() + "-" + ModuleName, UEBuildBinaryType.DynamicLinkLibrary, TargetType, bIsRocketModule, null, AppName);
 					UEBuildBinaryConfiguration Config = new UEBuildBinaryConfiguration( InType: UEBuildBinaryType.DynamicLinkLibrary,
-																						InOutputFilePath: OutputFilePath,
-																						InIntermediateDirectory: RulesCompiler.IsGameModule(ModuleName)? ProjectIntermediateDirectory : EngineIntermediateDirectory,
+																						InOutputFilePaths: OutputFilePaths,
+																						InIntermediateDirectory: RulesCompiler.IsGameModule(ModuleName) ? ProjectIntermediateDirectory : EngineIntermediateDirectory,
 																						bInAllowExports: true,
 																						InModuleNames: new List<string> { ModuleName } );
 
 					// Tell the target about this new binary
-					AppBinaries.Add( new UEBuildBinaryCPP( this, Config ) );
+					AppBinaries.Add(new UEBuildBinaryCPP(this, Config));
 				}
 			}
 		}
@@ -2152,7 +2191,7 @@ namespace UnrealBuildTool
 		}
 
 		/** Given a UBT-built binary name (e.g. "Core"), returns a relative path to the binary for the current build configuration (e.g. "../Binaries/Win64/Core-Win64-Debug.lib") */
-		public static string MakeBinaryPath(string ModuleName, string BinaryName, UnrealTargetPlatform Platform, 
+		public static string[] MakeBinaryPaths(string ModuleName, string BinaryName, UnrealTargetPlatform Platform, 
 			UnrealTargetConfiguration Configuration, UEBuildBinaryType BinaryType, TargetRules.TargetType? TargetType, bool bIsRocketModule, PluginInfo PluginInfo, string AppName, bool bForceNameAsForDevelopment = false, string ExeBinariesSubFolder = null)
 		{
 			// Determine the binary extension for the platform and binary type.
@@ -2228,19 +2267,20 @@ namespace UnrealBuildTool
 				OutBinaryPath = Path.Combine(BaseDirectory, String.Format("{5}{0}-{1}-{2}{3}{4}", 
 					BinaryName, Platform.ToString(), LocalConfig.ToString(), BinarySuffix, BinaryExtension, Prefix));
 			}
-			return OutBinaryPath;
+
+			return BuildPlatform.FinalizeBinaryPaths(OutBinaryPath);
 		}
 
 		/** Given a UBT-built binary name (e.g. "Core"), returns a relative path to the binary for the current build configuration (e.g. "../../Binaries/Win64/Core-Win64-Debug.lib") */
-		public string MakeBinaryPath(string ModuleName, string BinaryName, UEBuildBinaryType BinaryType, TargetRules.TargetType? TargetType, bool bIsRocketModule, PluginInfo PluginInfo, string AppName, bool bForceNameAsForDevelopment = false, string ExeBinariesSubFolder = null)
+		public string[] MakeBinaryPaths(string ModuleName, string BinaryName, UEBuildBinaryType BinaryType, TargetRules.TargetType? TargetType, bool bIsRocketModule, PluginInfo PluginInfo, string AppName, bool bForceNameAsForDevelopment = false, string ExeBinariesSubFolder = null)
 		{
 			if (String.IsNullOrEmpty(ModuleName) && Configuration == UnrealTargetConfiguration.DebugGame && !bCompileMonolithic)
 			{
-				return MakeBinaryPath(ModuleName, BinaryName, Platform, UnrealTargetConfiguration.Development, BinaryType, TargetType, bIsRocketModule, PluginInfo, AppName, bForceNameAsForDevelopment);
+				return MakeBinaryPaths(ModuleName, BinaryName, Platform, UnrealTargetConfiguration.Development, BinaryType, TargetType, bIsRocketModule, PluginInfo, AppName, bForceNameAsForDevelopment);
 			}
 			else
 			{
-				return MakeBinaryPath(ModuleName, BinaryName, Platform, Configuration, BinaryType, TargetType, bIsRocketModule, PluginInfo, AppName, bForceNameAsForDevelopment, ExeBinariesSubFolder);
+				return MakeBinaryPaths(ModuleName, BinaryName, Platform, Configuration, BinaryType, TargetType, bIsRocketModule, PluginInfo, AppName, bForceNameAsForDevelopment, ExeBinariesSubFolder);
 			}
 		}
 
@@ -2339,7 +2379,7 @@ namespace UnrealBuildTool
 
 				foreach (UEBuildBinaryConfiguration BinaryConfig in RulesBuildBinaryConfigurations)
 				{
-					BinaryConfig.OutputFilePath = OutputPath;
+					BinaryConfig.OutputFilePaths = OutputPaths;
 					BinaryConfig.IntermediateDirectory = ProjectIntermediateDirectory;
 
 					if (BinaryConfig.ModuleNames.Count > 0)
@@ -2623,7 +2663,7 @@ namespace UnrealBuildTool
 		{
 			UEBuildPlatform.GetBuildPlatform(Platform).SetUpConfigurationEnvironment(this);
 
-			var UpperExtension = Path.GetExtension (OutputPath).ToUpperInvariant (); 
+			var UpperExtension = Path.GetExtension (OutputPaths[0]).ToUpperInvariant (); 
 			// Check to see if we're compiling a library or not
 			bool bIsBuildingDLL = UpperExtension == ".DLL";
 			GlobalCompileEnvironment.Config.bIsBuildingDLL = bIsBuildingDLL;
diff --git a/Engine/Source/Programs/UnrealBuildTool/IOS/IOSToolChain.cs b/Engine/Source/Programs/UnrealBuildTool/IOS/IOSToolChain.cs
--- a/Engine/Source/Programs/UnrealBuildTool/IOS/IOSToolChain.cs
+++ b/Engine/Source/Programs/UnrealBuildTool/IOS/IOSToolChain.cs
@@ -1109,7 +1109,10 @@ namespace UnrealBuildTool
 						{
 							if (Module.bBuildingRedistStaticLibrary)
 							{
-								BuiltBinaries.Add(Path.GetFullPath(Module.RedistStaticLibraryPath));
+								foreach (var LibraryPath in Module.RedistStaticLibraryPaths)
+								{
+									BuiltBinaries.Add(Path.GetFullPath(LibraryPath));
+								}
 							}
 						}
 					}
diff --git a/Engine/Source/Programs/UnrealBuildTool/IOS/UEBuildIOS.cs b/Engine/Source/Programs/UnrealBuildTool/IOS/UEBuildIOS.cs
--- a/Engine/Source/Programs/UnrealBuildTool/IOS/UEBuildIOS.cs
+++ b/Engine/Source/Programs/UnrealBuildTool/IOS/UEBuildIOS.cs
@@ -273,63 +273,6 @@ namespace UnrealBuildTool
                 }
             }
         }
-
-        public override void SetupBinaries(UEBuildTarget InBuildTarget)
-        {
-            if (BuildHostPlatform.Current.Platform != UnrealTargetPlatform.Mac)
-            {
-                // dangerously fast mode doesn't generate stub files
-                if (!IOSToolChain.bUseDangerouslyFastMode)
-                {
-                    List<UEBuildBinary> NewBinaries = new List<UEBuildBinary> ();
-                    // add the .stub to the binaries
-                    foreach (var Binary in InBuildTarget.AppBinaries)
-                    {
-                        // make a binary that just points to the .stub of this executable
-                        UEBuildBinaryConfiguration NewConfig = new UEBuildBinaryConfiguration(
-                                                                  InType: Binary.Config.Type,
-                                                                  InOutputFilePath: Binary.Config.OutputFilePath + ".stub",
-                                                                  InIntermediateDirectory: Binary.Config.IntermediateDirectory,
-                                                                  bInCreateImportLibrarySeparately: Binary.Config.bCreateImportLibrarySeparately,
-                                                                  bInAllowExports: Binary.Config.bAllowExports,
-                                                                  InModuleNames: Binary.Config.ModuleNames);
-
-                        NewBinaries.Add (new UEStubDummyBinary (InBuildTarget, NewConfig));
-                    }
-
-                    InBuildTarget.AppBinaries.AddRange (NewBinaries);
-                }
-            }
-        }
-    }
-
-    /// <summary>
-    /// A .stub that has the executable and metadata included
-    /// Note that this doesn't actually build anything. It could potentially be used to perform the IPhonePackage stuff that happens in IOSToolChain.PostBuildSync()
-    /// </summary>
-    class UEStubDummyBinary : UEBuildBinary
-    {
-        /// <summary>
-        /// Create an instance initialized to the given configuration
-        /// </summary>
-        /// <param name="InConfig">The build binary configuration to initialize the instance to</param>
-        public UEStubDummyBinary(UEBuildTarget InTarget, UEBuildBinaryConfiguration InConfig)
-            : base(InTarget, InConfig)
-        {
-        }
-
-        /// <summary>
-        /// Builds the binary.
-        /// </summary>
-        /// <param name="CompileEnvironment">The environment to compile the binary in</param>
-        /// <param name="LinkEnvironment">The environment to link the binary in</param>
-        /// <returns></returns>
-        public override IEnumerable<FileItem> Build(IUEToolChain ToolChain, CPPEnvironment CompileEnvironment, LinkEnvironment LinkEnvironment)
-        {
-            // generate the .stub!
-            return new FileItem[] { FileItem.GetItemByPath(Config.OutputFilePath) };
-        }
     }
-
 }
 
diff --git a/Engine/Source/Programs/UnrealBuildTool/Mac/MacToolChain.cs b/Engine/Source/Programs/UnrealBuildTool/Mac/MacToolChain.cs
--- a/Engine/Source/Programs/UnrealBuildTool/Mac/MacToolChain.cs
+++ b/Engine/Source/Programs/UnrealBuildTool/Mac/MacToolChain.cs
@@ -1073,12 +1073,12 @@ namespace UnrealBuildTool
 						// kind of protection to be deleted if they aren't supported for a current user that has the dylibs locally.
 						// string SubDir = DylibDir.Replace(ExeDir, "");
 						string SubDir = "";
-						Binary.Config.OutputFilePath = BundleContentsPath + "MacOS" + SubDir + "/" + BinaryFileName;
+						Binary.Config.OutputFilePaths[0] = BundleContentsPath + "MacOS" + SubDir + "/" + BinaryFileName;
 					}
 				}
 				else if (!BinaryFileName.EndsWith(".a"))
 				{
-					Binary.Config.OutputFilePath += ".app/Contents/MacOS/" + BinaryFileName;
+					Binary.Config.OutputFilePaths[0] += ".app/Contents/MacOS/" + BinaryFileName;
 				}
 			}
 		}
diff --git a/Engine/Source/Programs/UnrealBuildTool/System/LinkEnvironment.cs b/Engine/Source/Programs/UnrealBuildTool/System/LinkEnvironment.cs
--- a/Engine/Source/Programs/UnrealBuildTool/System/LinkEnvironment.cs
+++ b/Engine/Source/Programs/UnrealBuildTool/System/LinkEnvironment.cs
@@ -21,7 +21,20 @@ namespace UnrealBuildTool
 		public string LocalShadowDirectory = null;
 
 		/** The file path for the executable file that is output by the linker. */
-		public string OutputFilePath;
+		public string[] OutputFilePaths = null;
+
+		/** Returns the OutputFilePath is there is only one entry in OutputFilePaths */
+		public string OutputFilePath
+		{
+			get
+			{
+				if (OutputFilePaths.Length != 1)
+				{
+					throw new BuildException("Attempted to use LinkEnvironmentConfiguration.OutputFilePath property, but there are multiple (or no) OutputFilePaths. You need to handle multiple in the code that called this (size = {0})", OutputFilePaths.Length);
+				}
+				return OutputFilePaths[0];
+			}
+		}
 
 		/** A list of the paths used to find libraries. */
 		public List<string> LibraryPaths = new List<string>();
@@ -104,7 +117,7 @@ namespace UnrealBuildTool
 			OutputDirectory = InCopyEnvironment.OutputDirectory;
 			IntermediateDirectory = InCopyEnvironment.IntermediateDirectory;
 			LocalShadowDirectory = InCopyEnvironment.LocalShadowDirectory;
-			OutputFilePath = InCopyEnvironment.OutputFilePath;
+			OutputFilePaths = InCopyEnvironment.OutputFilePaths;
 			LibraryPaths.AddRange(InCopyEnvironment.LibraryPaths);
 			ExcludedLibraries.AddRange(InCopyEnvironment.ExcludedLibraries);
 			AdditionalLibraries.AddRange(InCopyEnvironment.AdditionalLibraries);
diff --git a/Engine/Source/Programs/UnrealBuildTool/System/UnrealBuildTool.cs b/Engine/Source/Programs/UnrealBuildTool/System/UnrealBuildTool.cs
--- a/Engine/Source/Programs/UnrealBuildTool/System/UnrealBuildTool.cs
+++ b/Engine/Source/Programs/UnrealBuildTool/System/UnrealBuildTool.cs
@@ -1793,7 +1793,12 @@ namespace UnrealBuildTool
 			bool bIsRunning = false;
 			if (!ProjectFileGenerator.bGenerateProjectFiles && !UEBuildConfiguration.bGenerateManifest && Target.TargetType == TargetRules.TargetType.Editor)
 			{
-				var EditorProcessFilename = UEBuildTarget.MakeBinaryPath("Launcher", "UE4Editor", Target.Platform, Target.Configuration, UEBuildBinaryType.Executable, null, false, null, "UE4Editor");
+				var EditorProcessFilenames = UEBuildTarget.MakeBinaryPaths("Launcher", "UE4Editor", Target.Platform, Target.Configuration, UEBuildBinaryType.Executable, null, false, null, "UE4Editor");
+				if (EditorProcessFilenames.Length != 1)
+				{
+					throw new BuildException("ShouldDoHotReload cannot handle multiple binaries returning from UEBuildTarget.MakeBinaryPaths");
+				}
+				var EditorProcessFilename = EditorProcessFilenames[0];
 				var EditorProcessName = Path.GetFileNameWithoutExtension(EditorProcessFilename);
 				var EditorProcesses = Process.GetProcessesByName(EditorProcessName);
 				var BinariesPath = Path.GetFullPath(Path.GetDirectoryName(EditorProcessFilename));
