- Fixed Android PCHs with "fat" binaries

diff --git a/Engine/Source/Programs/UnrealBuildTool/Android/AndroidToolChain.cs b/Engine/Source/Programs/UnrealBuildTool/Android/AndroidToolChain.cs
--- a/Engine/Source/Programs/UnrealBuildTool/Android/AndroidToolChain.cs
+++ b/Engine/Source/Programs/UnrealBuildTool/Android/AndroidToolChain.cs
@@ -526,20 +526,12 @@ namespace UnrealBuildTool
 			}
 
 			string BaseArguments = "";
-			string PCHArguments = "";
 
 			if (CompileEnvironment.Config.PrecompiledHeaderAction != PrecompiledHeaderAction.Create)
 			{
 				BaseArguments += " -Werror";
 
 			}
-			if (CompileEnvironment.Config.PrecompiledHeaderAction == PrecompiledHeaderAction.Include)
-			{
-				var PCHExtension = UEBuildPlatform.BuildPlatformDictionary[UnrealTargetPlatform.Android].GetBinaryExtension(UEBuildBinaryType.PrecompiledHeader);
-				// Add the precompiled header file's path to the include path so Clang can find it.
-				// This needs to be before the other include paths to ensure Clang uses it instead of the source header file.
-				PCHArguments += string.Format(" -include \"{0}\"", CompileEnvironment.PrecompiledHeaderFile.AbsolutePath.Replace(PCHExtension, ""));
-			}
 
 			// Directly added NDK files for NDK extensions
 			if (!UnrealBuildTool.RunningRocket())
@@ -555,12 +547,30 @@ namespace UnrealBuildTool
 
 			var BuildPlatform = UEBuildPlatform.GetBuildPlatformForCPPTargetPlatform(CompileEnvironment.Config.Target.Platform);
 
+
+
+			string BasePCHName = "";
+			var PCHExtension = UEBuildPlatform.BuildPlatformDictionary[UnrealTargetPlatform.Android].GetBinaryExtension(UEBuildBinaryType.PrecompiledHeader);
+			if (CompileEnvironment.Config.PrecompiledHeaderAction == PrecompiledHeaderAction.Include)
+			{
+				BasePCHName = RemoveArchName(CompileEnvironment.PrecompiledHeaderFile.AbsolutePath).Replace(PCHExtension, "");
+			}
+
 			// Create a compile action for each source file.
 			CPPOutput Result = new CPPOutput();
 			foreach (string Arch in Arches)
 			{
 				// which toolchain to use
 				string Arguments = GetCLArguments_Global(CompileEnvironment, Arch) + BaseArguments;
+				
+				// which PCH file to include
+				string PCHArguments = "";
+				if (CompileEnvironment.Config.PrecompiledHeaderAction == PrecompiledHeaderAction.Include)
+				{
+					// Add the precompiled header file's path to the include path so Clang can find it.
+					// This needs to be before the other include paths to ensure Clang uses it instead of the source header file.
+					PCHArguments += string.Format(" -include \"{0}\"", InlineArchName(BasePCHName, Arch));
+				}
 
 				// Add include paths to the argument list (filtered by architecture)
 				foreach (string IncludePath in CompileEnvironment.Config.CPPIncludeInfo.SystemIncludePaths)
@@ -616,13 +626,11 @@ namespace UnrealBuildTool
 
 					if (CompileEnvironment.Config.PrecompiledHeaderAction == PrecompiledHeaderAction.Create)
 					{
-						var PCHExtension = UEBuildPlatform.BuildPlatformDictionary[UnrealTargetPlatform.Android].GetBinaryExtension(UEBuildBinaryType.PrecompiledHeader);
-
 						// Add the precompiled header file to the produced item list.
 						FileItem PrecompiledHeaderFile = FileItem.GetItemByPath(
 							Path.Combine(
 								CompileEnvironment.Config.OutputDirectory,
-								Path.GetFileName(SourceFile.AbsolutePath) + Arch + PCHExtension
+								Path.GetFileName(InlineArchName(SourceFile.AbsolutePath, Arch) + PCHExtension)
 								)
 							);
 
@@ -637,7 +645,8 @@ namespace UnrealBuildTool
 						if (CompileEnvironment.Config.PrecompiledHeaderAction == PrecompiledHeaderAction.Include)
 						{
 							CompileAction.bIsUsingPCH = true;
-							CompileAction.PrerequisiteItems.Add(CompileEnvironment.PrecompiledHeaderFile);
+							FileItem ArchPrecompiledHeaderFile = FileItem.GetItemByPath(InlineArchName(BasePCHName, Arch) + PCHExtension);
+							CompileAction.PrerequisiteItems.Add(ArchPrecompiledHeaderFile);
 						}
 
 						var ObjectFileExtension = UEBuildPlatform.BuildPlatformDictionary[UnrealTargetPlatform.Android].GetBinaryExtension(UEBuildBinaryType.Object);
diff --git a/Engine/Source/Programs/UnrealBuildTool/Android/UEBuildAndroid.cs b/Engine/Source/Programs/UnrealBuildTool/Android/UEBuildAndroid.cs
--- a/Engine/Source/Programs/UnrealBuildTool/Android/UEBuildAndroid.cs
+++ b/Engine/Source/Programs/UnrealBuildTool/Android/UEBuildAndroid.cs
@@ -152,9 +152,7 @@ namespace UnrealBuildTool
 
 		public override bool ShouldUsePCHFiles(CPPTargetPlatform Platform, CPPTargetConfiguration Configuration)
 		{
-			// if we compile multiple architectures at once, PCH doesn't work right, so only use PCH when building for one platform
-			// @todo android: Handle PCH files in the multiple architecture case
-			return false;// AndroidToolChain.GetAllArchitectures().Length == 1;
+			return true;
 		}
 
 		public override string GetDebugInfoExtension(UEBuildBinaryType InBinaryType)
