Use Clone() instead of assignment for arrays added to support multiple output items in UBT so they can be modified independently.

diff --git a/Engine/Source/Programs/UnrealBuildTool/Configuration/UEBuildBinary.cs b/Engine/Source/Programs/UnrealBuildTool/Configuration/UEBuildBinary.cs
--- a/Engine/Source/Programs/UnrealBuildTool/Configuration/UEBuildBinary.cs
+++ b/Engine/Source/Programs/UnrealBuildTool/Configuration/UEBuildBinary.cs
@@ -162,7 +162,7 @@ namespace UnrealBuildTool
 			)
 		{
 			Type = InType;
-			OutputFilePaths = InOutputFilePaths;
+			OutputFilePaths = InOutputFilePaths != null ? (string[])InOutputFilePaths.Clone() : null;
 			IntermediateDirectory = InIntermediateDirectory;
 			bAllowExports = bInAllowExports;
 			bCreateImportLibrarySeparately = bInCreateImportLibrarySeparately;
@@ -569,7 +569,7 @@ namespace UnrealBuildTool
 			}
 
 			// Set the link output file.
-			BinaryLinkEnvironment.Config.OutputFilePaths = Config.OutputFilePaths;
+			BinaryLinkEnvironment.Config.OutputFilePaths = Config.OutputFilePaths != null ? (string[])Config.OutputFilePaths.Clone() : null;
 
 			// Set whether the link is allowed to have exports.
 			BinaryLinkEnvironment.Config.bHasExports = Config.bAllowExports;
diff --git a/Engine/Source/Programs/UnrealBuildTool/Configuration/UEBuildModule.cs b/Engine/Source/Programs/UnrealBuildTool/Configuration/UEBuildModule.cs
--- a/Engine/Source/Programs/UnrealBuildTool/Configuration/UEBuildModule.cs
+++ b/Engine/Source/Programs/UnrealBuildTool/Configuration/UEBuildModule.cs
@@ -1492,7 +1492,7 @@ namespace UnrealBuildTool
 				RedistLinkEnvironment.Config.bIsBuildingDLL        = false;
 				RedistLinkEnvironment.Config.bIsBuildingLibrary    = true;
 				RedistLinkEnvironment.Config.IntermediateDirectory = Binary.Config.IntermediateDirectory;
-				RedistLinkEnvironment.Config.OutputFilePaths        = RedistStaticLibraryPaths;
+				RedistLinkEnvironment.Config.OutputFilePaths       = RedistStaticLibraryPaths != null ? (string[])RedistStaticLibraryPaths.Clone() : null;
 
 				// Replace the items built so far with the library
 				RedistLinkEnvironment.LinkExecutable(false);
diff --git a/Engine/Source/Programs/UnrealBuildTool/Configuration/UEBuildTarget.cs b/Engine/Source/Programs/UnrealBuildTool/Configuration/UEBuildTarget.cs
--- a/Engine/Source/Programs/UnrealBuildTool/Configuration/UEBuildTarget.cs
+++ b/Engine/Source/Programs/UnrealBuildTool/Configuration/UEBuildTarget.cs
@@ -1749,7 +1749,7 @@ namespace UnrealBuildTool
 
 					if (!String.IsNullOrEmpty(FoundOnlyModule.OnlyModuleSuffix))
 					{
-						DLLBinary.Config.OriginalOutputFilePaths = DLLBinary.Config.OutputFilePaths;
+						DLLBinary.Config.OriginalOutputFilePaths = DLLBinary.Config.OutputFilePaths != null ? (string[])DLLBinary.Config.OutputFilePaths.Clone() : null;
 						for (int Index = 0; Index < DLLBinary.Config.OutputFilePaths.Length; Index++ )
 						{
 							DLLBinary.Config.OutputFilePaths[Index] = AddModuleFilenameSuffix(FoundOnlyModule.OnlyModuleName, DLLBinary.Config.OutputFilePaths[Index], FoundOnlyModule.OnlyModuleSuffix);
@@ -1778,7 +1778,7 @@ namespace UnrealBuildTool
 				if (GameModules != null && GameModules.Count > 0)
 				{
 					var UniqueSuffix = (new Random((int)(DateTime.Now.Ticks % Int32.MaxValue)).Next(10000)).ToString();
-					DLLBinary.Config.OriginalOutputFilePaths = DLLBinary.Config.OutputFilePaths;
+					DLLBinary.Config.OriginalOutputFilePaths = DLLBinary.Config.OutputFilePaths != null ? (string[])DLLBinary.Config.OutputFilePaths.Clone() : null;
 					for (int Index = 0; Index < DLLBinary.Config.OutputFilePaths.Length; Index++)
 					{
 						DLLBinary.Config.OutputFilePaths[Index] = AddModuleFilenameSuffix(GameModules[0].Name, DLLBinary.Config.OutputFilePaths[Index], UniqueSuffix);
@@ -2379,7 +2379,7 @@ namespace UnrealBuildTool
 
 				foreach (UEBuildBinaryConfiguration BinaryConfig in RulesBuildBinaryConfigurations)
 				{
-					BinaryConfig.OutputFilePaths = OutputPaths;
+					BinaryConfig.OutputFilePaths = OutputPaths != null ? (string[])OutputPaths.Clone() : null;
 					BinaryConfig.IntermediateDirectory = ProjectIntermediateDirectory;
 
 					if (BinaryConfig.ModuleNames.Count > 0)
diff --git a/Engine/Source/Programs/UnrealBuildTool/System/LinkEnvironment.cs b/Engine/Source/Programs/UnrealBuildTool/System/LinkEnvironment.cs
--- a/Engine/Source/Programs/UnrealBuildTool/System/LinkEnvironment.cs
+++ b/Engine/Source/Programs/UnrealBuildTool/System/LinkEnvironment.cs
@@ -117,7 +117,7 @@ namespace UnrealBuildTool
 			OutputDirectory = InCopyEnvironment.OutputDirectory;
 			IntermediateDirectory = InCopyEnvironment.IntermediateDirectory;
 			LocalShadowDirectory = InCopyEnvironment.LocalShadowDirectory;
-			OutputFilePaths = InCopyEnvironment.OutputFilePaths;
+			OutputFilePaths = InCopyEnvironment.OutputFilePaths != null ? (string[])InCopyEnvironment.OutputFilePaths.Clone() : null;
 			LibraryPaths.AddRange(InCopyEnvironment.LibraryPaths);
 			ExcludedLibraries.AddRange(InCopyEnvironment.ExcludedLibraries);
 			AdditionalLibraries.AddRange(InCopyEnvironment.AdditionalLibraries);
