UE4: Moved some Mac/IOS platform specific functions to the UEToolChain interface so branches without Mac code can still run UBT.

diff --git a/Engine/Source/Programs/UnrealBuildTool/Configuration/UEBuildTarget.cs b/Engine/Source/Programs/UnrealBuildTool/Configuration/UEBuildTarget.cs
--- a/Engine/Source/Programs/UnrealBuildTool/Configuration/UEBuildTarget.cs
+++ b/Engine/Source/Programs/UnrealBuildTool/Configuration/UEBuildTarget.cs
@@ -1324,7 +1324,7 @@ namespace UnrealBuildTool
 			OutputItems = new List<FileItem>();
 			UObjectModules = new List<UHTModuleInfo>();
 
-			var SpecialRocketLibFilesThatAreBuildProducts = PreBuildSetup();
+			var SpecialRocketLibFilesThatAreBuildProducts = PreBuildSetup(TargetToolChain);
 
 			EULAViolationWarning = !ProjectFileGenerator.bGenerateProjectFiles
 				? CheckForEULAViolation()
@@ -1402,7 +1402,7 @@ namespace UnrealBuildTool
 			// On Mac we have actions that should be executed after all the binaries are created
 			if (GlobalLinkEnvironment.Config.Target.Platform == CPPTargetPlatform.Mac)
 			{
-				MacToolChain.SetupBundleDependencies(AppBinaries, GameName);
+				TargetToolChain.SetupBundleDependencies(AppBinaries, GameName);
 			}
 
 			// Generate the external file list 
@@ -1563,7 +1563,7 @@ namespace UnrealBuildTool
 		/// Setup target before build. This method finds dependencies, sets up global environment etc.
 		/// </summary>
 		/// <returns>Special Rocket lib files that are build products.</returns>
-		public List<string> PreBuildSetup()
+		public List<string> PreBuildSetup(IUEToolChain TargetToolChain)
 		{
 			// Set up the global compile and link environment in GlobalCompileEnvironment and GlobalLinkEnvironment.
 			SetupGlobalEnvironment();
@@ -1625,7 +1625,7 @@ namespace UnrealBuildTool
 			// On Mac AppBinaries paths for non-console targets need to be adjusted to be inside the app bundle
 			if (GlobalLinkEnvironment.Config.Target.Platform == CPPTargetPlatform.Mac && !GlobalLinkEnvironment.Config.bIsBuildingConsoleApplication)
 			{
-				MacToolChain.FixBundleBinariesPaths(this, AppBinaries);
+				TargetToolChain.FixBundleBinariesPaths(this, AppBinaries);
 			}
 
 			// If we're building a single module, then find the binary for that module and add it to our target
diff --git a/Engine/Source/Programs/UnrealBuildTool/IOS/IOSToolChain.cs b/Engine/Source/Programs/UnrealBuildTool/IOS/IOSToolChain.cs
--- a/Engine/Source/Programs/UnrealBuildTool/IOS/IOSToolChain.cs
+++ b/Engine/Source/Programs/UnrealBuildTool/IOS/IOSToolChain.cs
@@ -1313,6 +1313,11 @@ namespace UnrealBuildTool
 			}
 		}
 
+		public override string GetPlatformVersion()
+		{
+			return IOSVersion;
+		}
+
 		public static int RunExecutableAndWait( string ExeName, string ArgumentList, out string StdOutResults )
 		{
 			// Create the process
diff --git a/Engine/Source/Programs/UnrealBuildTool/Linux/LinuxToolChain.cs b/Engine/Source/Programs/UnrealBuildTool/Linux/LinuxToolChain.cs
--- a/Engine/Source/Programs/UnrealBuildTool/Linux/LinuxToolChain.cs
+++ b/Engine/Source/Programs/UnrealBuildTool/Linux/LinuxToolChain.cs
@@ -865,7 +865,7 @@ namespace UnrealBuildTool
             throw new BuildException("Linux cannot compile C# files");
         }
 
-        static public void SetupBundleDependencies(List<UEBuildBinary> Binaries, string GameName)
+        public override void SetupBundleDependencies(List<UEBuildBinary> Binaries, string GameName)
         {
             foreach (UEBuildBinary Binary in Binaries)
             {
diff --git a/Engine/Source/Programs/UnrealBuildTool/Mac/MacToolChain.cs b/Engine/Source/Programs/UnrealBuildTool/Mac/MacToolChain.cs
--- a/Engine/Source/Programs/UnrealBuildTool/Mac/MacToolChain.cs
+++ b/Engine/Source/Programs/UnrealBuildTool/Mac/MacToolChain.cs
@@ -1044,16 +1044,20 @@ namespace UnrealBuildTool
 			return TargetItem;
 		}
 
-		static public void SetupBundleDependencies(List<UEBuildBinary> Binaries, string GameName)
+		public override void SetupBundleDependencies(List<UEBuildBinary> Binaries, string GameName)
 		{
+			base.SetupBundleDependencies(Binaries, GameName);
+
 			foreach (UEBuildBinary Binary in Binaries)
 			{
 				BundleDependencies.Add(FileItem.GetItemByPath(Binary.ToString()));
 			}
 		}
 		
-		static public void FixBundleBinariesPaths(UEBuildTarget Target, List<UEBuildBinary> Binaries)
+		public override void FixBundleBinariesPaths(UEBuildTarget Target, List<UEBuildBinary> Binaries)
 		{
+			base.FixBundleBinariesPaths(Target, Binaries);
+
 			string BundleContentsPath = Target.OutputPath + ".app/Contents/";
 			foreach (UEBuildBinary Binary in Binaries)
 			{
diff --git a/Engine/Source/Programs/UnrealBuildTool/System/XcodeProjectFileGenerator.cs b/Engine/Source/Programs/UnrealBuildTool/System/XcodeProjectFileGenerator.cs
--- a/Engine/Source/Programs/UnrealBuildTool/System/XcodeProjectFileGenerator.cs
+++ b/Engine/Source/Programs/UnrealBuildTool/System/XcodeProjectFileGenerator.cs
@@ -476,6 +476,8 @@ namespace UnrealBuildTool
 			}
 			else
 			{
+				IUEToolChain Toolchain = UEToolChain.GetPlatformToolChain(CPPTargetPlatform.IOS);
+
 				Contents.Append(
 					"\t\t" + ConfigGuid + " /* " + ConfigName + " */ = {" + ProjectFileGenerator.NewLine +
 					"\t\t\tisa = XCBuildConfiguration;" + ProjectFileGenerator.NewLine +
@@ -487,7 +489,7 @@ namespace UnrealBuildTool
 					"\t\t\t\t\"PRODUCT_NAME[sdk=iphoneos*]\" = \"$(TARGET_NAME)\";" + ProjectFileGenerator.NewLine +
 					"\t\t\t\t\"PRODUCT_NAME[sdk=iphonesimulator*]\" = \"$(TARGET_NAME)-simulator\";" + ProjectFileGenerator.NewLine +
 					"\t\t\t\tMACOSX_DEPLOYMENT_TARGET = 10.9;" + ProjectFileGenerator.NewLine +
-					"\t\t\t\tIPHONEOS_DEPLOYMENT_TARGET = " + IOSToolChain.IOSVersion + ";" + ProjectFileGenerator.NewLine +
+					"\t\t\t\tIPHONEOS_DEPLOYMENT_TARGET = " + Toolchain.GetPlatformVersion() + ";" + ProjectFileGenerator.NewLine +
 					"\t\t\t\tTARGETED_DEVICE_FAMILY = \"1,2\";" + ProjectFileGenerator.NewLine +
 					"\t\t\t\tCONFIGURATION_BUILD_DIR = \"Engine/Binaries/IOS/Payload\";" + ProjectFileGenerator.NewLine +
 					"\t\t\t\tSDKROOT = macosx;" + ProjectFileGenerator.NewLine +
@@ -503,6 +505,8 @@ namespace UnrealBuildTool
 
 		private void AppendIOSBuildConfig(ref StringBuilder Contents, string ConfigName, string ConfigGuid)
 		{
+			IUEToolChain Toolchain = UEToolChain.GetPlatformToolChain(CPPTargetPlatform.IOS);
+
 			Contents.Append(
 				"\t\t" + ConfigGuid + " /* " + ConfigName + " */ = {" + ProjectFileGenerator.NewLine +
 				"\t\t\tisa = XCBuildConfiguration;" + ProjectFileGenerator.NewLine +
@@ -511,7 +515,7 @@ namespace UnrealBuildTool
 				"\t\t\t\t\"PRODUCT_NAME[sdk=iphonesimulator*]\" = \"$(TARGET_NAME)-simulator\";" + ProjectFileGenerator.NewLine +
 				"\t\t\t\tSUPPORTED_PLATFORMS = \"iphoneos iphonesimulator\";" + ProjectFileGenerator.NewLine +
 				"\t\t\t\tCONFIGURATION_BUILD_DIR = \"Engine/Binaries/IOS/Payload\";" + ProjectFileGenerator.NewLine +
-				"\t\t\t\tIPHONEOS_DEPLOYMENT_TARGET = " + IOSToolChain.IOSVersion + ";" + ProjectFileGenerator.NewLine +
+				"\t\t\t\tIPHONEOS_DEPLOYMENT_TARGET = " + Toolchain.GetPlatformVersion() + ";" + ProjectFileGenerator.NewLine +
 				"\t\t\t\tTARGETED_DEVICE_FAMILY = \"1,2\";" + ProjectFileGenerator.NewLine +
 				"\t\t\t\tSDKROOT = iphoneos;" + ProjectFileGenerator.NewLine +
 				"\t\t\t};" + ProjectFileGenerator.NewLine +
@@ -521,6 +525,8 @@ namespace UnrealBuildTool
 
 		private void AppendIOSRunConfig(ref StringBuilder Contents, string ConfigName, string ConfigGuid, string TargetName, string EngineRelative, string GamePath, bool bIsUE4Game, bool IsAGame, bool bIsUE4Client)
 		{
+			IUEToolChain Toolchain = UEToolChain.GetPlatformToolChain(CPPTargetPlatform.IOS);
+
 			Contents.Append(
 				"\t\t" + ConfigGuid + " /* " + ConfigName + " */ = {" + ProjectFileGenerator.NewLine +
 				"\t\t\tisa = XCBuildConfiguration;" + ProjectFileGenerator.NewLine +
@@ -529,7 +535,7 @@ namespace UnrealBuildTool
 				"\t\t\t\t\"PRODUCT_NAME[sdk=iphonesimulator*]\" = \"" + TargetName + "-simulator\";" + ProjectFileGenerator.NewLine +
 				"\t\t\t\tCONFIGURATION_BUILD_DIR = \"Engine/Binaries/IOS/Payload\";" + ProjectFileGenerator.NewLine +
 				"\t\t\t\tSUPPORTED_PLATFORMS = \"iphoneos iphonesimulator\";" + ProjectFileGenerator.NewLine +
-				"\t\t\t\tIPHONEOS_DEPLOYMENT_TARGET = " + IOSToolChain.IOSVersion + ";" + ProjectFileGenerator.NewLine +
+				"\t\t\t\tIPHONEOS_DEPLOYMENT_TARGET = " + Toolchain.GetPlatformVersion() + ";" + ProjectFileGenerator.NewLine +
 				"\t\t\t\t\"CODE_SIGN_IDENTITY[sdk=iphoneos*]\" = \"iPhone Developer\";" + ProjectFileGenerator.NewLine);
 
 			if (bIsUE4Game)
diff --git a/Engine/Source/Programs/UnrealBuildTool/ToolChain/UEToolChain.cs b/Engine/Source/Programs/UnrealBuildTool/ToolChain/UEToolChain.cs
--- a/Engine/Source/Programs/UnrealBuildTool/ToolChain/UEToolChain.cs
+++ b/Engine/Source/Programs/UnrealBuildTool/ToolChain/UEToolChain.cs
@@ -39,6 +39,12 @@ namespace UnrealBuildTool
 		void SetUpGlobalEnvironment();
 
         void AddFilesToManifest(ref FileManifest manifest, UEBuildBinary Binary );
+
+		void SetupBundleDependencies(List<UEBuildBinary> Binaries, string GameName);
+
+		void FixBundleBinariesPaths(UEBuildTarget Target, List<UEBuildBinary> Binaries);
+
+		string GetPlatformVersion();
 	}
 
 	public abstract class UEToolChain : IUEToolChain
@@ -197,5 +203,19 @@ namespace UnrealBuildTool
 			}
 		}
 
+		public virtual void SetupBundleDependencies(List<UEBuildBinary> Binaries, string GameName)
+		{
+
+		}
+
+		public virtual void FixBundleBinariesPaths(UEBuildTarget Target, List<UEBuildBinary> Binaries)
+		{
+
+		}
+
+		public virtual string GetPlatformVersion()
+		{
+			return "";
+		}
 	};
 }
\ No newline at end of file
