- Allow for UAT scripts to set up UBT to have proper engine relative path and UProject path. Only BuildCookRun does it currently
- Fixed .ini parsing of bools to allow for True, true, TRUE, etc.

diff --git a/Engine/Source/Programs/AutomationTool/Scripts/BuildCookRun.Automation.cs b/Engine/Source/Programs/AutomationTool/Scripts/BuildCookRun.Automation.cs
--- a/Engine/Source/Programs/AutomationTool/Scripts/BuildCookRun.Automation.cs
+++ b/Engine/Source/Programs/AutomationTool/Scripts/BuildCookRun.Automation.cs
@@ -36,6 +36,10 @@ public class BuildCookRun : BuildCommand
 			MakeForeignCodeSample();
 		}
 		var Params = SetupParams();
+
+		// allow BCR functions to call UBT functions (especially .ini parsing)
+		UnrealBuildTool.UnrealBuildTool.SetupUBTFromUAT(ProjectPath);
+
 		DoBuildCookRun(Params);
 	}
 
diff --git a/Engine/Source/Programs/UnrealBuildTool/Configuration/BuildConfiguration.cs b/Engine/Source/Programs/UnrealBuildTool/Configuration/BuildConfiguration.cs
--- a/Engine/Source/Programs/UnrealBuildTool/Configuration/BuildConfiguration.cs
+++ b/Engine/Source/Programs/UnrealBuildTool/Configuration/BuildConfiguration.cs
@@ -215,11 +215,16 @@ namespace UnrealBuildTool
 		/// <summary>
 		/// Relative root engine path.
 		/// </summary>
+		private static string _RelativeEnginePath = "../../Engine/";
 		public static string RelativeEnginePath
 		{
 			get
 			{
-				return "../../Engine/";
+				return _RelativeEnginePath;
+			}
+			set
+			{
+				_RelativeEnginePath = value;
 			}
 		}
 
diff --git a/Engine/Source/Programs/UnrealBuildTool/Configuration/EngineConfiguration.cs b/Engine/Source/Programs/UnrealBuildTool/Configuration/EngineConfiguration.cs
--- a/Engine/Source/Programs/UnrealBuildTool/Configuration/EngineConfiguration.cs
+++ b/Engine/Source/Programs/UnrealBuildTool/Configuration/EngineConfiguration.cs
@@ -164,11 +164,11 @@ namespace UnrealBuildTool
 			if (Result)
 			{
 				// C# Boolean type expects "False" or "True" but since we're not case sensitive, we need to suppor that manually
-				if (String.Compare(TextValue, "true", false) == 0 || String.Compare(TextValue, "1") == 0)
+				if (String.Compare(TextValue, "true", true) == 0 || String.Compare(TextValue, "1") == 0)
 				{
 					Value = true;
 				}
-				else if (String.Compare(TextValue, "false", false) == 0 || String.Compare(TextValue, "0") == 0)
+				else if (String.Compare(TextValue, "false", true) == 0 || String.Compare(TextValue, "0") == 0)
 				{
 					Value = false;
 				}
diff --git a/Engine/Source/Programs/UnrealBuildTool/System/UnrealBuildTool.cs b/Engine/Source/Programs/UnrealBuildTool/System/UnrealBuildTool.cs
--- a/Engine/Source/Programs/UnrealBuildTool/System/UnrealBuildTool.cs
+++ b/Engine/Source/Programs/UnrealBuildTool/System/UnrealBuildTool.cs
@@ -2292,5 +2292,16 @@ namespace UnrealBuildTool
 
             return TargetCollectionName.ToString();
         }
+
+		/// <summary>
+		/// Sets up UBT when running from UAT
+		/// </summary>
+		/// <param name="UProjectDir"></param>
+		public static void SetupUBTFromUAT(string UProjectFile)
+		{
+			// when running UAT, the working directory is the root UE4 dir
+			BuildConfiguration.RelativeEnginePath = "Engine";
+			SetProjectFile(UProjectFile);
+		}
     }
 }
