Fix for UBT failure when recompiling level editor in the editor.

diff --git a/Engine/Source/Programs/UnrealBuildTool/Configuration/UEBuildTarget.cs b/Engine/Source/Programs/UnrealBuildTool/Configuration/UEBuildTarget.cs
--- a/Engine/Source/Programs/UnrealBuildTool/Configuration/UEBuildTarget.cs
+++ b/Engine/Source/Programs/UnrealBuildTool/Configuration/UEBuildTarget.cs
@@ -2879,7 +2879,7 @@ namespace UnrealBuildTool
 							InDynamicallyLoadedModuleNames: RulesObject.DynamicallyLoadedModuleNames,
 							InPlatformSpecificDynamicallyLoadedModuleNames: RulesObject.PlatformSpecificDynamicallyLoadedModuleNames,
 							InOptimizeCode: RulesObject.OptimizeCode,
-							InAllowSharedPCH: (RulesObject.PCHUsage == ModuleRules.PCHUsageMode.NoSharedPCHs) ? false : true,
+							InAllowSharedPCH: (RulesObject.PCHUsage == ModuleRules.PCHUsageMode.NoSharedPCHs) ? false : !this.bEditorRecompile,
 							InSharedPCHHeaderFile: RulesObject.SharedPCHHeaderFile,
 							InUseRTTI: RulesObject.bUseRTTI,
 							InEnableBufferSecurityChecks: RulesObject.bEnableBufferSecurityChecks,
