Switching Linux builds to clang 3.5.0-based toolchain.

- Temporarily suppressed a warning until we resolve if (!this) checks.
- Upped optimization from -O2 to -O3 (as was the case a year ago).

Note: 3.5.0 generates much more compact debug info, reducing development server binary size from 1GB to roughly 600 MB.

diff --git a/Engine/Source/Programs/UnrealBuildTool/Linux/LinuxToolChain.cs b/Engine/Source/Programs/UnrealBuildTool/Linux/LinuxToolChain.cs
--- a/Engine/Source/Programs/UnrealBuildTool/Linux/LinuxToolChain.cs
+++ b/Engine/Source/Programs/UnrealBuildTool/Linux/LinuxToolChain.cs
@@ -128,7 +128,6 @@ namespace UnrealBuildTool
                 Result += " -Wno-unused-local-typedefs";
                 Result += " -Wno-multichar";
                 Result += " -Wno-unused-but-set-variable";
-                Result += " -Wno-sequence-point"; // TaskGraph.cpp:755 FTaskThread* Target = Target = &Thread(ThreadToExecuteOn);
                 Result += " -Wno-strict-overflow"; // Array.h:518
             }
             else
@@ -138,6 +137,7 @@ namespace UnrealBuildTool
                 Result += " -Wno-unused-private-field";     // MultichannelTcpSocket.h triggers this, possibly more
                 // this hides the "warning : comparison of unsigned expression < 0 is always false" type warnings due to constant comparisons, which are possible with template arguments
                 Result += " -Wno-tautological-compare";
+				Result += " -Wno-undefined-bool-conversion";	// hides checking if 'this' pointer is null
                 if (!CrossCompiling())
                 {
                     Result += " -Wno-logical-op-parentheses";   // needed for external headers we can't change
@@ -190,7 +190,7 @@ namespace UnrealBuildTool
             }
             else
             {
-                Result += " -O2";
+                Result += " -O3";
             }
 
             if (CompileEnvironment.Config.bIsBuildingDLL)
diff --git a/Engine/Source/Programs/UnrealBuildTool/Linux/UEBuildLinux.cs b/Engine/Source/Programs/UnrealBuildTool/Linux/UEBuildLinux.cs
--- a/Engine/Source/Programs/UnrealBuildTool/Linux/UEBuildLinux.cs
+++ b/Engine/Source/Programs/UnrealBuildTool/Linux/UEBuildLinux.cs
@@ -11,7 +11,7 @@ namespace UnrealBuildTool
     class LinuxPlatform : UEBuildPlatform
     {
         /** This is the SDK version we support */
-        static private string ExpectedSDKVersion = "v3_clang-3.3_ld-2.24_glibc-2.12.2";
+		static private string ExpectedSDKVersion = "v4_clang-3.5.0_ld-2.24_glibc-2.12.2";
 
         /** Platform name (embeds architecture for now) */
         static private string TargetPlatformName = "Linux_x64";
