Remove the default resource file from Windows binaries if the user has specified their own, and prevent the user's resource file from being ignored if it isn't the first encountered by UBT.

diff --git a/Engine/Source/Programs/UnrealBuildTool/Configuration/UEBuildBinary.cs b/Engine/Source/Programs/UnrealBuildTool/Configuration/UEBuildBinary.cs
--- a/Engine/Source/Programs/UnrealBuildTool/Configuration/UEBuildBinary.cs
+++ b/Engine/Source/Programs/UnrealBuildTool/Configuration/UEBuildBinary.cs
@@ -536,32 +536,25 @@ namespace UnrealBuildTool
 
 				// NOTE: Because of 'Shared PCHs', in monolithic builds the same PCH file may appear as a link input
 				// multiple times for a single binary.  We'll check for that here, and only add it once.  This avoids
-				// a linker warning about redundant .obj files. The same is true for .res files. We only take the first one.
-				bool bHasResFile = false;
-				foreach (var InputFile in BinaryLinkEnvironment.InputFiles)
-				{
-					if (InputFile.AbsolutePath.EndsWith(".res") && !InputFile.AbsolutePath.EndsWith(".inl.res")) 
-					{ 
-						bHasResFile = true;
-						break;
-					}
-				}
-
+				// a linker warning about redundant .obj files. 
 				foreach( var LinkInputFile in LinkInputFiles )
 				{
-					bool bIsResourceFile = LinkInputFile.AbsolutePath.EndsWith(".res") && !LinkInputFile.AbsolutePath.EndsWith(".inl.res");
-					if( !BinaryLinkEnvironment.InputFiles.Contains( LinkInputFile ) && (!bHasResFile || !bIsResourceFile))
+					if( !BinaryLinkEnvironment.InputFiles.Contains( LinkInputFile ))
 					{
 						BinaryLinkEnvironment.InputFiles.Add( LinkInputFile );
-						bHasResFile = bHasResFile || bIsResourceFile;
-
 					}
 				}
-
+				
 				// Allow the module to modify the link environment for the binary.
 				Module.SetupPrivateLinkEnvironment(ref BinaryLinkEnvironment,ref BinaryDependencies,ref LinkEnvironmentVisitedModules);
 			}
 
+			// Remove the default resource file on Windows (PCLaunch.rc) if the user has specified their own
+			if(BinaryLinkEnvironment.InputFiles.Select(Item => Path.GetFileName(Item.AbsolutePath).ToLower()).Any(Name => Name.EndsWith(".res") && !Name.EndsWith(".inl.res") && Name != "pclaunch.rc.res"))
+			{
+				BinaryLinkEnvironment.InputFiles.RemoveAll(x => Path.GetFileName(x.AbsolutePath).ToLower() == "pclaunch.rc.res");
+			}
+
 			// Allow the binary dependencies to modify the link environment.
 			foreach(var BinaryDependency in BinaryDependencies)
 			{
