Determine compiler version on Linux in runtime.

This allows us to better match switches given the multitude of (cross-)toolchains in use currently.

diff --git a/Engine/Source/Programs/UnrealBuildTool/Linux/LinuxToolChain.cs b/Engine/Source/Programs/UnrealBuildTool/Linux/LinuxToolChain.cs
--- a/Engine/Source/Programs/UnrealBuildTool/Linux/LinuxToolChain.cs
+++ b/Engine/Source/Programs/UnrealBuildTool/Linux/LinuxToolChain.cs
@@ -44,6 +44,89 @@ namespace UnrealBuildTool
             return null;
         }
 
+		/** Splits compiler version string into numerical components, leaving unchanged if not known */
+		private void DetermineCompilerMajMinPatchFromVersionString()
+		{
+			string[] Parts = CompilerVersionString.Split('.');
+			if (Parts.Length >= 1)
+			{
+				CompilerVersionMajor = Convert.ToInt32(Parts[0]);
+			}
+			if (Parts.Length >= 2)
+			{
+				CompilerVersionMinor = Convert.ToInt32(Parts[1]);
+			}
+			if (Parts.Length >= 3)
+			{
+				CompilerVersionPatch = Convert.ToInt32(Parts[2]);
+			}
+		}
+
+		/** Queries compiler for the version */
+		private bool DetermineCompilerVersion()
+		{
+			CompilerVersionString = null;
+			CompilerVersionMajor = -1;
+			CompilerVersionMinor = -1;
+			CompilerVersionPatch = -1;
+
+			Process Proc = new Process();
+			Proc.StartInfo.UseShellExecute = false;
+			Proc.StartInfo.CreateNoWindow = true;
+			Proc.StartInfo.RedirectStandardOutput = true;
+			Proc.StartInfo.RedirectStandardError = true;
+
+			if (!String.IsNullOrEmpty(GCCPath))
+			{
+				Proc.StartInfo.FileName = GCCPath;
+				Proc.StartInfo.Arguments = " -dumpversion";
+
+				Proc.Start();
+				Proc.WaitForExit();
+
+				if (Proc.ExitCode == 0)
+				{
+					// read just the first string
+					CompilerVersionString = Proc.StandardOutput.ReadLine();
+					DetermineCompilerMajMinPatchFromVersionString();
+				}
+			}
+			else if (!String.IsNullOrEmpty(ClangPath))
+			{
+				Proc.StartInfo.FileName = ClangPath;
+				Proc.StartInfo.Arguments = " --version";
+
+				Proc.Start();
+				Proc.WaitForExit();
+
+				if (Proc.ExitCode == 0)
+				{
+					// read just the first string
+					string VersionString = Proc.StandardOutput.ReadLine();
+
+					Regex VersionPattern = new Regex("version \\d+(\\.\\d+)+");
+					Match VersionMatch = VersionPattern.Match(VersionString);
+
+					// version match will be like "version 3.3", so remove the "version"
+					if (VersionMatch.Value.StartsWith("version "))
+					{
+						CompilerVersionString = VersionMatch.Value.Replace("version ", "");
+
+						DetermineCompilerMajMinPatchFromVersionString();
+					}
+				}
+			}
+			else
+			{
+				// icl?
+			}
+
+			Console.WriteLine("Using {0} version '{1}' (string), {2} (major), {3} (minor), {4} (patch)", 
+				String.IsNullOrEmpty(ClangPath) ? "gcc" : "clang",
+				CompilerVersionString, CompilerVersionMajor, CompilerVersionMinor, CompilerVersionPatch);
+			return !String.IsNullOrEmpty(CompilerVersionString);
+		}
+
         public override void RegisterToolChain()
         {
             if (!CrossCompiling())
@@ -53,6 +136,12 @@ namespace UnrealBuildTool
                 GCCPath = Which("g++");
                 ArPath = Which("ar");
                 RanlibPath = Which("ranlib");
+
+				// if clang is available, zero out gcc (@todo: support runtime switching?)
+				if (!String.IsNullOrEmpty(ClangPath))
+				{
+					GCCPath = null;
+				}
             }
             else
             {
@@ -72,6 +161,12 @@ namespace UnrealBuildTool
                 RanlibPath = Path.Combine(BaseLinuxPath, @"bin\x86_64-unknown-linux-gnu-ranlib.exe");
             }
 
+			if (!DetermineCompilerVersion())
+			{
+				Console.WriteLine("Could not determine version of the compiler, not registering Linux toolchain");
+				return;
+			}
+
             // Register this tool chain for both Linux
             if (BuildConfiguration.bPrintDebugInfo)
             {
@@ -80,6 +175,14 @@ namespace UnrealBuildTool
             UEToolChain.RegisterPlatformToolChain(CPPTargetPlatform.Linux, this);
         }
 
+		/** Checks if compiler version matches the requirements */
+		private static bool CompilerVersionGreaterOrEqual(int Major, int Minor, int Patch)
+		{
+			return CompilerVersionMajor > Major ||
+				(CompilerVersionMajor == Major && CompilerVersionMinor > Minor) ||
+				(CompilerVersionMajor == Major && CompilerVersionMinor == Minor && CompilerVersionPatch >= Patch);
+		}
+
         static string GetCLArguments_Global(CPPEnvironment CompileEnvironment)
         {
             string Result = "";
@@ -137,7 +240,13 @@ namespace UnrealBuildTool
                 Result += " -Wno-unused-private-field";     // MultichannelTcpSocket.h triggers this, possibly more
                 // this hides the "warning : comparison of unsigned expression < 0 is always false" type warnings due to constant comparisons, which are possible with template arguments
                 Result += " -Wno-tautological-compare";
-				Result += " -Wno-undefined-bool-conversion";	// hides checking if 'this' pointer is null
+
+				// this switch is understood by clang 3.5.0, but not clang-3.5 as packaged by Ubuntu 14.04 atm
+				if (CompilerVersionGreaterOrEqual(3, 5, 0))
+				{
+					Result += " -Wno-undefined-bool-conversion";	// hides checking if 'this' pointer is null
+				}
+
                 if (!CrossCompiling())
                 {
                     Result += " -Wno-logical-op-parentheses";   // needed for external headers we can't change
@@ -359,6 +468,15 @@ namespace UnrealBuildTool
         static string ArPath;
         static string RanlibPath;
 
+		/** Version string of the current compiler, whether clang or gcc or whatever */
+		static string CompilerVersionString;
+		/** Major version of the current compiler, whether clang or gcc or whatever */
+		static int CompilerVersionMajor = -1;
+		/** Minor version of the current compiler, whether clang or gcc or whatever */
+		static int CompilerVersionMinor = -1;
+		/** Patch version of the current compiler, whether clang or gcc or whatever */
+		static int CompilerVersionPatch = -1;
+
 		/** Track which scripts need to be deleted before appending to */
 		private bool bHasWipedFixDepsScript = false;
 
