UBT: Do not spam with compiler version when cross-compiling.

diff --git a/Engine/Source/Programs/UnrealBuildTool/Linux/LinuxToolChain.cs b/Engine/Source/Programs/UnrealBuildTool/Linux/LinuxToolChain.cs
--- a/Engine/Source/Programs/UnrealBuildTool/Linux/LinuxToolChain.cs
+++ b/Engine/Source/Programs/UnrealBuildTool/Linux/LinuxToolChain.cs
@@ -121,9 +121,12 @@ namespace UnrealBuildTool
 				// icl?
 			}
 
-			Console.WriteLine("Using {0} version '{1}' (string), {2} (major), {3} (minor), {4} (patch)", 
-				String.IsNullOrEmpty(ClangPath) ? "gcc" : "clang",
-				CompilerVersionString, CompilerVersionMajor, CompilerVersionMinor, CompilerVersionPatch);
+			if (!CrossCompiling())
+			{
+				Console.WriteLine("Using {0} version '{1}' (string), {2} (major), {3} (minor), {4} (patch)",
+					String.IsNullOrEmpty(ClangPath) ? "gcc" : "clang",
+					CompilerVersionString, CompilerVersionMajor, CompilerVersionMinor, CompilerVersionPatch);
+			}
 			return !String.IsNullOrEmpty(CompilerVersionString);
 		}
 
