Move all include and library search directories representation from List to HashSet to automatically avoid duplicates.

diff --git a/Engine/Source/Programs/UnrealBuildTool/Configuration/UEBuildModule.cs b/Engine/Source/Programs/UnrealBuildTool/Configuration/UEBuildModule.cs
--- a/Engine/Source/Programs/UnrealBuildTool/Configuration/UEBuildModule.cs
+++ b/Engine/Source/Programs/UnrealBuildTool/Configuration/UEBuildModule.cs
@@ -98,10 +98,10 @@ namespace UnrealBuildTool
 			return ModuleType;
 		}
 
-		/** Converts an optional string list parameter to a well-defined list. */
-		protected static List<string> ListFromOptionalEnumerableStringParameter(IEnumerable<string> InEnumerableStrings)
+		/** Converts an optional string list parameter to a well-defined hash set. */
+		protected static HashSet<string> HashSetFromOptionalEnumerableStringParameter(IEnumerable<string> InEnumerableStrings)
 		{
-			return InEnumerableStrings == null ? new List<string>() : new List<string>(InEnumerableStrings);
+			return InEnumerableStrings == null ? new HashSet<string>() : new HashSet<string>(InEnumerableStrings);
 		}
 
 		/** The target which owns this module. */
@@ -140,32 +140,32 @@ namespace UnrealBuildTool
 		/** Whether this module is included in the current target.  Only set after UEBuildBinary.BindModules is called. */
 		public bool bIncludedInTarget = false;
 
-		protected readonly List<string> PublicDefinitions;
-		protected readonly List<string> PublicIncludePaths;
-		protected readonly List<string> PrivateIncludePaths;
-		protected readonly List<string> PublicSystemIncludePaths;
-		protected readonly List<string> PublicLibraryPaths;
-		protected readonly List<string> PublicAdditionalLibraries;
-		protected readonly List<string> PublicFrameworks;
-		protected readonly List<string> PublicWeakFrameworks;
-		protected readonly List<UEBuildFramework> PublicAdditionalFrameworks;
-		protected readonly List<string> PublicAdditionalShadowFiles;
-		protected readonly List<UEBuildBundleResource> PublicAdditionalBundleResources;
+		protected readonly HashSet<string> PublicDefinitions;
+		protected readonly HashSet<string> PublicIncludePaths;
+		protected readonly HashSet<string> PrivateIncludePaths;
+		protected readonly HashSet<string> PublicSystemIncludePaths;
+		protected readonly HashSet<string> PublicLibraryPaths;
+		protected readonly HashSet<string> PublicAdditionalLibraries;
+		protected readonly HashSet<string> PublicFrameworks;
+		protected readonly HashSet<string> PublicWeakFrameworks;
+		protected readonly HashSet<UEBuildFramework> PublicAdditionalFrameworks;
+		protected readonly HashSet<string> PublicAdditionalShadowFiles;
+		protected readonly HashSet<UEBuildBundleResource> PublicAdditionalBundleResources;
 
 		/** Names of modules with header files that this module's public interface needs access to. */
-		protected List<string> PublicIncludePathModuleNames;
+		protected HashSet<string> PublicIncludePathModuleNames;
 
 		/** Names of modules that this module's public interface depends on. */
-		protected List<string> PublicDependencyModuleNames;
+		protected HashSet<string> PublicDependencyModuleNames;
 
 		/** Names of DLLs that this module should delay load */
-		protected List<string> PublicDelayLoadDLLs;
+		protected HashSet<string> PublicDelayLoadDLLs;
 
 		/** Names of modules with header files that this module's private implementation needs access to. */
-		protected List<string> PrivateIncludePathModuleNames;
+		protected HashSet<string> PrivateIncludePathModuleNames;
 	
 		/** Names of modules that this module's private implementation depends on. */
-		protected List<string> PrivateDependencyModuleNames;
+		protected HashSet<string> PrivateDependencyModuleNames;
 
 		/** If any of this module's dependent modules circularly reference this module, then they must be added to this list */
 		public HashSet<string> CircularlyReferencedDependentModules
@@ -175,10 +175,10 @@ namespace UnrealBuildTool
 		}
 
 		/** Extra modules this module may require at run time */
-		protected List<string> DynamicallyLoadedModuleNames;
+		protected HashSet<string> DynamicallyLoadedModuleNames;
 
         /** Extra modules this module may require at run time, that are on behalf of another platform (i.e. shader formats and the like) */
-        protected List<string> PlatformSpecificDynamicallyLoadedModuleNames;
+		protected HashSet<string> PlatformSpecificDynamicallyLoadedModuleNames;
 
 		public UEBuildModule(
 			UEBuildTarget InTarget,
@@ -213,25 +213,25 @@ namespace UnrealBuildTool
 			Type = InType;
 			ModuleDirectory = InModuleDirectory;
 			OutputDirectory = InOutputDirectory;
-			PublicDefinitions = ListFromOptionalEnumerableStringParameter(InPublicDefinitions);
-			PublicIncludePaths = ListFromOptionalEnumerableStringParameter(InPublicIncludePaths);
-			PublicSystemIncludePaths = ListFromOptionalEnumerableStringParameter(InPublicSystemIncludePaths);
-			PublicLibraryPaths = ListFromOptionalEnumerableStringParameter(InPublicLibraryPaths);
-			PublicAdditionalLibraries = ListFromOptionalEnumerableStringParameter(InPublicAdditionalLibraries);
-			PublicFrameworks = ListFromOptionalEnumerableStringParameter(InPublicFrameworks);
-			PublicWeakFrameworks = ListFromOptionalEnumerableStringParameter(InPublicWeakFrameworks);
-			PublicAdditionalFrameworks = InPublicAdditionalFrameworks == null ? new List<UEBuildFramework>() : new List<UEBuildFramework>(InPublicAdditionalFrameworks);
-			PublicAdditionalShadowFiles = ListFromOptionalEnumerableStringParameter(InPublicAdditionalShadowFiles);
-			PublicAdditionalBundleResources = InPublicAdditionalBundleResources == null ? new List<UEBuildBundleResource>() : new List<UEBuildBundleResource>(InPublicAdditionalBundleResources);
-			PublicIncludePathModuleNames = ListFromOptionalEnumerableStringParameter( InPublicIncludePathModuleNames );
-			PublicDependencyModuleNames = ListFromOptionalEnumerableStringParameter(InPublicDependencyModuleNames);
-			PublicDelayLoadDLLs = ListFromOptionalEnumerableStringParameter(InPublicDelayLoadDLLs);
-			PrivateIncludePaths = ListFromOptionalEnumerableStringParameter(InPrivateIncludePaths);
-			PrivateIncludePathModuleNames = ListFromOptionalEnumerableStringParameter( InPrivateIncludePathModuleNames );
-			PrivateDependencyModuleNames = ListFromOptionalEnumerableStringParameter( InPrivateDependencyModuleNames );
-			CircularlyReferencedDependentModules = new HashSet<string>( ListFromOptionalEnumerableStringParameter( InCircularlyReferencedDependentModules ) );
-			DynamicallyLoadedModuleNames = ListFromOptionalEnumerableStringParameter( InDynamicallyLoadedModuleNames );
-            PlatformSpecificDynamicallyLoadedModuleNames = ListFromOptionalEnumerableStringParameter(InPlatformSpecificDynamicallyLoadedModuleNames);
+			PublicDefinitions = HashSetFromOptionalEnumerableStringParameter(InPublicDefinitions);
+			PublicIncludePaths = HashSetFromOptionalEnumerableStringParameter(InPublicIncludePaths);
+			PublicSystemIncludePaths = HashSetFromOptionalEnumerableStringParameter(InPublicSystemIncludePaths);
+			PublicLibraryPaths = HashSetFromOptionalEnumerableStringParameter(InPublicLibraryPaths);
+			PublicAdditionalLibraries = HashSetFromOptionalEnumerableStringParameter(InPublicAdditionalLibraries);
+			PublicFrameworks = HashSetFromOptionalEnumerableStringParameter(InPublicFrameworks);
+			PublicWeakFrameworks = HashSetFromOptionalEnumerableStringParameter(InPublicWeakFrameworks);
+			PublicAdditionalFrameworks = InPublicAdditionalFrameworks == null ? new HashSet<UEBuildFramework>() : new HashSet<UEBuildFramework>(InPublicAdditionalFrameworks);
+			PublicAdditionalShadowFiles = HashSetFromOptionalEnumerableStringParameter(InPublicAdditionalShadowFiles);
+			PublicAdditionalBundleResources = InPublicAdditionalBundleResources == null ? new HashSet<UEBuildBundleResource>() : new HashSet<UEBuildBundleResource>(InPublicAdditionalBundleResources);
+			PublicIncludePathModuleNames = HashSetFromOptionalEnumerableStringParameter( InPublicIncludePathModuleNames );
+			PublicDependencyModuleNames = HashSetFromOptionalEnumerableStringParameter(InPublicDependencyModuleNames);
+			PublicDelayLoadDLLs = HashSetFromOptionalEnumerableStringParameter(InPublicDelayLoadDLLs);
+			PrivateIncludePaths = HashSetFromOptionalEnumerableStringParameter(InPrivateIncludePaths);
+			PrivateIncludePathModuleNames = HashSetFromOptionalEnumerableStringParameter( InPrivateIncludePathModuleNames );
+			PrivateDependencyModuleNames = HashSetFromOptionalEnumerableStringParameter( InPrivateDependencyModuleNames );
+			CircularlyReferencedDependentModules = new HashSet<string>( HashSetFromOptionalEnumerableStringParameter( InCircularlyReferencedDependentModules ) );
+			DynamicallyLoadedModuleNames = HashSetFromOptionalEnumerableStringParameter( InDynamicallyLoadedModuleNames );
+            PlatformSpecificDynamicallyLoadedModuleNames = HashSetFromOptionalEnumerableStringParameter(InPlatformSpecificDynamicallyLoadedModuleNames);
 			IsRedistributableOverride = InIsRedistributableOverride;
 
 			Target.RegisterModule(this);
@@ -524,8 +524,8 @@ namespace UnrealBuildTool
 		protected virtual void SetupPublicCompileEnvironment(
 			UEBuildBinary SourceBinary,
 			bool bIncludePathsOnly,
-			ref List<string> IncludePaths,
-			ref List<string> SystemIncludePaths,
+			ref HashSet<string> IncludePaths,
+			ref HashSet<string> SystemIncludePaths,
 			ref List<string> Definitions,
 			ref List<UEBuildFramework> AdditionalFrameworks,
 			ref Dictionary<UEBuildModule, bool> VisitedModules
@@ -624,26 +624,23 @@ namespace UnrealBuildTool
 		}
 
 		/** Adds PathsToAdd to IncludePaths, performing path normalization and ignoring duplicates. */
-		protected void AddIncludePathsWithChecks(List<string> IncludePaths, List<string> PathsToAdd)
+		protected void AddIncludePathsWithChecks(HashSet<string> IncludePaths, HashSet<string> PathsToAdd)
 		{
 			foreach (var Path in PathsToAdd)
 			{
 				var NormalizedPath = Path.TrimEnd('/');
-				if (!IncludePaths.Contains(NormalizedPath))
+				// If path doesn't exist, it may contain VC macro (which is passed directly to and expanded by compiler).
+				if (Directory.Exists(NormalizedPath) || DoesPathContainVCMacro(NormalizedPath))
 				{
-					// If path doesn't exist, it may contain VC macro (which is passed directly to and expanded by compiler).
-					if (Directory.Exists(NormalizedPath) || DoesPathContainVCMacro(NormalizedPath))
-					{
-						IncludePaths.Add(NormalizedPath);
-					}
+					IncludePaths.Add(NormalizedPath);
 				}
 			}
 		}
 
 		/** Sets up the environment for compiling this module. */
 		protected virtual void SetupPrivateCompileEnvironment(
-			ref List<string> IncludePaths,
-			ref List<string> SystemIncludePaths,
+			ref HashSet<string> IncludePaths,
+			ref HashSet<string> SystemIncludePaths,
 			ref List<string> Definitions,
 			ref List<UEBuildFramework> AdditionalFrameworks
 			)
@@ -932,7 +929,7 @@ namespace UnrealBuildTool
 		public readonly SourceFilesClass SourceFilesFound = new SourceFilesClass();
 
 		/** The preprocessor definitions used to compile this module's private implementation. */
-		List<string> Definitions;
+		HashSet<string> Definitions;
 
 		/// When set, allows this module to report compiler definitions and include paths for Intellisense
 		IntelliSenseGatherer IntelliSenseGatherer;
@@ -1087,7 +1084,7 @@ namespace UnrealBuildTool
 				CategorizeSourceFiles(InSourceFiles, SourceFilesToBuild);
 			}
 
-			Definitions = ListFromOptionalEnumerableStringParameter(InDefinitions);
+			Definitions = HashSetFromOptionalEnumerableStringParameter(InDefinitions);
 			foreach (var Def in Definitions)
 			{
 				Log.TraceVerbose("Compile Env {0}: {1}", Name, Def);
@@ -2089,7 +2086,7 @@ namespace UnrealBuildTool
 	class UEBuildModuleCPPCLR : UEBuildModuleCPP
 	{
 		/** The assemblies referenced by the module's private implementation. */
-		List<string> PrivateAssemblyReferences;
+		HashSet<string> PrivateAssemblyReferences;
 
 		public UEBuildModuleCPPCLR(
 			UEBuildTarget InTarget,
@@ -2137,7 +2134,7 @@ namespace UnrealBuildTool
 			InAllowSharedPCH, InSharedPCHHeaderFile, InUseRTTI, InEnableBufferSecurityChecks, InFasterWithoutUnity, InMinFilesUsingPrecompiledHeaderOverride,
 			InEnableExceptions, bInBuildSourceFiles)
 		{
-			PrivateAssemblyReferences = ListFromOptionalEnumerableStringParameter(InPrivateAssemblyReferences);
+			PrivateAssemblyReferences = HashSetFromOptionalEnumerableStringParameter(InPrivateAssemblyReferences);
 		}
 
 		// UEBuildModule interface.
diff --git a/Engine/Source/Programs/UnrealBuildTool/IOS/IOSToolChain.cs b/Engine/Source/Programs/UnrealBuildTool/IOS/IOSToolChain.cs
--- a/Engine/Source/Programs/UnrealBuildTool/IOS/IOSToolChain.cs
+++ b/Engine/Source/Programs/UnrealBuildTool/IOS/IOSToolChain.cs
@@ -462,8 +462,8 @@ namespace UnrealBuildTool
 			}
 
 			// Add include paths to the argument list.
-			List<string> AllIncludes = CompileEnvironment.Config.CPPIncludeInfo.IncludePaths;
-			AllIncludes.AddRange(CompileEnvironment.Config.CPPIncludeInfo.SystemIncludePaths);
+			HashSet<string> AllIncludes = CompileEnvironment.Config.CPPIncludeInfo.IncludePaths;
+			AllIncludes.UnionWith(CompileEnvironment.Config.CPPIncludeInfo.SystemIncludePaths);
 			foreach (string IncludePath in AllIncludes)
 			{
 				Arguments += string.Format(" -I\"{0}\"", ConvertPath(Path.GetFullPath(IncludePath)));
diff --git a/Engine/Source/Programs/UnrealBuildTool/Mac/MacToolChain.cs b/Engine/Source/Programs/UnrealBuildTool/Mac/MacToolChain.cs
--- a/Engine/Source/Programs/UnrealBuildTool/Mac/MacToolChain.cs
+++ b/Engine/Source/Programs/UnrealBuildTool/Mac/MacToolChain.cs
@@ -245,8 +245,8 @@ namespace UnrealBuildTool
 			}
 
 			// Add include paths to the argument list.
-			List<string> AllIncludes = CompileEnvironment.Config.CPPIncludeInfo.IncludePaths;
-			AllIncludes.AddRange(CompileEnvironment.Config.CPPIncludeInfo.SystemIncludePaths);
+			HashSet<string> AllIncludes = CompileEnvironment.Config.CPPIncludeInfo.IncludePaths;
+			AllIncludes.UnionWith(CompileEnvironment.Config.CPPIncludeInfo.SystemIncludePaths);
 			foreach (string IncludePath in AllIncludes)
 			{
 				Arguments.Append(" -I\"");
diff --git a/Engine/Source/Programs/UnrealBuildTool/System/CPPEnvironment.cs b/Engine/Source/Programs/UnrealBuildTool/System/CPPEnvironment.cs
--- a/Engine/Source/Programs/UnrealBuildTool/System/CPPEnvironment.cs
+++ b/Engine/Source/Programs/UnrealBuildTool/System/CPPEnvironment.cs
@@ -271,8 +271,8 @@ namespace UnrealBuildTool
 			bIsBuildingLibrary                     = InCopyEnvironment.bIsBuildingLibrary;
 			bIsBuildingDLL                         = InCopyEnvironment.bIsBuildingDLL;
 			CLRMode                                = InCopyEnvironment.CLRMode;
-			CPPIncludeInfo.IncludePaths                 .AddRange(InCopyEnvironment.CPPIncludeInfo.IncludePaths);
-			CPPIncludeInfo.SystemIncludePaths           .AddRange(InCopyEnvironment.CPPIncludeInfo.SystemIncludePaths);
+			CPPIncludeInfo.IncludePaths            .UnionWith(InCopyEnvironment.CPPIncludeInfo.IncludePaths);
+			CPPIncludeInfo.SystemIncludePaths      .UnionWith(InCopyEnvironment.CPPIncludeInfo.SystemIncludePaths);
 			SystemDotNetAssemblyPaths    .AddRange(InCopyEnvironment.SystemDotNetAssemblyPaths);
 			FrameworkAssemblyDependencies.AddRange(InCopyEnvironment.FrameworkAssemblyDependencies);
  			PrivateAssemblyDependencies  .AddRange(InCopyEnvironment.PrivateAssemblyDependencies);
diff --git a/Engine/Source/Programs/UnrealBuildTool/System/CPPHeaders.cs b/Engine/Source/Programs/UnrealBuildTool/System/CPPHeaders.cs
--- a/Engine/Source/Programs/UnrealBuildTool/System/CPPHeaders.cs
+++ b/Engine/Source/Programs/UnrealBuildTool/System/CPPHeaders.cs
@@ -14,13 +14,13 @@ namespace UnrealBuildTool
 	public class CPPIncludeInfo
 	{
 		/** Ordered list of include paths for the module  */
-		public List<string> IncludePaths = new List<string>();
+		public HashSet<string> IncludePaths = new HashSet<string>();
 
 		/**
 		 * The include paths where changes to contained files won't cause dependent C++ source files to
 		 * be recompiled, unless BuildConfiguration.bCheckSystemHeadersForModification==true.
 		 */
-		public List<string> SystemIncludePaths = new List<string>();
+		public HashSet<string> SystemIncludePaths = new HashSet<string>();
 
 		/** Contains a mapping from filename to the full path of the header in this environment.  This is used to optimized include path lookups at runtime for any given single module. */
 		[NonSerialized]
diff --git a/Engine/Source/Programs/UnrealBuildTool/System/Project.cs b/Engine/Source/Programs/UnrealBuildTool/System/Project.cs
--- a/Engine/Source/Programs/UnrealBuildTool/System/Project.cs
+++ b/Engine/Source/Programs/UnrealBuildTool/System/Project.cs
@@ -24,7 +24,7 @@ namespace UnrealBuildTool
 		/// </summary>
 		/// <param name="NewIncludePaths">List of include paths to add</param>
 		/// <param name="bAddingSystemIncludes">Are the include paths to add system include paths</param>
-		void AddInteliiSenseIncludePaths( List<string> NewIncludePaths, bool bAddingSystemIncludes );
+		void AddInteliiSenseIncludePaths(HashSet<string> NewIncludePaths, bool bAddingSystemIncludes);
 	}
 
 
@@ -307,7 +307,7 @@ namespace UnrealBuildTool
 		/// Adds all of the specified include paths to this VCProject's list of include paths for all modules in the project
 		/// </summary>
 		/// <param name="NewIncludePaths">List of include paths to add</param>
-		public void AddInteliiSenseIncludePaths( List<string> NewIncludePaths, bool bAddingSystemIncludes ) 
+		public void AddInteliiSenseIncludePaths(HashSet<string> NewIncludePaths, bool bAddingSystemIncludes) 
 		{
 			foreach( var CurPath in NewIncludePaths )
 			{
