Skip include path normalization and check for existence when generating intellisense.

diff --git a/Engine/Source/Programs/UnrealBuildTool/Configuration/UEBuildModule.cs b/Engine/Source/Programs/UnrealBuildTool/Configuration/UEBuildModule.cs
--- a/Engine/Source/Programs/UnrealBuildTool/Configuration/UEBuildModule.cs
+++ b/Engine/Source/Programs/UnrealBuildTool/Configuration/UEBuildModule.cs
@@ -626,13 +626,22 @@ namespace UnrealBuildTool
 		/** Adds PathsToAdd to IncludePaths, performing path normalization and ignoring duplicates. */
 		protected void AddIncludePathsWithChecks(HashSet<string> IncludePaths, HashSet<string> PathsToAdd)
 		{
-			foreach (var Path in PathsToAdd)
+			if (ProjectFileGenerator.bGenerateProjectFiles)
 			{
-				var NormalizedPath = Path.TrimEnd('/');
-				// If path doesn't exist, it may contain VC macro (which is passed directly to and expanded by compiler).
-				if (Directory.Exists(NormalizedPath) || DoesPathContainVCMacro(NormalizedPath))
+				// Extra checks are switched off for IntelliSense generation as they provide
+				// no additional value and cause performance impact.
+				IncludePaths.UnionWith(PathsToAdd);
+			}
+			else
+			{
+				foreach (var Path in PathsToAdd)
 				{
-					IncludePaths.Add(NormalizedPath);
+					var NormalizedPath = Path.TrimEnd('/');
+					// If path doesn't exist, it may contain VC macro (which is passed directly to and expanded by compiler).
+					if (Directory.Exists(NormalizedPath) || DoesPathContainVCMacro(NormalizedPath))
+					{
+						IncludePaths.Add(NormalizedPath);
+					}
 				}
 			}
 		}
