HotReload: pass -nosharedpch when doing hot-reload or module recompile from the editor instead of forcing off shared pch when passing -editorrecompile to UBT

diff --git a/Engine/Source/Developer/HotReload/Private/HotReload.cpp b/Engine/Source/Developer/HotReload/Private/HotReload.cpp
--- a/Engine/Source/Developer/HotReload/Private/HotReload.cpp
+++ b/Engine/Source/Developer/HotReload/Private/HotReload.cpp
@@ -1093,6 +1093,9 @@ bool FHotReloadModule::StartCompilingModuleDLLs(const FString& GameName, const T
 		ExtraArg += TEXT( "-FailIfGeneratedCodeChanges " );
 	}
 
+	// Shared PCH does no work with hot-reloading modules as we don't scan all modules for them.
+	ExtraArg += TEXT("-nosharedpch ");
+
 	// If the're no game modules loaded, then it's not a code-based project and the target
 	// for UBT should be the editor.
 	FString TargetName;
diff --git a/Engine/Source/Programs/UnrealBuildTool/Configuration/UEBuildTarget.cs b/Engine/Source/Programs/UnrealBuildTool/Configuration/UEBuildTarget.cs
--- a/Engine/Source/Programs/UnrealBuildTool/Configuration/UEBuildTarget.cs
+++ b/Engine/Source/Programs/UnrealBuildTool/Configuration/UEBuildTarget.cs
@@ -2879,7 +2879,7 @@ namespace UnrealBuildTool
 							InDynamicallyLoadedModuleNames: RulesObject.DynamicallyLoadedModuleNames,
 							InPlatformSpecificDynamicallyLoadedModuleNames: RulesObject.PlatformSpecificDynamicallyLoadedModuleNames,
 							InOptimizeCode: RulesObject.OptimizeCode,
-							InAllowSharedPCH: (RulesObject.PCHUsage == ModuleRules.PCHUsageMode.NoSharedPCHs) ? false : !this.bEditorRecompile,
+							InAllowSharedPCH: (RulesObject.PCHUsage == ModuleRules.PCHUsageMode.NoSharedPCHs) ? false : true,
 							InSharedPCHHeaderFile: RulesObject.SharedPCHHeaderFile,
 							InUseRTTI: RulesObject.bUseRTTI,
 							InEnableBufferSecurityChecks: RulesObject.bEnableBufferSecurityChecks,
