build script, reenable xge for linux, android, and ps4

diff --git a/Engine/Source/Programs/UnrealBuildTool/Android/UEBuildAndroid.cs b/Engine/Source/Programs/UnrealBuildTool/Android/UEBuildAndroid.cs
--- a/Engine/Source/Programs/UnrealBuildTool/Android/UEBuildAndroid.cs
+++ b/Engine/Source/Programs/UnrealBuildTool/Android/UEBuildAndroid.cs
@@ -26,11 +26,6 @@ namespace UnrealBuildTool
 			return base.GetActiveArchitecture();
 		}
 
-        public override bool CanUseXGE()
-        {
-            return false;
-        }
-
         protected override bool PlatformSupportsAutoSDKs()
         {
             return true;
diff --git a/Engine/Source/Programs/UnrealBuildTool/Linux/UEBuildLinux.cs b/Engine/Source/Programs/UnrealBuildTool/Linux/UEBuildLinux.cs
--- a/Engine/Source/Programs/UnrealBuildTool/Linux/UEBuildLinux.cs
+++ b/Engine/Source/Programs/UnrealBuildTool/Linux/UEBuildLinux.cs
@@ -100,12 +100,6 @@ namespace UnrealBuildTool
 
             return SDKStatus.Invalid;
         }
-
-        public override bool CanUseXGE()
-        {
-            // disabled until XGE crash is fixed
-            return false;
-        }
         
         /**
          *	Register the platform with the UEBuildPlatform class
