Lazy-evaluated MSBuildPath, so that it only fails if it is needed.

diff --git a/Engine/Source/Programs/UnrealBuildTool/Windows/VCEnvironment.cs b/Engine/Source/Programs/UnrealBuildTool/Windows/VCEnvironment.cs
--- a/Engine/Source/Programs/UnrealBuildTool/Windows/VCEnvironment.cs
+++ b/Engine/Source/Programs/UnrealBuildTool/Windows/VCEnvironment.cs
@@ -21,7 +21,19 @@ namespace UnrealBuildTool
 		public readonly string            LinkerPath;           // The path to the linker for linking executables
 		public readonly string            LibraryLinkerPath;    // The path to the linker for linking libraries
 		public readonly string            ResourceCompilerPath; // The path to the resource compiler
-		public readonly string            MSBuildPath;          // The path to MSBuild
+
+		private string _MSBuildPath = null;
+		public string MSBuildPath // The path to MSBuild
+		{
+			get
+			{
+				if (_MSBuildPath == null)
+				{
+					_MSBuildPath = GetMSBuildToolPath();
+				}
+				return _MSBuildPath;
+			}
+		}
 
 		/**
 		 * Initializes environment variables required by toolchain. Different for 32 and 64 bit.
@@ -60,8 +72,6 @@ namespace UnrealBuildTool
 			var VCVarsBatchFile = Path.Combine(BaseVSToolPath, (Platform == CPPTargetPlatform.Win64) ? "../../VC/bin/x86_amd64/vcvarsx86_amd64.bat" : "vsvars32.bat");
 			Utils.SetEnvironmentVariablesFromBatchFile(VCVarsBatchFile);
 
-			MSBuildPath = GetMSBuildToolPath();
-
 			// When targeting Windows XP on Visual Studio 2012+, we need to override the Windows SDK include and lib path set
 			// by the batch file environment (http://blogs.msdn.com/b/vcblog/archive/2012/10/08/10357555.aspx)
 			if (WindowsPlatform.SupportWindowsXP)
