build script, disable android xge again

diff --git a/Engine/Source/Programs/UnrealBuildTool/Android/UEBuildAndroid.cs b/Engine/Source/Programs/UnrealBuildTool/Android/UEBuildAndroid.cs
--- a/Engine/Source/Programs/UnrealBuildTool/Android/UEBuildAndroid.cs
+++ b/Engine/Source/Programs/UnrealBuildTool/Android/UEBuildAndroid.cs
@@ -26,6 +26,11 @@ namespace UnrealBuildTool
 			return base.GetActiveArchitecture();
 		}
 
+		public override bool CanUseXGE()
+		{
+			return false;
+		}
+
         protected override bool PlatformSupportsAutoSDKs()
         {
             return true;
