Get ready for sdk upgrade. Remove linked in demangling code - newer sdks have this built in.

#HTML5

diff --git a/Engine/Source/Programs/UnrealBuildTool/HTML5/HTML5ToolChain.cs b/Engine/Source/Programs/UnrealBuildTool/HTML5/HTML5ToolChain.cs
--- a/Engine/Source/Programs/UnrealBuildTool/HTML5/HTML5ToolChain.cs
+++ b/Engine/Source/Programs/UnrealBuildTool/HTML5/HTML5ToolChain.cs
@@ -94,8 +94,6 @@ namespace UnrealBuildTool
             Result += " -s WARN_ON_UNDEFINED_SYMBOLS=1";
             // we want full ES2
             Result += " -s FULL_ES2=1 ";
-            // don't need UTF8 string support, and it slows string ops down
-            Result += " -s UTF_STRING_SUPPORT=0";
             // export console command handler. Export main func too because default exports ( e.g Main ) are overridden if we use custom exported functions. 
             Result += " -s EXPORTED_FUNCTIONS=\"['_main', '_resize_game']\" ";
 
@@ -215,7 +213,6 @@ namespace UnrealBuildTool
 
 				Result += " -s CASE_INSENSITIVE_FS=1 ";
 
-				Result += " --js-library Runtime/Core/Public/HTML5/HTML5DebugLogging.js ";
 
 				string BaseSDKPath = Environment.GetEnvironmentVariable("EMSCRIPTEN");
 				Result += " --js-library \"" + BaseSDKPath + "/Src/library_openal.js\" ";
@@ -327,18 +324,13 @@ namespace UnrealBuildTool
 				Arguments += string.Format(" -I\"{0}\"", IncludePath);
 			}
 
-            if ( ModuleName == "Launch" ) 
-                Arguments += string.Format(" -I\"{0}\"", BaseSDKPath + "/system/lib/libcxxabi/include" );
-        
+      
 			// Add preprocessor definitions to the argument list.
 			foreach (string Definition in CompileEnvironment.Config.Definitions)
 			{
 				Arguments += string.Format(" -D{0}", Definition);
 			}
 
-			// Create a compile action for each source file.
-            if (ModuleName == "Launch")
-                SourceFiles.Add(FileItem.GetItemByPath(BaseSDKPath + "/system/lib/libcxxabi/src/cxa_demangle.cpp")); 
         
 			var BuildPlatform = UEBuildPlatform.GetBuildPlatformForCPPTargetPlatform(CompileEnvironment.Config.Target.Platform);
 
