Fix crashes on broken/old toolchains.

- UBT was invoking non-existent compiler to determine its version.
- Happened during AutoSDK switch as well.

diff --git a/Engine/Source/Programs/UnrealBuildTool/Linux/LinuxToolChain.cs b/Engine/Source/Programs/UnrealBuildTool/Linux/LinuxToolChain.cs
--- a/Engine/Source/Programs/UnrealBuildTool/Linux/LinuxToolChain.cs
+++ b/Engine/Source/Programs/UnrealBuildTool/Linux/LinuxToolChain.cs
@@ -76,7 +76,7 @@ namespace UnrealBuildTool
 			Proc.StartInfo.RedirectStandardOutput = true;
 			Proc.StartInfo.RedirectStandardError = true;
 
-			if (!String.IsNullOrEmpty(GCCPath))
+			if (!String.IsNullOrEmpty(GCCPath) && File.Exists(GCCPath))
 			{
 				Proc.StartInfo.FileName = GCCPath;
 				Proc.StartInfo.Arguments = " -dumpversion";
@@ -91,7 +91,7 @@ namespace UnrealBuildTool
 					DetermineCompilerMajMinPatchFromVersionString();
 				}
 			}
-			else if (!String.IsNullOrEmpty(ClangPath))
+			else if (!String.IsNullOrEmpty(ClangPath) && File.Exists(ClangPath))
 			{
 				Proc.StartInfo.FileName = ClangPath;
 				Proc.StartInfo.Arguments = " --version";
