- Added early support for LaunchScreen storyboard/.xib files for IOS8, still a work in progress (plist issues need to be fixed up still)

diff --git a/Engine/Source/Programs/UnrealBuildTool/IOS/UEDeployIOS.cs b/Engine/Source/Programs/UnrealBuildTool/IOS/UEDeployIOS.cs
--- a/Engine/Source/Programs/UnrealBuildTool/IOS/UEDeployIOS.cs
+++ b/Engine/Source/Programs/UnrealBuildTool/IOS/UEDeployIOS.cs
@@ -1,4 +1,4 @@
-﻿// Copyright 1998-2014 Epic Games, Inc. All Rights Reserved.
+// Copyright 1998-2014 Epic Games, Inc. All Rights Reserved.
 
 using System;
 using System.Collections.Generic;
@@ -6,6 +6,7 @@ using System.Linq;
 using System.Text;
 using System.Xml;
 using System.IO;
+using System.Diagnostics;
 
 namespace UnrealBuildTool.IOS
 {
@@ -20,7 +21,7 @@ namespace UnrealBuildTool.IOS
 			UEBuildDeploy.RegisterBuildDeploy(UnrealTargetPlatform.IOS, this);
 		}
 
-		private static void CopyFileWithReplacements(string SourceFilename, string DestFilename, Dictionary<string, string> Replacements)
+		private static void CopyFileWithReplacements(string SourceFilename, string DestFilename, Dictionary<string, string> Replacements, List<string> AdditionalLines)
 		{
 			if (!File.Exists(SourceFilename))
 			{
@@ -44,16 +45,48 @@ namespace UnrealBuildTool.IOS
 			string Ext = Path.GetExtension(SourceFilename);
 			if (Ext == ".plist")
 			{
-				string Contents = File.ReadAllText(SourceFilename);
+				string[] Contents = File.ReadAllLines(SourceFilename);
+				StringBuilder NewContents = new StringBuilder();
+
+				int LastDictLine = 0;
+				for (int LineIndex = Contents.Length - 1; LineIndex >= 0; LineIndex--)
+				{
+					if (Contents[LineIndex].Trim() == "</dict>")
+					{
+						LastDictLine = LineIndex;
+						break;
+					}
+				}
 
 				// replace some varaibles
-				foreach (var Pair in Replacements)
+				for (int LineIndex = 0; LineIndex < Contents.Length; LineIndex++)
 				{
-					Contents = Contents.Replace(Pair.Key, Pair.Value);
+					string Line = Contents[LineIndex];
+
+					// inject before the last line
+					if (LineIndex == LastDictLine)
+					{
+						foreach (string ExtraLine in AdditionalLines)
+						{
+							string FixedLine = ExtraLine;
+							foreach (var Pair in Replacements)
+							{
+								FixedLine = FixedLine.Replace(Pair.Key, Pair.Value);
+							}
+							NewContents.Append(FixedLine + Environment.NewLine);
+						}
+					}
+
+					foreach (var Pair in Replacements)
+					{
+						Line = Line.Replace(Pair.Key, Pair.Value);
+					}
+
+					NewContents.Append(Line + Environment.NewLine);
+
 				}
 
-				// write out file
-				File.WriteAllText(DestFilename, Contents);
+				File.WriteAllText(DestFilename, NewContents.ToString());
 			}
 			else
 			{
@@ -215,18 +248,59 @@ namespace UnrealBuildTool.IOS
 				DestFileInfo.Attributes = DestFileInfo.Attributes & ~FileAttributes.ReadOnly;
 			}
 
+			// compile the launch .xib
+			string LaunchXib = InEngineDir + "/Build/IOS/Resources/Interface/LaunchScreen.xib";
+			if (File.Exists(BuildDirectory + "/Resources/Interface/LaunchScreen.xib"))
+			{
+				LaunchXib = BuildDirectory + "/Resources/Interface/LaunchScreen.xib";
+			}
+
+			List<string> PListAdditionalLines = new List<string>();
+			PListAdditionalLines.Add("\t<key>TestKey</key>");
+			PListAdditionalLines.Add("\t<string>" + BuildDirectory + "</string>");
+
+			bool bSkipDefaultPNGs = false;
+			if (File.Exists(LaunchXib))
+			{
+				string CommandLine = string.Format("--target-device iphone --target-device ipad --errors --warnings --notices --module {0} --minimum-deployment-target 8.0 --auto-activate-custom-fonts --output-format human-readable-text --compile {1}/LaunchScreen.nib {2}", InProjectName, Path.GetFullPath(AppDirectory), Path.GetFileName(LaunchXib));
+
+				// now we need to zipalign the apk to the final destination (to 4 bytes, must be 4)
+				ProcessStartInfo IBToolStartInfo = new ProcessStartInfo();
+				IBToolStartInfo.WorkingDirectory = Path.GetDirectoryName(LaunchXib);
+				IBToolStartInfo.FileName = "/Applications/Xcode.app/Contents/Developer/usr/bin/ibtool";
+				IBToolStartInfo.Arguments = CommandLine;
+				IBToolStartInfo.UseShellExecute = false;
+				Process CallIBTool = new Process();
+				CallIBTool.StartInfo = IBToolStartInfo;
+				CallIBTool.Start();
+				CallIBTool.WaitForExit();
+
+				PListAdditionalLines.Add("\t<key>UILaunchStoryboardName</key>");
+				PListAdditionalLines.Add("\t<string>LaunchScreen</string>");
+
+				PListAdditionalLines.Add("\t<key>NibLaunch</key>");
+				PListAdditionalLines.Add("\t<string>" + CommandLine + "</string>");
+
+				bSkipDefaultPNGs = true;
+			}
+
 			// copy plist file
 			string PListFile = InEngineDir + "/Build/IOS/UE4Game-Info.plist";
-			if (File.Exists(BuildDirectory + "/" + InProjectName + "-Info.plist"))
+			if (File.Exists(BuildDirectory + "/Info.plist"))
+			{
+				PListFile = BuildDirectory + "/Info.plist";
+			}
+			else if (File.Exists(BuildDirectory + "/" + InProjectName + "-Info.plist"))
 			{
 				PListFile = BuildDirectory + "/" + InProjectName + "-Info.plist";
 			}
 
+			// plist replacements
 			Dictionary<string, string> Replacements = new Dictionary<string, string>();
 			Replacements.Add("${EXECUTABLE_NAME}", GameName);
 			Replacements.Add("${BUNDLE_IDENTIFIER}", InProjectName.Replace("_", ""));
-			CopyFileWithReplacements(PListFile, AppDirectory + "/Info.plist", Replacements);
-			CopyFileWithReplacements(PListFile, IntermediateDirectory + "/" + GameName + "-Info.plist", Replacements);
+			CopyFileWithReplacements(PListFile, AppDirectory + "/Info.plist", Replacements, PListAdditionalLines);
+			CopyFileWithReplacements(PListFile, IntermediateDirectory + "/" + GameName + "-Info.plist", Replacements, PListAdditionalLines);
 
 			// ensure the destination is writable
 			if (File.Exists(AppDirectory + "/" + GameName))
@@ -239,7 +313,15 @@ namespace UnrealBuildTool.IOS
 			File.Copy(BinaryPath + "/" + GameExeName, AppDirectory + "/" + GameName, true);
 
 			// copy engine assets in
-			CopyFiles(InEngineDir + "/Build/IOS/Resources/Graphics", AppDirectory, "*.png", true);
+			if (bSkipDefaultPNGs)
+			{
+				// we still want default icons
+				CopyFiles(InEngineDir + "/Build/IOS/Resources/Graphics", AppDirectory, "Icon*.png", true);
+			}
+			else
+			{
+				CopyFiles(InEngineDir + "/Build/IOS/Resources/Graphics", AppDirectory, "*.png", true);
+			}
 			// merge game assets on top
 			if (Directory.Exists(BuildDirectory + "/Resources/Graphics"))
 			{
