- Removed some debug code

diff --git a/Engine/Source/Programs/UnrealBuildTool/IOS/UEDeployIOS.cs b/Engine/Source/Programs/UnrealBuildTool/IOS/UEDeployIOS.cs
--- a/Engine/Source/Programs/UnrealBuildTool/IOS/UEDeployIOS.cs
+++ b/Engine/Source/Programs/UnrealBuildTool/IOS/UEDeployIOS.cs
@@ -255,10 +255,6 @@ namespace UnrealBuildTool.IOS
 				LaunchXib = BuildDirectory + "/Resources/Interface/LaunchScreen.xib";
 			}
 
-			List<string> PListAdditionalLines = new List<string>();
-			PListAdditionalLines.Add("\t<key>TestKey</key>");
-			PListAdditionalLines.Add("\t<string>" + BuildDirectory + "</string>");
-
 			bool bSkipDefaultPNGs = false;
 			if (File.Exists(LaunchXib))
 			{
