UBT compile error fix

diff --git a/Engine/Source/Programs/UnrealBuildTool/IOS/UEDeployIOS.cs b/Engine/Source/Programs/UnrealBuildTool/IOS/UEDeployIOS.cs
--- a/Engine/Source/Programs/UnrealBuildTool/IOS/UEDeployIOS.cs
+++ b/Engine/Source/Programs/UnrealBuildTool/IOS/UEDeployIOS.cs
@@ -255,6 +255,8 @@ namespace UnrealBuildTool.IOS
 				LaunchXib = BuildDirectory + "/Resources/Interface/LaunchScreen.xib";
 			}
 
+			List<string> PListAdditionalLines = new List<string>();
+
 			bool bSkipDefaultPNGs = false;
 			if (File.Exists(LaunchXib))
 			{
