Hey, we've never met before, but I have been told I'm build watcher today, so I was watching the build and saw this.

This was causing a warning on CIS. The warning was...

HTML5/HTML5ToolChain.cs(313,11): warning CS0219: The variable `BaseSDKPath' is assigned but its value is never used

It is a one line fix, and even though I am 99.999% this won't cause any problems, I don't have a great way to compile this, so I'm code reviewing you. :-)

diff --git a/Engine/Source/Programs/UnrealBuildTool/HTML5/HTML5ToolChain.cs b/Engine/Source/Programs/UnrealBuildTool/HTML5/HTML5ToolChain.cs
--- a/Engine/Source/Programs/UnrealBuildTool/HTML5/HTML5ToolChain.cs
+++ b/Engine/Source/Programs/UnrealBuildTool/HTML5/HTML5ToolChain.cs
@@ -310,7 +310,6 @@ namespace UnrealBuildTool
 			}
 
 			string Arguments = GetCLArguments_Global(CompileEnvironment);
-			string BaseSDKPath = Environment.GetEnvironmentVariable("EMSCRIPTEN");
 
 			CPPOutput Result = new CPPOutput();
 
