A workaround for iOS signing problem in Xcode 6

diff --git a/Engine/Source/Programs/UnrealBuildTool/IOS/UEDeployIOS.cs b/Engine/Source/Programs/UnrealBuildTool/IOS/UEDeployIOS.cs
--- a/Engine/Source/Programs/UnrealBuildTool/IOS/UEDeployIOS.cs
+++ b/Engine/Source/Programs/UnrealBuildTool/IOS/UEDeployIOS.cs
@@ -297,12 +297,21 @@ namespace UnrealBuildTool.IOS
 			}
 
 			// plist replacements
+			bool bCreatingNewPlistFile = !File.Exists(IntermediateDirectory + "/" + GameName + "-Info.plist");
 			Dictionary<string, string> Replacements = new Dictionary<string, string>();
 			Replacements.Add("${EXECUTABLE_NAME}", GameName);
 			Replacements.Add("${BUNDLE_IDENTIFIER}", InProjectName.Replace("_", ""));
 			CopyFileWithReplacements(PListFile, AppDirectory + "/Info.plist", Replacements, PListAdditionalLines);
 			CopyFileWithReplacements(PListFile, IntermediateDirectory + "/" + GameName + "-Info.plist", Replacements, PListAdditionalLines);
 
+			if (bCreatingNewPlistFile)
+			{
+				// @todo: It seems that Xcode 6 doesn't check for the existence of Info.plist file often enough to know we created one right before it's trying to sign the app bundle.
+				// To give it time, we wait a few seconds. This needs a proper fix which is to not sign using a target in the project opened in Xcode, but rather create a separate
+				// deploy only project and use that (similar to what UAT does).
+				System.Threading.Thread.Sleep(8000);
+			}
+
 			// ensure the destination is writable
 			if (File.Exists(AppDirectory + "/" + GameName))
 			{
