Adding support for programs at the UE4 root level.

diff --git a/Engine/Source/Programs/UnrealBuildTool/Configuration/UEBuildTarget.cs b/Engine/Source/Programs/UnrealBuildTool/Configuration/UEBuildTarget.cs
--- a/Engine/Source/Programs/UnrealBuildTool/Configuration/UEBuildTarget.cs
+++ b/Engine/Source/Programs/UnrealBuildTool/Configuration/UEBuildTarget.cs
@@ -2228,7 +2228,11 @@ namespace UnrealBuildTool
 
 			// Build base directory string ("../Binaries/<Platform>/")
 			string BinariesDirName;
-			if( PluginInfo != null )
+			if(TargetType.HasValue && TargetType.Value == TargetRules.TargetType.Program && UnrealBuildTool.HasUProjectFile() && !ProjectFileGenerator.bGenerateProjectFiles)
+			{
+				BinariesDirName = Path.Combine( UnrealBuildTool.GetUProjectPath(), "Binaries" );
+			}
+			else if( PluginInfo != null )
 			{
 				BinariesDirName = Path.Combine( PluginInfo.Directory, "Binaries" );
 			}
diff --git a/Engine/Source/Programs/UnrealBuildTool/System/VCProject.cs b/Engine/Source/Programs/UnrealBuildTool/System/VCProject.cs
--- a/Engine/Source/Programs/UnrealBuildTool/System/VCProject.cs
+++ b/Engine/Source/Programs/UnrealBuildTool/System/VCProject.cs
@@ -942,7 +942,7 @@ namespace UnrealBuildTool
 
 					// Get the output directory
 					string RootDirectory = Path.GetFullPath(ProjectFileGenerator.EngineRelativePath);
-					if ((TargetRules.IsAGame(TargetRulesObject.Type) || TargetRulesObject.Type == TargetRules.TargetType.Server) && bShouldCompileMonolithic && !TargetRulesObject.bOutputToEngineBinaries)
+                    if ((TargetRules.IsAGame(TargetRulesObject.Type) || TargetRulesObject.Type == TargetRules.TargetType.Server) && bShouldCompileMonolithic && !TargetRulesObject.bOutputToEngineBinaries)
 					{
 						if (UnrealBuildTool.HasUProjectFile() && Utils.IsFileUnderDirectory(TargetFilePath, UnrealBuildTool.GetUProjectPath()))
 						{
@@ -958,6 +958,15 @@ namespace UnrealBuildTool
 						}
 					}
 
+					if(TargetRulesObject.Type == TargetRules.TargetType.Program && !TargetRulesObject.bOutputToEngineBinaries)
+					{
+						string UnrealProjectPath = UProjectInfo.GetProjectForTarget(TargetName);
+						if (!String.IsNullOrEmpty(UnrealProjectPath))
+						{
+							RootDirectory = Path.GetDirectoryName(Path.GetFullPath(UnrealProjectPath));
+						}
+					}
+
 					// Get the output directory
 					string OutputDirectory = Path.Combine(RootDirectory, "Binaries", UBTPlatformName);
 
