Fixed issue with code signing in Xcode 6. Generating project files will now create all necessary intermediate Info.plist files for iOS targets so Xcode has the info it needs immediately after opening a project. Also fixed paths to Info.plist files for clients and programs.

diff --git a/Engine/Source/Programs/UnrealBuildTool/IOS/UEDeployIOS.cs b/Engine/Source/Programs/UnrealBuildTool/IOS/UEDeployIOS.cs
--- a/Engine/Source/Programs/UnrealBuildTool/IOS/UEDeployIOS.cs
+++ b/Engine/Source/Programs/UnrealBuildTool/IOS/UEDeployIOS.cs
@@ -297,21 +297,12 @@ namespace UnrealBuildTool.IOS
 			}
 
 			// plist replacements
-			bool bCreatingNewPlistFile = !File.Exists(IntermediateDirectory + "/" + GameName + "-Info.plist");
 			Dictionary<string, string> Replacements = new Dictionary<string, string>();
 			Replacements.Add("${EXECUTABLE_NAME}", GameName);
 			Replacements.Add("${BUNDLE_IDENTIFIER}", InProjectName.Replace("_", ""));
 			CopyFileWithReplacements(PListFile, AppDirectory + "/Info.plist", Replacements, PListAdditionalLines);
 			CopyFileWithReplacements(PListFile, IntermediateDirectory + "/" + GameName + "-Info.plist", Replacements, PListAdditionalLines);
 
-			if (bCreatingNewPlistFile)
-			{
-				// @todo: It seems that Xcode 6 doesn't check for the existence of Info.plist file often enough to know we created one right before it's trying to sign the app bundle.
-				// To give it time, we wait a few seconds. This needs a proper fix which is to not sign using a target in the project opened in Xcode, but rather create a separate
-				// deploy only project and use that (similar to what UAT does).
-				System.Threading.Thread.Sleep(8000);
-			}
-
 			// ensure the destination is writable
 			if (File.Exists(AppDirectory + "/" + GameName))
 			{
diff --git a/Engine/Source/Programs/UnrealBuildTool/System/XcodeProjectFileGenerator.cs b/Engine/Source/Programs/UnrealBuildTool/System/XcodeProjectFileGenerator.cs
--- a/Engine/Source/Programs/UnrealBuildTool/System/XcodeProjectFileGenerator.cs
+++ b/Engine/Source/Programs/UnrealBuildTool/System/XcodeProjectFileGenerator.cs
@@ -538,38 +538,50 @@ namespace UnrealBuildTool
 				"\t\t\t\tIPHONEOS_DEPLOYMENT_TARGET = " + Toolchain.GetPlatformVersion() + ";" + ProjectFileGenerator.NewLine +
 				"\t\t\t\t\"CODE_SIGN_IDENTITY[sdk=iphoneos*]\" = \"iPhone Developer\";" + ProjectFileGenerator.NewLine);
 
+			string InfoPlistPath = "";
 			if (bIsUE4Game)
 			{
+				InfoPlistPath = EngineRelative + "Engine/Intermediate/IOS/" + TargetName + "-Info.plist";
 				Contents.Append(
 					"\t\t\t\tCODE_SIGN_RESOURCE_RULES_PATH = \"" + EngineRelative + "Engine/Build/iOS/XcodeSupportFiles/CustomResourceRules.plist\";" + ProjectFileGenerator.NewLine +
-					"\t\t\t\tINFOPLIST_FILE = \"" + EngineRelative + "Engine/Intermediate/IOS/" + TargetName + "-Info.plist\";" + ProjectFileGenerator.NewLine +
+					"\t\t\t\tINFOPLIST_FILE = \"" + InfoPlistPath + "\";" + ProjectFileGenerator.NewLine +
 					"\t\t\t\tSYMROOT = \"" + EngineRelative + "Engine/Binaries/IOS/Payload\";" + ProjectFileGenerator.NewLine +
 					"\t\t\t\tOBJROOT = \"" + EngineRelative + "Engine/Intermediate/IOS/build\";" + ProjectFileGenerator.NewLine +
 					"\t\t\t\tCONFIGURATION_BUILD_DIR = \"" + EngineRelative + "Engine/Binaries/IOS/Payload\";" + ProjectFileGenerator.NewLine);
 			}
 			else if (bIsUE4Client)
 			{
+				InfoPlistPath = EngineRelative + "Engine/Intermediate/IOS/UE4Game-Info.plist";
 				Contents.Append(
 					"\t\t\t\tCODE_SIGN_RESOURCE_RULES_PATH = \"" + EngineRelative + "Engine/Build/iOS/XcodeSupportFiles/CustomResourceRules.plist\";" + ProjectFileGenerator.NewLine +
-					"\t\t\t\tINFOPLIST_FILE = \"" + EngineRelative + "Engine/Intermediate/IOS/UE4Game-Info.plist\";" + ProjectFileGenerator.NewLine +
+					"\t\t\t\tINFOPLIST_FILE = \"" + InfoPlistPath + "\";" + ProjectFileGenerator.NewLine +
 					"\t\t\t\tSYMROOT = \"" + EngineRelative + "Engine/Binaries/IOS/Payload\";" + ProjectFileGenerator.NewLine +
 					"\t\t\t\tOBJROOT = \"" + EngineRelative + "Engine/Intermediate/IOS/build\";" + ProjectFileGenerator.NewLine +
 					"\t\t\t\tCONFIGURATION_BUILD_DIR = \"" + EngineRelative + "Engine/Binaries/IOS/Payload\";" + ProjectFileGenerator.NewLine);
 			}
 			else if (IsAGame)
 			{
+				InfoPlistPath = GamePath + "/Intermediate/IOS/" + TargetName + "-Info.plist";
 				Contents.Append(
 					"\t\t\t\tCODE_SIGN_RESOURCE_RULES_PATH = \"" + EngineRelative + "Engine/Build/iOS/XcodeSupportFiles/CustomResourceRules.plist\";" + ProjectFileGenerator.NewLine +
-					"\t\t\t\tINFOPLIST_FILE = \"" + GamePath + "/Intermediate/IOS/" + TargetName + "-Info.plist\";" + ProjectFileGenerator.NewLine +
+					"\t\t\t\tINFOPLIST_FILE = \"" + InfoPlistPath + "\";" + ProjectFileGenerator.NewLine +
 					"\t\t\t\tSYMROOT = \"" + GamePath + "/Binaries/IOS\";" + ProjectFileGenerator.NewLine +
 					"\t\t\t\tOBJROOT = \"" + GamePath + "/Intermediate/IOS/build\";" + ProjectFileGenerator.NewLine +
 					"\t\t\t\tCONFIGURATION_BUILD_DIR = \"" + GamePath + "/Binaries/IOS/Payload\";" + ProjectFileGenerator.NewLine);
 			}
 			else
 			{
+				if (string.IsNullOrEmpty(GamePath))
+				{
+					InfoPlistPath = EngineRelative + "Engine/Intermediate/IOS/" + TargetName + "-Info.plist";
+				}
+				else
+				{
+					InfoPlistPath = GamePath + "/Intermediate/IOS/" + TargetName + "-Info.plist";
+				}
 				Contents.Append(
 					"\t\t\t\tCODE_SIGN_RESOURCE_RULES_PATH = \"" + EngineRelative + "Engine/Build/iOS/XcodeSupportFiles/CustomResourceRules.plist\";" + ProjectFileGenerator.NewLine +
-					"\t\t\t\tINFOPLIST_FILE = \"" + EngineRelative + TargetName + "/Intermediate/IOS/" + TargetName + "-Info.plist\";" + ProjectFileGenerator.NewLine +
+					"\t\t\t\tINFOPLIST_FILE = \"" + InfoPlistPath + "\";" + ProjectFileGenerator.NewLine +
 					"\t\t\t\tSYMROOT = \"" + EngineRelative + "Engine/Binaries/IOS/Payload\";" + ProjectFileGenerator.NewLine +
 					"\t\t\t\tOBJROOT = \"" + EngineRelative + "Engine/Intermediate/IOS/build\";" + ProjectFileGenerator.NewLine +
 					"\t\t\t\tCONFIGURATION_BUILD_DIR = \"" + EngineRelative + "Engine/Binaries/IOS/Payload\";" + ProjectFileGenerator.NewLine);
@@ -588,6 +600,19 @@ namespace UnrealBuildTool
 				"\t\t\t};" + ProjectFileGenerator.NewLine +
 				"\t\t\tname = " + ConfigName + ";" + ProjectFileGenerator.NewLine +
 				"\t\t};" + ProjectFileGenerator.NewLine);
+
+			// Prepare a temp Info.plist file so Xcode has some basic info about the target immediately after opening the project.
+			// This is needed for the target to pass the settings validation before code signing. UBT will overwrite this plist file later, with proper contents.
+			if (File.Exists(InfoPlistPath))
+			{
+				File.Delete(InfoPlistPath);
+			}
+			else
+			{
+				Directory.CreateDirectory(Path.GetDirectoryName(InfoPlistPath));
+			}
+			string InfoPlistContents = File.ReadAllText(EngineRelative + "Engine/Build/IOS/UE4Game-Info.plist").Replace("${BUNDLE_IDENTIFIER}", TargetName).Replace("${EXECUTABLE_NAME}", TargetName);
+			File.WriteAllText(InfoPlistPath, InfoPlistContents);
 		}
 
 		private void AppendIOSXCTestConfig(ref StringBuilder Contents, string ConfigName, string ConfigGuid, string TargetName, string EngineSubdir, string EngineRelative)
@@ -689,6 +714,7 @@ namespace UnrealBuildTool
 
 			bool IsAGame = false;
 			string GamePath = null;
+			string GameNameFromClientName = Target.TargetName.EndsWith("Client") ? Target.TargetName.Remove(Target.TargetName.LastIndexOf("Client")) + "Game" : null;
 			bool bIsUE4Game = Target.TargetName.Equals("UE4Game", StringComparison.InvariantCultureIgnoreCase);
 			bool bIsUE4Client = Target.TargetName.Equals("UE4Client", StringComparison.InvariantCultureIgnoreCase);
 			string EngineRelative = "";
@@ -702,6 +728,11 @@ namespace UnrealBuildTool
 					GamePath = Toolchain.ConvertPath(Path.GetFullPath(GameFolder));
 					break;
 				}
+				else if (!string.IsNullOrEmpty(GameNameFromClientName) && GameFolder.EndsWith(GameNameFromClientName))
+				{
+					GamePath = Toolchain.ConvertPath(Path.GetFullPath(GameFolder));
+					break;
+				}
 			}
 
 			EngineRelative = Path.GetFullPath(EngineRelativePath + "/../");
