UBT: Proper dependencies for modular builds + fix Linux cross-build.

diff --git a/Engine/Source/Programs/UnrealBuildTool/Configuration/UEBuildTarget.cs b/Engine/Source/Programs/UnrealBuildTool/Configuration/UEBuildTarget.cs
--- a/Engine/Source/Programs/UnrealBuildTool/Configuration/UEBuildTarget.cs
+++ b/Engine/Source/Programs/UnrealBuildTool/Configuration/UEBuildTarget.cs
@@ -1399,8 +1399,8 @@ namespace UnrealBuildTool
 				}
 			}
 
-			// On Mac we have actions that should be executed after all the binaries are created
-			if (GlobalLinkEnvironment.Config.Target.Platform == CPPTargetPlatform.Mac)
+			// On Mac and Linux we have actions that should be executed after all the binaries are created
+			if (GlobalLinkEnvironment.Config.Target.Platform == CPPTargetPlatform.Mac || GlobalLinkEnvironment.Config.Target.Platform == CPPTargetPlatform.Linux)
 			{
 				TargetToolChain.SetupBundleDependencies(AppBinaries, GameName);
 			}
diff --git a/Engine/Source/Programs/UnrealBuildTool/Linux/LinuxToolChain.cs b/Engine/Source/Programs/UnrealBuildTool/Linux/LinuxToolChain.cs
--- a/Engine/Source/Programs/UnrealBuildTool/Linux/LinuxToolChain.cs
+++ b/Engine/Source/Programs/UnrealBuildTool/Linux/LinuxToolChain.cs
@@ -931,7 +931,8 @@ namespace UnrealBuildTool
 
                     if (bUseCmdExe)
                     {
-					    Writer.Write("rem Automatically generated by UnrealBuildTool\n");
+						Writer.Write("@echo off\n");
+						Writer.Write("rem Automatically generated by UnrealBuildTool\n");
 					    Writer.Write("rem *DO NOT EDIT*\n\n");
                     }
                     else
@@ -961,7 +962,8 @@ namespace UnrealBuildTool
 				string Replace = "-Wl,--allow-shlib-undefined";
 					
 				FixDepsLine = FixDepsLine.Replace(Replace, EngineAndGameLibrariesString);
-				FixDepsLine = FixDepsLine.Replace(OutputFile.AbsolutePath, OutputFile.AbsolutePath + ".fixed");
+				string OutputFileForwardSlashes = OutputFile.AbsolutePath.Replace("\\", "/");
+				FixDepsLine = FixDepsLine.Replace(OutputFileForwardSlashes, OutputFileForwardSlashes + ".fixed");
 				FixDepsLine = FixDepsLine.Replace("$", "\\$");
 				FixDepsScript.Write(FixDepsLine + "\n");
                 if (bUseCmdExe)
