Enable HTML5 1.22.00 SDK

diff --git a/Engine/Source/Programs/UnrealBuildTool/HTML5/UEBuildHTML5.cs b/Engine/Source/Programs/UnrealBuildTool/HTML5/UEBuildHTML5.cs
--- a/Engine/Source/Programs/UnrealBuildTool/HTML5/UEBuildHTML5.cs
+++ b/Engine/Source/Programs/UnrealBuildTool/HTML5/UEBuildHTML5.cs
@@ -24,7 +24,7 @@ namespace UnrealBuildTool
 		// May include minor revisions or descriptions that a default install from SDK_Manager won't have.
 		// e.g. 1.600_Patch001, or 1.610.001.  The SDK_Manager always installs minor revision patches straight into
 		// the default major revision folder.
-		static private string ExpectedSDKVersion = "1.21.0";
+		static private string ExpectedSDKVersion = "1.22.0";
 
 		/** 
 		 * Whether platform supports switching SDKs during runtime
