fix for plist trying to be generated from PC
#ue4

diff --git a/Engine/Source/Programs/UnrealBuildTool/System/XcodeProjectFileGenerator.cs b/Engine/Source/Programs/UnrealBuildTool/System/XcodeProjectFileGenerator.cs
--- a/Engine/Source/Programs/UnrealBuildTool/System/XcodeProjectFileGenerator.cs
+++ b/Engine/Source/Programs/UnrealBuildTool/System/XcodeProjectFileGenerator.cs
@@ -603,16 +603,19 @@ namespace UnrealBuildTool
 
 			// Prepare a temp Info.plist file so Xcode has some basic info about the target immediately after opening the project.
 			// This is needed for the target to pass the settings validation before code signing. UBT will overwrite this plist file later, with proper contents.
-			if (File.Exists(InfoPlistPath))
+			if (BuildHostPlatform.Current.Platform == UnrealTargetPlatform.Mac)
 			{
-				File.Delete(InfoPlistPath);
-			}
-			else
-			{
-				Directory.CreateDirectory(Path.GetDirectoryName(InfoPlistPath));
+				if (File.Exists(InfoPlistPath))
+				{
+					File.Delete(InfoPlistPath);
+				}
+				else
+				{
+					Directory.CreateDirectory(Path.GetDirectoryName(InfoPlistPath));
+				}
+				string InfoPlistContents = File.ReadAllText(EngineRelative + "Engine/Build/IOS/UE4Game-Info.plist").Replace("${BUNDLE_IDENTIFIER}", TargetName).Replace("${EXECUTABLE_NAME}", TargetName);
+				File.WriteAllText(InfoPlistPath, InfoPlistContents);
 			}
-			string InfoPlistContents = File.ReadAllText(EngineRelative + "Engine/Build/IOS/UE4Game-Info.plist").Replace("${BUNDLE_IDENTIFIER}", TargetName).Replace("${EXECUTABLE_NAME}", TargetName);
-			File.WriteAllText(InfoPlistPath, InfoPlistContents);
 		}
 
 		private void AppendIOSXCTestConfig(ref StringBuilder Contents, string ConfigName, string ConfigGuid, string TargetName, string EngineSubdir, string EngineRelative)
