- Big iteration speedup for Android - don't remake the .apk if not needed!

diff --git a/Engine/Source/Programs/UnrealBuildTool/Android/UEDeployAndroid.cs b/Engine/Source/Programs/UnrealBuildTool/Android/UEDeployAndroid.cs
--- a/Engine/Source/Programs/UnrealBuildTool/Android/UEDeployAndroid.cs
+++ b/Engine/Source/Programs/UnrealBuildTool/Android/UEDeployAndroid.cs
@@ -345,11 +345,16 @@ namespace UnrealBuildTool.Android
 				DateTime ApkTime = File.GetLastWriteTimeUtc(DestApkName);
 				foreach (var InputFileName in InputFiles)
 				{
+					// skip .log files
+					if (Path.GetExtension(InputFileName) == ".log")
+					{
+						continue;
+					}
 					DateTime InputFileTime = File.GetLastWriteTimeUtc(InputFileName);
 					if (InputFileTime.CompareTo(ApkTime) > 0)
 					{
-						// could break here
 						bAllInputsCurrent = false;
+						Log.TraceInformation("{0} is out of date due to newer input file {1}", DestApkName, InputFileName);
 						break;
 					}
 				}
