- Added Engine.ini as dependencies to .apk genertion checking

diff --git a/Engine/Source/Programs/UnrealBuildTool/Android/UEDeployAndroid.cs b/Engine/Source/Programs/UnrealBuildTool/Android/UEDeployAndroid.cs
--- a/Engine/Source/Programs/UnrealBuildTool/Android/UEDeployAndroid.cs
+++ b/Engine/Source/Programs/UnrealBuildTool/Android/UEDeployAndroid.cs
@@ -353,21 +353,29 @@ namespace UnrealBuildTool.Android
 						InputFiles.AddRange(Directory.EnumerateFiles(GameBuildFilesPath, "*.*", SearchOption.AllDirectories));
 					}
 
+					// rebuild if .ini files change
+					// @todo android: programmatically determine if any .ini setting changed?
+					InputFiles.Add(Path.Combine(EngineDirectory, "Config\\BaseEngine.ini"));
+					InputFiles.Add(Path.Combine(ProjectDirectory, "Config\\DefaultEngine.ini"));
+
 					// look for any newer input file
 					DateTime ApkTime = File.GetLastWriteTimeUtc(DestApkName);
 					foreach (var InputFileName in InputFiles)
 					{
-						// skip .log files
-						if (Path.GetExtension(InputFileName) == ".log")
-						{
-							continue;
-						}
-						DateTime InputFileTime = File.GetLastWriteTimeUtc(InputFileName);
-						if (InputFileTime.CompareTo(ApkTime) > 0)
+						if (File.Exists(InputFileName))
 						{
-							bAllInputsCurrent = false;
-							Log.TraceInformation("{0} is out of date due to newer input file {1}", DestApkName, InputFileName);
-							break;	
+							// skip .log files
+							if (Path.GetExtension(InputFileName) == ".log")
+							{
+								continue;
+							}
+							DateTime InputFileTime = File.GetLastWriteTimeUtc(InputFileName);
+							if (InputFileTime.CompareTo(ApkTime) > 0)
+							{
+								bAllInputsCurrent = false;
+								Log.TraceInformation("{0} is out of date due to newer input file {1}", DestApkName, InputFileName);
+								break;	
+							}
 						}
 					}
 				}
