- Fixed a problem with fat binary linking after the GPU architecture fat binary changes

diff --git a/Engine/Source/Programs/UnrealBuildTool/Android/AndroidToolChain.cs b/Engine/Source/Programs/UnrealBuildTool/Android/AndroidToolChain.cs
--- a/Engine/Source/Programs/UnrealBuildTool/Android/AndroidToolChain.cs
+++ b/Engine/Source/Programs/UnrealBuildTool/Android/AndroidToolChain.cs
@@ -13,9 +13,12 @@ namespace UnrealBuildTool
 	{
 		static private bool bHasNDKExtensionsCompiled = false;
 
-		// public so that AndroidPlatform can get to it
+		// the list of architectures we will compile for
 		static private string[] Arches = null;
+		// the list of GPU architectures we will compile for
 		static private string[] GPUArchitectures = null;
+		// a list of all architecture+GPUArchitecture names (-armv7-es2, etc)
+		static private string[] AllComboNames = null;
 
 		static private Dictionary<string, string[]> AllArchNames = new Dictionary<string, string[]> {
 			{ "-armv7", new string[] { "armv7", "armeabi-v7a", } }, 
@@ -77,6 +80,16 @@ namespace UnrealBuildTool
 				ProjectGPUArches.Add("-es2");
 			}
 			GPUArchitectures = ProjectGPUArches.ToArray();
+
+			List<string> FullArchCombinations = new List<string>();
+			foreach (string Arch in Arches)
+			{
+				foreach (string GPUArch in GPUArchitectures)
+				{
+					FullArchCombinations.Add(Arch + GPUArch);
+				}
+			}
+			AllComboNames = FullArchCombinations.ToArray();
 		}
 
 		public override void SetUpGlobalEnvironment()
@@ -494,7 +507,7 @@ namespace UnrealBuildTool
 			return true;
 		}
 
-		static bool ShouldSkipLib(string Lib, string Arch)
+		static bool ShouldSkipLib(string Lib, string Arch, string GPUArchitecture)
 		{
 			// reject any libs we outright don't want to link with
 			foreach (var LibName in LibrariesToSkip[Arch])
@@ -506,11 +519,11 @@ namespace UnrealBuildTool
 			}
 
 			// if another architecture is in the filename, reject it
-			foreach (string ArchName in Arches)
+			foreach (string ComboName in AllComboNames)
 			{
-				if (ArchName != Arch)
+				if (ComboName != Arch + GPUArchitecture)
 				{
-					if (Path.GetFileNameWithoutExtension(Lib).EndsWith(ArchName))
+					if (Path.GetFileNameWithoutExtension(Lib).EndsWith(ComboName))
 					{
 						return true;
 					}
@@ -859,7 +872,7 @@ namespace UnrealBuildTool
 						LinkAction.CommandArguments += string.Format(" -Wl,--start-group");
 						foreach (string AdditionalLibrary in LinkEnvironment.Config.AdditionalLibraries)
 						{
-							if (!ShouldSkipLib(AdditionalLibrary, Arch))
+							if (!ShouldSkipLib(AdditionalLibrary, Arch, GPUArchitecture))
 							{
 								if (String.IsNullOrEmpty(Path.GetDirectoryName(AdditionalLibrary)))
 								{
