- No longer compile Engine.generated.1.cpp and .2.cpp together, as the large resulting code is crashing x86 android compiler.
- Fixed a couple issues with Android fat binaries (cleaning and obb being archived multiple times)
- Touched Engine.h to force a UHT refresh

diff --git a/Engine/Source/Programs/AutomationTool/Android/AndroidPlatform.Automation.cs b/Engine/Source/Programs/AutomationTool/Android/AndroidPlatform.Automation.cs
--- a/Engine/Source/Programs/AutomationTool/Android/AndroidPlatform.Automation.cs
+++ b/Engine/Source/Programs/AutomationTool/Android/AndroidPlatform.Automation.cs
@@ -195,6 +195,7 @@ public class AndroidPlatform : Platform
 		string[] GPUArchitectures = UnrealBuildTool.AndroidToolChain.GetAllGPUArchitectures();
 		bool bMakeSeparateApks = UnrealBuildTool.Android.UEDeployAndroid.ShouldMakeSeparateApks();
 
+		bool bAddedOBB = false;
 		foreach (string Architecture in Architectures)
 		{
 			foreach (string GPUArchitecture in GPUArchitectures)
@@ -218,8 +219,9 @@ public class AndroidPlatform : Platform
             }
 
 			SC.ArchiveFiles(Path.GetDirectoryName(ApkName), Path.GetFileName(ApkName));
-			if (!Params.OBBinAPK)
+			if (!Params.OBBinAPK && !bAddedOBB)
 			{
+				bAddedOBB = true;
 				SC.ArchiveFiles(Path.GetDirectoryName(ObbName), Path.GetFileName(ObbName));
 			}
 
diff --git a/Engine/Source/Programs/UnrealBuildTool/Configuration/UEBuildModule.cs b/Engine/Source/Programs/UnrealBuildTool/Configuration/UEBuildModule.cs
--- a/Engine/Source/Programs/UnrealBuildTool/Configuration/UEBuildModule.cs
+++ b/Engine/Source/Programs/UnrealBuildTool/Configuration/UEBuildModule.cs
@@ -883,14 +883,14 @@ namespace UnrealBuildTool
 		{
 			public class BuildInfoClass
 			{
-				/** The filename of the *.generated.cpp file which was generated for the module */
-				public readonly string Filename;
+				/** The wildcard of the *.generated.cpp file which was generated for the module */
+				public readonly string FileWildcard;
 
-				public BuildInfoClass(string InFilename)
+				public BuildInfoClass(string InWildcard)
 				{
-					Debug.Assert(InFilename != null);
+					Debug.Assert(InWildcard != null);
 
-					Filename = InFilename;
+					FileWildcard = InWildcard;
 				}
 			}
 
@@ -1471,12 +1471,16 @@ namespace UnrealBuildTool
 
 			if (AutoGenerateCppInfo != null && AutoGenerateCppInfo.BuildInfo != null && !CPPCompileEnvironment.bHackHeaderGenerator)
 			{
-				var GeneratedCppFileItem = FileItem.GetItemByPath( AutoGenerateCppInfo.BuildInfo.Filename );
+				string[] GeneratedFiles = Directory.GetFiles(Path.GetDirectoryName(AutoGenerateCppInfo.BuildInfo.FileWildcard), Path.GetFileName(AutoGenerateCppInfo.BuildInfo.FileWildcard));
+				foreach (string GeneratedFilename in GeneratedFiles)
+				{
+					var GeneratedCppFileItem = FileItem.GetItemByPath(GeneratedFilename);
 
-				CachePCHUsageForModuleSourceFile( this.Target, CPPCompileEnvironment, GeneratedCppFileItem );
+					CachePCHUsageForModuleSourceFile(this.Target, CPPCompileEnvironment, GeneratedCppFileItem);
 
-				// @todo fastubt: Check for ALL other places where we might be injecting .cpp or .rc files for compiling without caching CachedCPPIncludeInfo first (anything platform specfic?)
-				LinkInputFiles.AddRange( CPPCompileEnvironment.CompileFiles( Target, new List<FileItem>{ GeneratedCppFileItem }, Name ).ObjectFiles );
+					// @todo fastubt: Check for ALL other places where we might be injecting .cpp or .rc files for compiling without caching CachedCPPIncludeInfo first (anything platform specfic?)
+					LinkInputFiles.AddRange(CPPCompileEnvironment.CompileFiles(Target, new List<FileItem> { GeneratedCppFileItem }, Name).ObjectFiles);
+				}
 			}
 
 			// Compile C files directly.
diff --git a/Engine/Source/Programs/UnrealBuildTool/Configuration/UEBuildTarget.cs b/Engine/Source/Programs/UnrealBuildTool/Configuration/UEBuildTarget.cs
--- a/Engine/Source/Programs/UnrealBuildTool/Configuration/UEBuildTarget.cs
+++ b/Engine/Source/Programs/UnrealBuildTool/Configuration/UEBuildTarget.cs
@@ -1001,7 +1001,7 @@ namespace UnrealBuildTool
 					if (AppBinCPP != null)
 					{
 						// Collect all modules used by this binary.
-						Log.TraceVerbose("\tProcessing AppBinary " + AppBin.Config.OutputFilePath);
+						Log.TraceVerbose("\tProcessing AppBinary " + AppBin.Config.OutputFilePaths[0]);
 						foreach (string ModuleName in AppBinCPP.ModuleNames)
 						{					
 							if (ModuleList.Contains(ModuleName) == false)
@@ -1017,7 +1017,7 @@ namespace UnrealBuildTool
 					}
 					else
 					{
-						Log.TraceVerbose("\t********* Skipping " + AppBin.Config.OutputFilePath);
+						Log.TraceVerbose("\t********* Skipping " + AppBin.Config.OutputFilePaths[0]);
 					}
 				}
 
@@ -1454,7 +1454,7 @@ namespace UnrealBuildTool
 							UHTModuleInfo.GeneratedCPPFilenameBase = Path.Combine( UEBuildModuleCPP.GetGeneratedCodeDirectoryForModule(this, UHTModuleInfo.ModuleDirectory, UHTModuleInfo.ModuleName), UHTModuleInfo.ModuleName ) + ".generated";
 							if (DependencyModuleCPP.SourceFilesToBuild.Count != 0)
 							{
-								BuildInfo = new UEBuildModuleCPP.AutoGenerateCppInfoClass.BuildInfoClass( UHTModuleInfo.GeneratedCPPFilenameBase + ".cpp" );
+								BuildInfo = new UEBuildModuleCPP.AutoGenerateCppInfoClass.BuildInfoClass( UHTModuleInfo.GeneratedCPPFilenameBase + "*.cpp" );
 							}
 
 							DependencyModuleCPP.AutoGenerateCppInfo = new UEBuildModuleCPP.AutoGenerateCppInfoClass(BuildInfo);
diff --git a/Engine/Source/Programs/UnrealHeaderTool/Private/CodeGenerator.cpp b/Engine/Source/Programs/UnrealHeaderTool/Private/CodeGenerator.cpp
--- a/Engine/Source/Programs/UnrealHeaderTool/Private/CodeGenerator.cpp
+++ b/Engine/Source/Programs/UnrealHeaderTool/Private/CodeGenerator.cpp
@@ -4376,7 +4376,7 @@ void FNativeClassHeaderGenerator::ExportGeneratedCPP()
 		}
 
 		FString CppPath = ModuleInfo->GeneratedCPPFilenameBase + (GeneratedFunctionBodyTextSplit.Num() > 1 ? *FString::Printf(TEXT(".%d.cpp"), FileIdx + 1) : TEXT(".cpp"));
-		SaveHeaderIfChanged(*CppPath, *(GeneratedCPPPreamble + ModulePCHInclude + GeneratedCPPClassesIncludes + ((GeneratedFunctionBodyTextSplit.Num() > 1) ? FString() : GeneratedLinkerFixupFunction) + FileText + GeneratedCPPEpilogue));
+		SaveHeaderIfChanged(*CppPath, *(GeneratedCPPPreamble + ModulePCHInclude + GeneratedCPPClassesIncludes + ((FileIdx > 0) ? FString() : GeneratedLinkerFixupFunction) + FileText + GeneratedCPPEpilogue));
 
 		if (GeneratedFunctionBodyTextSplit.Num() > 1)
 		{
@@ -4384,18 +4384,14 @@ void FNativeClassHeaderGenerator::ExportGeneratedCPP()
 		}
 	}
 
-	if (GeneratedFunctionBodyTextSplit.Num() > 1)
-	{
-		FStringOutputDevice FileText;
-
-		for (int32 i=0; i < NumberedHeaderNames.Num(); ++i)
-		{
-			FileText.Logf(TEXT("#include \"%s\"") LINE_TERMINATOR, *NumberedHeaderNames[i]);
-		}
-
-		FString CppPath = ModuleInfo->GeneratedCPPFilenameBase + TEXT(".cpp");
-		SaveHeaderIfChanged(*CppPath, *(GeneratedCPPPreamble + ModulePCHInclude + GeneratedLinkerFixupFunction + FileText + GeneratedCPPEpilogue));
-	}
+	// delete the old .cpp file that will cause link errors if it's left around (Engine.generated.cpp and Engine.generated.1.cpp will 
+	// conflict now that we no longer use Engine.generated.cpp to #include Engine.generated.1.cpp, and UBT would compile all 3)
+	// @todo: This is a temp measure so we don't force everyone to require a Clean
+ 	if (GeneratedFunctionBodyTextSplit.Num() > 1)
+ 	{
+ 		FString CppPath = ModuleInfo->GeneratedCPPFilenameBase + TEXT(".cpp");
+		IFileManager::Get().Delete(*CppPath);
+ 	}
 
 	FString InlPath = ModuleInfo->GeneratedCPPFilenameBase + TEXT(".inl");
 	SaveHeaderIfChanged(*InlPath, *(GeneratedCPPPreamble + GeneratedINLText));
