Add dependency on PAK file for OBBinAPK case of MakeAPK
#ue4
#android

diff --git a/Engine/Source/Programs/UnrealBuildTool/Android/UEDeployAndroid.cs b/Engine/Source/Programs/UnrealBuildTool/Android/UEDeployAndroid.cs
--- a/Engine/Source/Programs/UnrealBuildTool/Android/UEDeployAndroid.cs
+++ b/Engine/Source/Programs/UnrealBuildTool/Android/UEDeployAndroid.cs
@@ -358,7 +358,21 @@ namespace UnrealBuildTool.Android
 					InputFiles.Add(Path.Combine(EngineDirectory, "Config\\BaseEngine.ini"));
 					InputFiles.Add(Path.Combine(ProjectDirectory, "Config\\DefaultEngine.ini"));
 
-					// look for any newer input file
+                    // rebuild if .pak files exist for OBB in APK case
+                    if (UEBuildConfiguration.bOBBinAPK)
+                    {
+                        string PAKFileLocation = ProjectDirectory + "/Saved/StagedBuilds/Android" + CookFlavor + "/" + ProjectName + "/Content/Paks";
+                        if (Directory.Exists(PAKFileLocation))
+                        {
+                            var PakFiles = Directory.EnumerateFiles(PAKFileLocation, "*.pak", SearchOption.TopDirectoryOnly);
+                            foreach (var Name in PakFiles)
+                            {
+                                InputFiles.Add(Name);
+                            }
+                        }
+                    }
+
+                    // look for any newer input file
 					DateTime ApkTime = File.GetLastWriteTimeUtc(DestApkName);
 					foreach (var InputFileName in InputFiles)
 					{
