UBT: Improve error message when running UHT fails (showing a human-readable version of the error code)

diff --git a/Engine/Source/Programs/UnrealBuildTool/System/ExternalExecution.cs b/Engine/Source/Programs/UnrealBuildTool/System/ExternalExecution.cs
--- a/Engine/Source/Programs/UnrealBuildTool/System/ExternalExecution.cs
+++ b/Engine/Source/Programs/UnrealBuildTool/System/ExternalExecution.cs
@@ -640,7 +640,7 @@ namespace UnrealBuildTool
 
 					if (UHTResult != ECompilationResult.Succeeded)
 					{
-						Log.TraceInformation("Error: Failed to generate code for {0} - error code: {1}", ActualTargetName, (int) UHTResult);
+						Log.TraceInformation("Error: Failed to generate code for {0} - error code: {2} ({1})", ActualTargetName, (int) UHTResult, UHTResult.ToString());
 						return false;
 					}
 
