CORE: Add warning message for having same named project in .sln / .xcodeproj

#ttp 345954

diff --git a/Engine/Source/Programs/UnrealBuildTool/System/UProjectInfo.cs b/Engine/Source/Programs/UnrealBuildTool/System/UProjectInfo.cs
--- a/Engine/Source/Programs/UnrealBuildTool/System/UProjectInfo.cs
+++ b/Engine/Source/Programs/UnrealBuildTool/System/UProjectInfo.cs
@@ -109,6 +109,12 @@ namespace UnrealBuildTool
 				ProjectInfoDictionary.Add(ProjectFile, NewProjectInfo);
 				ShortProjectNameDictionary.Add(NewProjectInfo.GameName, ProjectFile);
 			}
+			else
+			{
+				var FirstProject = ProjectInfoDictionary[ShortProjectNameDictionary[NewProjectInfo.GameName]];
+				Log.TraceWarning("There are multiple projects with name {0}\n\t* {1}\n\t* {2}\nThis is not currently supported and only the first one will be maintained by UBT. Please rename.",
+					NewProjectInfo.GameName, Path.GetFullPath(FirstProject.FilePath), Path.GetFullPath(NewProjectInfo.FilePath));
+			}
 		}
 
 		/// <summary>
