[GitHub] 450 : Added UBT build option to limit number of jobs.

#ttp 347138

diff --git a/Engine/Source/Programs/UnrealBuildTool/Configuration/BuildConfiguration.cs b/Engine/Source/Programs/UnrealBuildTool/Configuration/BuildConfiguration.cs
--- a/Engine/Source/Programs/UnrealBuildTool/Configuration/BuildConfiguration.cs
+++ b/Engine/Source/Programs/UnrealBuildTool/Configuration/BuildConfiguration.cs
@@ -206,6 +206,12 @@ namespace UnrealBuildTool
 		[XmlConfig]
 		public static double ProcessorCountMultiplier;
 
+		/// <summary>
+		/// Maximum processor count for local execution. 
+		/// </summary>
+		[XmlConfig]
+		public static int MaxProcessorCount;
+
 		/// <summary>
 		/// The intermediate folder - i.e. Intermediate/Build.
 		/// </summary>
@@ -472,6 +478,8 @@ namespace UnrealBuildTool
 			// to get slightly faster build times in many cases, but your computer's responsiveness during compiling may be much worse.
 			ProcessorCountMultiplier = 1.0;
 
+			MaxProcessorCount = int.MaxValue;
+
 			bTestIncludeDependencyResolveCache = false;
 			// if we are testing the resolve cache, we require UBT to use it.
 			bUseIncludeDependencyResolveCache = true;
diff --git a/Engine/Source/Programs/UnrealBuildTool/System/LocalExecutor.cs b/Engine/Source/Programs/UnrealBuildTool/System/LocalExecutor.cs
--- a/Engine/Source/Programs/UnrealBuildTool/System/LocalExecutor.cs
+++ b/Engine/Source/Programs/UnrealBuildTool/System/LocalExecutor.cs
@@ -384,6 +384,8 @@ namespace UnrealBuildTool
 				MaxActionsToExecuteInParallel = Math.Min(MaxActionsToExecuteInParallel, MaxActionsAffordedByMemory);
             }
 
+			MaxActionsToExecuteInParallel = Math.Min(MaxActionsToExecuteInParallel, BuildConfiguration.MaxProcessorCount);
+
             Log.TraceInformation("Performing {0} actions ({1} in parallel)", Actions.Count, MaxActionsToExecuteInParallel);
 
 			Dictionary<Action, ActionThread> ActionThreadDictionary = new Dictionary<Action, ActionThread>();
