- More iPhone6 support
- TappyChicken updated for Metal and iPhone6

diff --git a/Engine/Source/Developer/iOS/IOSPlatformEditor/Private/IOSTargetSettingsCustomization.cpp b/Engine/Source/Developer/iOS/IOSPlatformEditor/Private/IOSTargetSettingsCustomization.cpp
--- a/Engine/Source/Developer/iOS/IOSPlatformEditor/Private/IOSTargetSettingsCustomization.cpp
+++ b/Engine/Source/Developer/iOS/IOSPlatformEditor/Private/IOSTargetSettingsCustomization.cpp
@@ -51,6 +51,9 @@ FIOSTargetSettingsCustomization::FIOSTargetSettingsCustomization()
 	new (LaunchImageNames) FPlatformIconInfo(TEXT("[REDACTED_EMAIL]"), LOCTEXT("LaunchImage_iPadRetina_Landscape", "Launch iPad Retina in Landscape"), FText::GetEmpty(), 2048, 1536, FPlatformIconInfo::Required);
 	new (LaunchImageNames) FPlatformIconInfo(TEXT("Default-Portrait.png"), LOCTEXT("LaunchImage_iPad_Portrait", "Launch iPad in Portrait"), FText::GetEmpty(), 768, 1024, FPlatformIconInfo::Required);
 	new (LaunchImageNames) FPlatformIconInfo(TEXT("[REDACTED_EMAIL]"), LOCTEXT("LaunchImage_iPadRetina_Portrait", "Launch iPad Retina in Portrait"), FText::GetEmpty(), 1536, 2048, FPlatformIconInfo::Required);
+	new (LaunchImageNames) FPlatformIconInfo(TEXT("Default-IPhone6.png"), LOCTEXT("LaunchImage_iPhone6", "Launch iPhone 6"), FText::GetEmpty(), 750, 1334, FPlatformIconInfo::Required);
+	new (LaunchImageNames) FPlatformIconInfo(TEXT("Default-IPhone6Plus-Landscape.png"), LOCTEXT("LaunchImage_iPhone6Plus_Landscape", "Launch iPhone 6 Plus in Landscape"), FText::GetEmpty(), 2208, 1242, FPlatformIconInfo::Required);
+	new (LaunchImageNames) FPlatformIconInfo(TEXT("Default-IPhone6Plus-Portrait.png"), LOCTEXT("LaunchImage_iPhone6Plus_Portrait", "Launch iPhone 6 Plus in Portrait"), FText::GetEmpty(), 1242, 2208, FPlatformIconInfo::Required);
 }
 
 void FIOSTargetSettingsCustomization::CustomizeDetails(IDetailLayoutBuilder& DetailLayout)
@@ -278,11 +281,11 @@ void FIOSTargetSettingsCustomization::OnPlistPropertyModified()
 	DeviceCapsArrayBody += TEXT("\t\t<string>armv7</string>\n");
 	if (Settings.bSupportsOpenGLES2)
 	{
-		DeviceCapsArrayBody += TEXT("\t\t<string>opengles-2</string>");
+		DeviceCapsArrayBody += TEXT("\t\t<string>opengles-2</string>\n");
 	}
 	else if (Settings.bSupportsMetal)
 	{
-		DeviceCapsArrayBody += TEXT("\t\t<string>metal</string>");
+		DeviceCapsArrayBody += TEXT("\t\t<string>metal</string>\n");
 	}
 	DeviceCapsArrayBody += TEXT("\t");
 	Updater.ReplaceKey(RequiredDeviceCaps, ClosingArray, DeviceCapsArrayBody);
diff --git a/Engine/Source/Programs/IOS/iPhonePackager/CompileTime.cs b/Engine/Source/Programs/IOS/iPhonePackager/CompileTime.cs
--- a/Engine/Source/Programs/IOS/iPhonePackager/CompileTime.cs
+++ b/Engine/Source/Programs/IOS/iPhonePackager/CompileTime.cs
@@ -168,6 +168,40 @@ namespace iPhonePackager
 			string OutString = Info.SaveToString();
 			OutString = OutString.Replace("${EXECUTABLE_NAME}", Program.GameName);
 			OutString = OutString.Replace("${BUNDLE_IDENTIFIER}", Program.GameName.Replace("_", ""));
+
+			// this is a temp way to inject the iphone 6 images without needing to upgrade everyone's plist
+			// eventually we want to generate this based on what the user has set in the project settings
+			string[] IPhoneConfigs =  
+				{ 
+					"Default-IPhone6", "Landscape", "{375, 667}", 
+					"Default-IPhone6", "Portrait", "{375, 667}", 
+					"Default-IPhone6Plus-Landscape", "Landscape", "{414, 736}", 
+					"Default-IPhone6Plus-Portrait", "Portrait", "{414, 736}", 
+					"Default", "Landscape", "{320, 480}",
+					"Default", "Portrait", "{320, 480}",
+					"Default-568h", "Landscape", "{320, 568}",
+					"Default-568h", "Portrait", "{320, 568}",
+				};
+
+			StringBuilder NewLaunchImagesString = new StringBuilder("<key>UILaunchImages~iphone</key>\n\t\t<array>\n");
+			for (int ConfigIndex = 0; ConfigIndex < IPhoneConfigs.Length; ConfigIndex += 3)
+			{
+				NewLaunchImagesString.Append("\t\t\t<dict>\n");
+				NewLaunchImagesString.Append("\t\t\t\t<key>UILaunchImageMinimumOSVersion</key>\n");
+				NewLaunchImagesString.Append("\t\t\t\t<string>8.0</string>\n");
+				NewLaunchImagesString.Append("\t\t\t\t<key>UILaunchImageName</key>\n");
+				NewLaunchImagesString.AppendFormat("\t\t\t\t<string>{0}</string>\n", IPhoneConfigs[ConfigIndex + 0]);
+				NewLaunchImagesString.Append("\t\t\t\t<key>UILaunchImageOrientation</key>\n");
+				NewLaunchImagesString.AppendFormat("\t\t\t\t<string>{0}</string>\n", IPhoneConfigs[ConfigIndex + 1]);
+				NewLaunchImagesString.Append("\t\t\t\t<key>UILaunchImageSize</key>\n");
+				NewLaunchImagesString.AppendFormat("\t\t\t\t<string>{0}</string>\n", IPhoneConfigs[ConfigIndex + 2]);
+				NewLaunchImagesString.Append("\t\t\t</dict>\n");
+			}
+
+			// close it out
+			NewLaunchImagesString.Append("\t\t\t</array>\n\t\t<key>UILaunchImages~ipad</key>");
+			OutString = OutString.Replace("<key>UILaunchImages~ipad</key>", NewLaunchImagesString.ToString());
+
 			byte[] RawInfoPList = Encoding.UTF8.GetBytes(OutString);
 			File.WriteAllBytes(TargetPListFilename, RawInfoPList);
 
diff --git a/Engine/Source/Programs/UnrealBuildTool/IOS/UEDeployIOS.cs b/Engine/Source/Programs/UnrealBuildTool/IOS/UEDeployIOS.cs
--- a/Engine/Source/Programs/UnrealBuildTool/IOS/UEDeployIOS.cs
+++ b/Engine/Source/Programs/UnrealBuildTool/IOS/UEDeployIOS.cs
@@ -262,6 +262,7 @@ namespace UnrealBuildTool.IOS
 			}
 
 			List<string> PListAdditionalLines = new List<string>();
+			Dictionary<string, string> Replacements = new Dictionary<string, string>();
 
 			bool bSkipDefaultPNGs = false;
 			if (File.Exists(LaunchXib))
@@ -284,6 +285,41 @@ namespace UnrealBuildTool.IOS
 
 				bSkipDefaultPNGs = true;
 			}
+			else
+			{
+				// this is a temp way to inject the iphone 6 images without needing to upgrade everyone's plist
+				// eventually we want to generate this based on what the user has set in the project settings
+				string[] IPhoneConfigs =  
+				{ 
+					"Default-IPhone6", "Landscape", "{375, 667}", 
+					"Default-IPhone6", "Portrait", "{375, 667}", 
+					"Default-IPhone6Plus-Landscape", "Landscape", "{414, 736}", 
+					"Default-IPhone6Plus-Portrait", "Portrait", "{414, 736}", 
+					"Default", "Landscape", "{320, 480}",
+					"Default", "Portrait", "{320, 480}",
+					"Default-568h", "Landscape", "{320, 568}",
+					"Default-568h", "Portrait", "{320, 568}",
+				};
+
+				StringBuilder NewLaunchImagesString = new StringBuilder("<key>UILaunchImages~iphone</key>\n\t\t<array>\n");
+				for (int ConfigIndex = 0; ConfigIndex < IPhoneConfigs.Length; ConfigIndex += 3)
+				{
+					NewLaunchImagesString.Append("\t\t\t<dict>\n");
+					NewLaunchImagesString.Append("\t\t\t\t<key>UILaunchImageMinimumOSVersion</key>\n");
+					NewLaunchImagesString.Append("\t\t\t\t<string>8.0</string>\n");
+					NewLaunchImagesString.Append("\t\t\t\t<key>UILaunchImageName</key>\n");
+					NewLaunchImagesString.AppendFormat("\t\t\t\t<string>{0}</string>\n", IPhoneConfigs[ConfigIndex + 0]);
+					NewLaunchImagesString.Append("\t\t\t\t<key>UILaunchImageOrientation</key>\n");
+					NewLaunchImagesString.AppendFormat("\t\t\t\t<string>{0}</string>\n", IPhoneConfigs[ConfigIndex + 1]);
+					NewLaunchImagesString.Append("\t\t\t\t<key>UILaunchImageSize</key>\n");
+					NewLaunchImagesString.AppendFormat("\t\t\t\t<string>{0}</string>\n", IPhoneConfigs[ConfigIndex + 2]);
+					NewLaunchImagesString.Append("\t\t\t</dict>\n");
+				}
+
+				// close it out
+				NewLaunchImagesString.Append("\t\t\t</array>\n\t\t<key>UILaunchImages~ipad</key>");
+				Replacements.Add("<key>UILaunchImages~ipad</key>", NewLaunchImagesString.ToString());
+			}
 
 			// copy plist file
 			string PListFile = InEngineDir + "/Build/IOS/UE4Game-Info.plist";
@@ -297,7 +333,6 @@ namespace UnrealBuildTool.IOS
 			}
 
 			// plist replacements
-			Dictionary<string, string> Replacements = new Dictionary<string, string>();
 			Replacements.Add("${EXECUTABLE_NAME}", GameName);
 			Replacements.Add("${BUNDLE_IDENTIFIER}", InProjectName.Replace("_", ""));
 			CopyFileWithReplacements(PListFile, AppDirectory + "/Info.plist", Replacements, PListAdditionalLines);
diff --git a/Engine/Source/Runtime/Core/Private/IOS/IOSAppDelegate.cpp b/Engine/Source/Runtime/Core/Private/IOS/IOSAppDelegate.cpp
--- a/Engine/Source/Runtime/Core/Private/IOS/IOSAppDelegate.cpp
+++ b/Engine/Source/Runtime/Core/Private/IOS/IOSAppDelegate.cpp
@@ -405,37 +405,64 @@ void InstallSignalHandlers()
     UIImageOrientation orient = UIImageOrientationUp;
     NSMutableString* ImageString = [[NSMutableString alloc]init];
     [ImageString appendString:@"Default"];
-    if (MainFrame.size.height == 320 && !self.bDeviceInPortraitMode)
-    {
-        [ImageString appendString:@"-568h"];
-        orient = UIImageOrientationLeft;
-    }
-    else if (MainFrame.size.height == 568)
-    {
-        [ImageString appendString:@"-568h"];
-    }
-    else if (MainFrame.size.height == 1024 && !self.bDeviceInPortraitMode)
-    {
-        [ImageString appendString:@"-Landscape"];
-        orient = UIImageOrientationRight;
-    }
-    else if (MainFrame.size.height == 1024)
-    {
-        [ImageString appendString:@"-Portrait"];
-    }
-    else if (MainFrame.size.height == 768 && !self.bDeviceInPortraitMode)
-    {
-        [ImageString appendString:@"-Landscape"];
-    }
 
-    if (NativeScale > 1.0f)
-    {
-        [ImageString appendString:@"@2x.png"];
-    }
-    else
-    {
-        [ImageString appendString:@".png"];
-    }
+	FPlatformMisc::EIOSDevice Device = FPlatformMisc::GetIOSDeviceType();
+
+	// iphone6 has specially named files
+	if (Device == FPlatformMisc::IOS_IPhone6)
+	{
+		[ImageString appendString:@"-IPhone6"];
+		if (!self.bDeviceInPortraitMode)
+		{
+			orient = UIImageOrientationLeft;
+		}
+	}
+	else if (Device == FPlatformMisc::IOS_IPhone6Plus)
+	{
+		[ImageString appendString : @"-IPhone6Plus"];
+		if (!self.bDeviceInPortraitMode)
+		{
+			[ImageString appendString : @"-Landscape"];
+		}
+		else
+		{
+			[ImageString appendString : @"-Portrait"];
+		}
+	}
+	else
+	{
+		if (MainFrame.size.height == 320 && !self.bDeviceInPortraitMode)
+		{
+			[ImageString appendString:@"-568h"];
+			orient = UIImageOrientationLeft;
+		}
+		else if (MainFrame.size.height == 568)
+		{
+			[ImageString appendString:@"-568h"];
+		}
+		else if (MainFrame.size.height == 1024 && !self.bDeviceInPortraitMode)
+		{
+			[ImageString appendString:@"-Landscape"];
+			orient = UIImageOrientationRight;
+		}
+		else if (MainFrame.size.height == 1024)
+		{
+			[ImageString appendString:@"-Portrait"];
+		}
+		else if (MainFrame.size.height == 768 && !self.bDeviceInPortraitMode)
+		{
+			[ImageString appendString:@"-Landscape"];
+		}
+
+		if (NativeScale > 1.0f)
+		{
+			[ImageString appendString:@"@2x.png"];
+		}
+		else
+		{
+			[ImageString appendString:@".png"];
+		}
+	}
 
     [path setString: [path stringByAppendingPathComponent:ImageString]];
     UIImage* image = [[UIImage alloc] initWithContentsOfFile: path];
