Reduce optimization level for Linux builds (TTP #347022).

- clang 3.5.0 (and Ubuntu's 3.5-trunk) seems to miscompile PlatformerGame (bControlsLocked bitfield becomes true for some reason, doesn't happen with -O2 and -O0).
- fixes TTP #347022.

diff --git a/Engine/Source/Programs/UnrealBuildTool/Linux/LinuxToolChain.cs b/Engine/Source/Programs/UnrealBuildTool/Linux/LinuxToolChain.cs
--- a/Engine/Source/Programs/UnrealBuildTool/Linux/LinuxToolChain.cs
+++ b/Engine/Source/Programs/UnrealBuildTool/Linux/LinuxToolChain.cs
@@ -302,7 +302,7 @@ namespace UnrealBuildTool
             }
             else
             {
-                Result += " -O3";
+                Result += " -O2";	// warning: as of now (2014-09-28), clang 3.5.0 miscompiles PlatformerGame with -O3 (bitfields?)
             }
 
             if (CompileEnvironment.Config.bIsBuildingDLL)
