Add a utility function to determine when a path should be ignored in Rocket builds. Currently removes plugins in NoRedist and NotForLicensees folders.

diff --git a/Engine/Source/Programs/UnrealBuildTool/System/Plugins.cs b/Engine/Source/Programs/UnrealBuildTool/System/Plugins.cs
--- a/Engine/Source/Programs/UnrealBuildTool/System/Plugins.cs
+++ b/Engine/Source/Programs/UnrealBuildTool/System/Plugins.cs
@@ -386,30 +386,33 @@ namespace UnrealBuildTool
 			var PluginsDirectoryInfo = new DirectoryInfo(PluginsDirectory);
 			foreach( var PossiblePluginDirectory in PluginsDirectoryInfo.EnumerateDirectories() )
 			{
-				// Do we have a plugin descriptor in this directory?
-				bool bFoundPlugin = false;
-				foreach( var PluginDescriptorFileName in Directory.GetFiles( PossiblePluginDirectory.FullName, "*" + PluginDescriptorFileExtension ) )
+				if (!UnrealBuildTool.IsPathIgnoredForBuild(PossiblePluginDirectory.FullName))
 				{
-					// Found a plugin directory!  No need to recurse any further, but make sure it's unique.
-					if (!Plugins.Any(x => x.Directory == PossiblePluginDirectory.FullName))
+					// Do we have a plugin descriptor in this directory?
+					bool bFoundPlugin = false;
+					foreach (var PluginDescriptorFileName in Directory.GetFiles(PossiblePluginDirectory.FullName, "*" + PluginDescriptorFileExtension))
 					{
-						// Load the plugin info and keep track of it
-						var PluginDescriptorFile = new FileInfo(PluginDescriptorFileName);
-						var PluginInfo = LoadPluginDescriptor(PluginDescriptorFile, LoadedFrom);
+						// Found a plugin directory!  No need to recurse any further, but make sure it's unique.
+						if (!Plugins.Any(x => x.Directory == PossiblePluginDirectory.FullName))
+						{
+							// Load the plugin info and keep track of it
+							var PluginDescriptorFile = new FileInfo(PluginDescriptorFileName);
+							var PluginInfo = LoadPluginDescriptor(PluginDescriptorFile, LoadedFrom);
 
-						Plugins.Add(PluginInfo);
-						bFoundPlugin = true;
-						Log.TraceVerbose("Found plugin in: " + PluginInfo.Directory);
-					}
+							Plugins.Add(PluginInfo);
+							bFoundPlugin = true;
+							Log.TraceVerbose("Found plugin in: " + PluginInfo.Directory);
+						}
 
-					// No need to search for more plugins
-					break;
-				}
+						// No need to search for more plugins
+						break;
+					}
 
-				if( !bFoundPlugin )
-				{
-					// Didn't find a plugin in this directory.  Continue to look in subfolders.
-					FindPluginsRecursively( PossiblePluginDirectory.FullName, LoadedFrom, ref Plugins );
+					if (!bFoundPlugin)
+					{
+						// Didn't find a plugin in this directory.  Continue to look in subfolders.
+						FindPluginsRecursively(PossiblePluginDirectory.FullName, LoadedFrom, ref Plugins);
+					}
 				}
 			}
 		}
diff --git a/Engine/Source/Programs/UnrealBuildTool/System/UnrealBuildTool.cs b/Engine/Source/Programs/UnrealBuildTool/System/UnrealBuildTool.cs
--- a/Engine/Source/Programs/UnrealBuildTool/System/UnrealBuildTool.cs
+++ b/Engine/Source/Programs/UnrealBuildTool/System/UnrealBuildTool.cs
@@ -405,6 +405,24 @@ namespace UnrealBuildTool
             }
         }
 
+		/// <summary>
+		/// Determines whether the given path is ignored for the active build.
+		/// </summary>
+		/// <param name="InString">The path to check</param>
+		/// <returns>true if the path should be ignored for the build</returns>
+		public static bool IsPathIgnoredForBuild(string InPath)
+		{
+			if (BuildingRocket())
+			{
+				string NormalizedPath = InPath.Replace('\\', '/').ToLowerInvariant();
+				if (NormalizedPath.Contains("/notforlicensees/") || NormalizedPath.Contains("/epicinternal/") || NormalizedPath.Contains("/noredist/"))
+				{
+					return true;
+				}
+			}
+			return false;
+		}
+
         /// <summary>
         /// See if the given string was pass in on the command line
         /// </summary>
