Remove stale exception when building games under Rocket.

diff --git a/Engine/Source/Programs/UnrealBuildTool/System/UnrealBuildTool.cs b/Engine/Source/Programs/UnrealBuildTool/System/UnrealBuildTool.cs
--- a/Engine/Source/Programs/UnrealBuildTool/System/UnrealBuildTool.cs
+++ b/Engine/Source/Programs/UnrealBuildTool/System/UnrealBuildTool.cs
@@ -1547,14 +1547,6 @@ namespace UnrealBuildTool
 
                         if ( (BuildConfiguration.bXGEExport && UEBuildConfiguration.bGenerateManifest) || (!ProjectFileGenerator.bGenerateProjectFiles && !UEBuildConfiguration.bGenerateManifest && !UEBuildConfiguration.bCleanProject))
                         {
-                            // We don't currently support building game targets in rocket.
-                            if (UnrealBuildTool.RunningRocket() && Target.Rules != null && Target.TargetType == TargetRules.TargetType.Game)
-                            {
-                                throw new BuildException(
-                                    "You currently can not build a game target in Rocket.\nTry again in a future release.\nFor now, build and run your editor project."
-                                );
-                            }
-
                             // Generate an action graph if we were asked to do that.  The graph generation needs access to the include dependency cache, so
                             // we generate it before saving and cleaning that up.
                             if( GeneratingActionGraph )
