Disable XGE for Linux builds since it still crashes on farm.

diff --git a/Engine/Source/Programs/UnrealBuildTool/Linux/UEBuildLinux.cs b/Engine/Source/Programs/UnrealBuildTool/Linux/UEBuildLinux.cs
--- a/Engine/Source/Programs/UnrealBuildTool/Linux/UEBuildLinux.cs
+++ b/Engine/Source/Programs/UnrealBuildTool/Linux/UEBuildLinux.cs
@@ -100,10 +100,16 @@ namespace UnrealBuildTool
 
             return SDKStatus.Invalid;
         }
-        
-        /**
-         *	Register the platform with the UEBuildPlatform class
-         */
+
+		public override bool CanUseXGE()
+		{
+			// disabled until XGE crash is fixed - it is still happening as of 2014-09-30
+			return false;
+		}
+
+		/**
+		 *	Register the platform with the UEBuildPlatform class
+		 */
         protected override void RegisterBuildPlatformInternal()
         {
 			if ((ProjectFileGenerator.bGenerateProjectFiles == true) || (HasRequiredSDKsInstalled() == SDKStatus.Valid))
