Linux: clean up and fix rpaths.

- Removed old rpaths for libs which we now link statically.
- Added rpath useful for target platform modules and other fixes from PR #442.

diff --git a/Engine/Source/Programs/UnrealBuildTool/Linux/LinuxToolChain.cs b/Engine/Source/Programs/UnrealBuildTool/Linux/LinuxToolChain.cs
--- a/Engine/Source/Programs/UnrealBuildTool/Linux/LinuxToolChain.cs
+++ b/Engine/Source/Programs/UnrealBuildTool/Linux/LinuxToolChain.cs
@@ -395,15 +395,12 @@ namespace UnrealBuildTool
             }
 
             // RPATH for third party libs
-            Result += " -Wl,-rpath=${{ORIGIN}}";
-            Result += " -Wl,-rpath-link=${{ORIGIN}}";
+            Result += " -Wl,-rpath=${ORIGIN}";
+            Result += " -Wl,-rpath-link=${ORIGIN}";
             Result += " -Wl,-rpath=${ORIGIN}/../../../Engine/Binaries/Linux";
-            // FIXME: really ugly temp solution. Modules need to be able to specify this
-            Result += " -Wl,-rpath=${ORIGIN}/../../../Engine/Binaries/ThirdParty/jemalloc/Linux/x86_64-unknown-linux-gnu";
+			Result += " -Wl,-rpath=${ORIGIN}/..";	// for modules that are in sub-folders of the main Engine/Binary/Linux folder
+			// FIXME: really ugly temp solution. Modules need to be able to specify this
             Result += " -Wl,-rpath=${ORIGIN}/../../../Engine/Binaries/ThirdParty/ICU/icu4c-53_1/Linux/x86_64-unknown-linux-gnu";
-            Result += " -Wl,-rpath=${ORIGIN}/../../../Engine/Binaries/ThirdParty/SDL2/Linux/x86_64-unknown-linux-gnu";
-            Result += " -Wl,-rpath=${ORIGIN}/../../../Engine/Binaries/ThirdParty/zlib/Linux/x86_64-unknown-linux-gnu";
-            Result += " -Wl,-rpath=${ORIGIN}/../../../Engine/Binaries/ThirdParty/libPNG/Linux/x86_64-unknown-linux-gnu";
 
             if (CrossCompiling())
             {
