Fixing hot-reload from IDE on the Mac

#change Adding better (working) versions of Process.GetProcesses(), Process.Name and Process.Modules to BuildHostPlatformAbstraction
#change ShouldDoHotReload will now use those new methods to learn if the target is running
#change On the Mac, we can't use 'Added' filesystem event to know if new binaries have been compiled as it reports temporary linker files instead of the final binaries which are only reported by 'Modified' event.

diff --git a/Engine/Source/Developer/HotReload/Private/HotReload.cpp b/Engine/Source/Developer/HotReload/Private/HotReload.cpp
--- a/Engine/Source/Developer/HotReload/Private/HotReload.cpp
+++ b/Engine/Source/Developer/HotReload/Private/HotReload.cpp
@@ -786,7 +786,14 @@ void FHotReloadModule::OnHotReloadBinariesChanged(const TArray<struct FFileChang
 		// Check if any of the game DLLs has been added
 		for (auto& Change : FileChanges)
 		{
+#if PLATFORM_MAC
+			// On the Mac the Add event is for a temporary linker(?) file that gets immediately renamed
+			// to a dylib. In the future we may want to support modified event for all platforms anyway once
+			// shadow copying works with hot-reload.
+			if (Change.Action == FFileChangeData::FCA_Modified)
+#else
 			if (Change.Action == FFileChangeData::FCA_Added)
+#endif
 			{
 				const FString Filename = FPaths::GetCleanFilename(Change.Filename);
 				if (Filename.EndsWith(FPlatformProcess::GetModuleExtension()))
diff --git a/Engine/Source/Programs/UnrealBuildTool/System/BuildHostPlatform.cs b/Engine/Source/Programs/UnrealBuildTool/System/BuildHostPlatform.cs
--- a/Engine/Source/Programs/UnrealBuildTool/System/BuildHostPlatform.cs
+++ b/Engine/Source/Programs/UnrealBuildTool/System/BuildHostPlatform.cs
@@ -3,6 +3,7 @@ using System.Collections.Generic;
 using System.Linq;
 using System.Text;
 using System.IO;
+using System.Diagnostics;
 using System.Runtime.InteropServices;
 
 namespace UnrealBuildTool
@@ -71,6 +72,124 @@ namespace UnrealBuildTool
 		/// <param name="Filename">Filename of the library</param>
 		/// <returns>API version of -1 if not found.</returns>
 		abstract public int GetDllApiVersion(string Filename);
+
+		/// <summary>
+		/// Class that holds information about a running process
+		/// </summary>
+		public class ProcessInfo
+		{
+			/// <summary>
+			/// Process ID
+			/// </summary>
+			public int PID;
+			/// <summary>
+			/// Name of the process
+			/// </summary>
+			public string Name;
+			/// <summary>
+			/// Filename of the process binary
+			/// </summary>
+			public string Filename;
+
+			public ProcessInfo(int InPID, string InName, string InFilename, string[] InModules)
+			{
+				PID = InPID;
+				Name = InName;
+				Filename = InFilename;
+			}
+			public ProcessInfo(Process Proc)
+			{
+				PID = Proc.Id;
+				Name = Proc.ProcessName;
+				Filename = Path.GetFullPath(Proc.MainModule.FileName);
+			}
+			public override string ToString()
+			{
+				return String.Format("{0}, {1}", Name, Filename);
+			}
+		}
+			
+		/// <summary>
+		/// Gets all currently running processes.
+		/// </summary>
+		/// <returns></returns>
+		public virtual ProcessInfo[] GetProcesses()
+		{
+			var AllProcesses = Process.GetProcesses();
+			var Result = new List<ProcessInfo>(AllProcesses.Length);
+			foreach (var Proc in AllProcesses) 
+			{
+				try 
+				{
+					if (!Proc.HasExited) 
+					{
+						Result.Add(new ProcessInfo(Proc));
+					}
+				} catch {}
+			}
+			return Result.ToArray();
+		}
+
+		/// <summary>
+		/// Gets a process by name.
+		/// </summary>
+		/// <param name="Name">Name of the process to get information for.</param>
+		/// <returns></returns>
+		public virtual ProcessInfo GetProcessByName(string Name)
+		{
+			var AllProcess = GetProcesses();
+			foreach (var Info in AllProcess) 
+			{
+				if (Info.Name == Name) 
+				{
+					return Info;
+				}
+			}
+			return null;
+		}
+
+		/// <summary>
+		/// Gets processes by name.
+		/// </summary>
+		/// <param name="Name">Name of the process to get information for.</param>
+		/// <returns></returns>
+		public virtual ProcessInfo[] GetProcessesByName(string Name)
+		{
+			var AllProcess = GetProcesses();
+			var Result = new List<ProcessInfo>();
+			foreach (var Info in AllProcess) 
+			{
+				if (Info.Name == Name) 
+				{
+					Result.Add(Info);
+				}
+			}
+			return Result.ToArray();
+		}
+
+		/// <summary>
+		/// Gets the filenames of all modules associated with a process
+		/// </summary>
+		/// <param name="PID">Process ID</param>
+		/// <param name="Filename">Filename of the binary associated with the process.</param>
+		/// <returns>An array of all module filenames associated with the process. Can be empty of the process is no longer running.</returns>
+		public virtual string[] GetProcessModules(int PID, string Filename)
+		{
+			List<string> Modules = new List<string>();
+			try
+			{
+				var Proc = Process.GetProcessById(PID);
+				if (Proc != null)
+				{
+					foreach (var Module in Proc.Modules.Cast<System.Diagnostics.ProcessModule>())
+					{
+						Modules.Add(Path.GetFullPath(Module.FileName));
+					}
+				}
+			}
+			catch { }
+			return Modules.ToArray();
+		}
 	}
 
 	class WindowsBuildHostPlatform : BuildHostPlatform
@@ -164,6 +283,7 @@ namespace UnrealBuildTool
 
 			return Result;
 		}
+
 	}
 
 	class MacBuildHostPlatform : BuildHostPlatform
@@ -178,6 +298,108 @@ namespace UnrealBuildTool
 			// @TODO: Implement GetDllApiVersion for Mac
 			return -1;
 		}
+
+		/// <summary>
+		/// Currently Mono returns incomplete process names in Process.GetProcesses() so we need to parse 'ps' output.
+		/// </summary>
+		/// <returns></returns>
+		public override ProcessInfo[] GetProcesses()
+		{
+			var Result = new List<ProcessInfo>();
+
+			var StartInfo = new ProcessStartInfo ();
+			StartInfo.FileName = "ps";
+			StartInfo.Arguments = "-eaw -o pid,comm";
+			StartInfo.CreateNoWindow = true;
+			StartInfo.UseShellExecute = false;
+			StartInfo.RedirectStandardOutput = true;
+
+			var Proc = new Process ();
+			Proc.StartInfo = StartInfo;
+			try
+			{
+				Proc.Start();
+				Proc.WaitForExit();
+				for (string Line = Proc.StandardOutput.ReadLine(); Line != null; Line = Proc.StandardOutput.ReadLine())
+				{
+					Line = Line.Trim();
+					int PIDEnd = Line.IndexOf(' ');
+					var PIDString = Line.Substring(0, PIDEnd);
+					if (PIDString != "PID")
+					{
+						var Filename = Line.Substring(PIDEnd + 1);
+						var Pid = Int32.Parse(PIDString);
+						try
+						{
+							var ExistingProc = Process.GetProcessById(Pid);
+							if (ExistingProc != null && ExistingProc.HasExited == false)
+							{
+								var ProcInfo = new ProcessInfo(ExistingProc);
+								ProcInfo.Name = Path.GetFileName(Filename);
+								ProcInfo.Filename = Filename;
+								Result.Add(ProcInfo);
+							}
+						} catch {}
+					}
+				}
+
+			}
+			catch {}
+			return Result.ToArray();
+		}
+
+		/// <summary>
+		/// Currently Mono returns incomplete list of modules for Process.Modules so we need to parse vmmap output.
+		/// </summary>
+		/// <param name="PID"></param>
+		/// <param name="Filename"></param>
+		/// <returns></returns>
+		public override string[] GetProcessModules(int PID, string Filename)
+		{
+			HashSet<string> Modules = new HashSet<string>();
+			// Add the process file name to the module list. This is to make it compatible with the results of Process.Modules on Windows.
+			Modules.Add(Filename);
+
+			var StartInfo = new ProcessStartInfo();
+			StartInfo.FileName = "vmmap";
+			StartInfo.Arguments = String.Format("{0} -w", PID);
+			StartInfo.CreateNoWindow = true;
+			StartInfo.UseShellExecute = false;
+			StartInfo.RedirectStandardOutput = true;
+
+			var Proc = new Process();
+			Proc.StartInfo = StartInfo;
+			try
+			{
+				Proc.Start();
+				// Start processing output before vmmap exits otherwise it's going to hang
+				while (!Proc.WaitForExit(1))
+				{
+					ProcessVMMapOutput(Proc, Modules);
+				}
+				ProcessVMMapOutput(Proc, Modules);
+			}
+			catch { }
+			return Modules.ToArray();
+		}
+		private void ProcessVMMapOutput(Process Proc, HashSet<string> Modules)
+		{
+			for (string Line = Proc.StandardOutput.ReadLine(); Line != null; Line = Proc.StandardOutput.ReadLine())
+			{
+				Line = Line.Trim();
+				if (Line.EndsWith(".dylib"))
+				{
+					const int SharingModeLength = 6;
+					int SMStart = Line.IndexOf("SM=");
+					int PathStart = SMStart + SharingModeLength;
+					string Module = Line.Substring(PathStart).Trim();
+					if (!Modules.Contains(Module))
+					{
+						Modules.Add(Module);
+					}
+				}
+			}
+		}
 	}
 
 	class LinuxBuildHostPlatform : BuildHostPlatform
diff --git a/Engine/Source/Programs/UnrealBuildTool/System/UnrealBuildTool.cs b/Engine/Source/Programs/UnrealBuildTool/System/UnrealBuildTool.cs
--- a/Engine/Source/Programs/UnrealBuildTool/System/UnrealBuildTool.cs
+++ b/Engine/Source/Programs/UnrealBuildTool/System/UnrealBuildTool.cs
@@ -1818,14 +1818,15 @@ namespace UnrealBuildTool
 				{
 					throw new BuildException("ShouldDoHotReload cannot handle multiple binaries returning from UEBuildTarget.MakeBinaryPaths");
 				}
+
 				var EditorProcessFilename = EditorProcessFilenames[0];
 				var EditorProcessName = Path.GetFileNameWithoutExtension(EditorProcessFilename);
-				var EditorProcesses = Process.GetProcessesByName(EditorProcessName);
+				var EditorProcesses = BuildHostPlatform.Current.GetProcessesByName(EditorProcessName);
 				var BinariesPath = Path.GetFullPath(Path.GetDirectoryName(EditorProcessFilename));
 				var PerProjectBinariesPath = Path.Combine(Target.ProjectDirectory, BinariesPath.Substring(BinariesPath.LastIndexOf("Binaries")));
-				bIsRunning = EditorProcesses.FirstOrDefault(p =>
+				bIsRunning = EditorProcesses.FirstOrDefault(EditorProc =>
 					{
-						if(!Path.GetFullPath(p.MainModule.FileName).StartsWith(BinariesPath, StringComparison.InvariantCultureIgnoreCase))
+						if(!Path.GetFullPath(EditorProc.Filename).StartsWith(BinariesPath, StringComparison.InvariantCultureIgnoreCase))
 						{
 							return false;
 						}
@@ -1835,13 +1836,14 @@ namespace UnrealBuildTool
 							return false;
 						}
 
-						if(!p.Modules.Cast<System.Diagnostics.ProcessModule>().Any(Module => Path.GetFullPath(Module.FileName).StartsWith(PerProjectBinariesPath)))
+						var Modules = BuildHostPlatform.Current.GetProcessModules(EditorProc.PID, EditorProc.Filename);
+						if(!Modules.Any(Module => Module.StartsWith(PerProjectBinariesPath)))
 						{
 							return false;
 						}
 
 						return true;
-					}) != default(Process);
+					}) != default(BuildHostPlatform.ProcessInfo);
 			}
 			return bIsRunning;
 		}
