remove some debug code
#ubt

diff --git a/Engine/Source/Programs/UnrealBuildTool/System/XcodeProjectFileGenerator.cs b/Engine/Source/Programs/UnrealBuildTool/System/XcodeProjectFileGenerator.cs
--- a/Engine/Source/Programs/UnrealBuildTool/System/XcodeProjectFileGenerator.cs
+++ b/Engine/Source/Programs/UnrealBuildTool/System/XcodeProjectFileGenerator.cs
@@ -1114,7 +1114,6 @@ namespace UnrealBuildTool
 					if (!IsEngineTarget && TargetRulesObject.Type != TargetRules.TargetType.Program)
 					{
 						List<UProjectInfo> AllGames = UProjectInfo.FilterGameProjects(true, bGeneratingGameProjectFiles ? GameProjectName : null);
-						Console.WriteLine("Game Count: " + AllGames.Count);
 						UProjectInfo ProjectInfo = FindGameContainingFile (AllGames, TargetFilePath);
 						if (ProjectInfo != null)
 						{
