Move WriteJavaBuildSettingsFile and add it to timestamp check so changing settings will rebuild APK
#ue4
#android

diff --git a/Engine/Source/Programs/UnrealBuildTool/Android/UEDeployAndroid.cs b/Engine/Source/Programs/UnrealBuildTool/Android/UEDeployAndroid.cs
--- a/Engine/Source/Programs/UnrealBuildTool/Android/UEDeployAndroid.cs
+++ b/Engine/Source/Programs/UnrealBuildTool/Android/UEDeployAndroid.cs
@@ -327,6 +327,10 @@ namespace UnrealBuildTool.Android
 			string[] GPUArchitectures = AndroidToolChain.GetAllGPUArchitectures();
 			int NumArches = Arches.Length * GPUArchitectures.Length;
 
+            // See if we need to create a 'default' Java Build settings file if one doesn't exist (if it does exist we have to assume it has been setup correctly)
+            string UE4JavaBuildSettingsFileName = GetUE4JavaBuildSettingsFileName(EngineDirectory);
+            WriteJavaBuildSettingsFile(UE4JavaBuildSettingsFileName, UEBuildConfiguration.bOBBinAPK);
+
 			// first check if all .so's are up to date
 			bool bAllInputsCurrent = true;
 			foreach (string Arch in Arches)
@@ -358,6 +362,9 @@ namespace UnrealBuildTool.Android
 					InputFiles.Add(Path.Combine(EngineDirectory, "Config\\BaseEngine.ini"));
 					InputFiles.Add(Path.Combine(ProjectDirectory, "Config\\DefaultEngine.ini"));
 
+                    // make sure changed java settings will rebuild apk
+                    InputFiles.Add(UE4JavaBuildSettingsFileName);
+
                     // rebuild if .pak files exist for OBB in APK case
                     if (UEBuildConfiguration.bOBBinAPK)
                     {
@@ -402,10 +409,6 @@ namespace UnrealBuildTool.Android
 
 			// Once for all arches code:
 
-			// See if we need to create a 'default' Java Build settings file if one doesn't exist (if it does exist we have to assume it has been setup correctly)
-			string UE4JavaBuildSettingsFileName = GetUE4JavaBuildSettingsFileName(EngineDirectory);
-			WriteJavaBuildSettingsFile(UE4JavaBuildSettingsFileName, UEBuildConfiguration.bOBBinAPK);
-
 			// make up a dictionary of strings to replace in the Manifest file
 			Dictionary<string, string> Replacements = new Dictionary<string, string>();
 			Replacements.Add("${EXECUTABLE_NAME}", ProjectName);
