use -Oz as default for HTML5 builds. Hopefully it will fix tool chain crashes on the buildbots, while we evaluate tool chain fixes put in mozilla dev branch.

diff --git a/Engine/Source/Programs/UnrealBuildTool/HTML5/UEBuildHTML5.cs b/Engine/Source/Programs/UnrealBuildTool/HTML5/UEBuildHTML5.cs
--- a/Engine/Source/Programs/UnrealBuildTool/HTML5/UEBuildHTML5.cs
+++ b/Engine/Source/Programs/UnrealBuildTool/HTML5/UEBuildHTML5.cs
@@ -272,6 +272,7 @@ namespace UnrealBuildTool
             UEBuildConfiguration.bRuntimePhysicsCooking = false;
             UEBuildConfiguration.bCompileSimplygon = false;
             UEBuildConfiguration.bCompileICU = true;
+            UEBuildConfiguration.bCompileForSize = true;
         }
 
         /**
