Updated engine code to compile correctly with OS X 10.10 SDK, added SDK detection to MacToolChain (based on iOS one)

diff --git a/Engine/Source/Programs/CrashReporter/CrashReportClient/Private/Mac/CrashReportClientMainMac.cpp b/Engine/Source/Programs/CrashReporter/CrashReportClient/Private/Mac/CrashReportClientMainMac.cpp
--- a/Engine/Source/Programs/CrashReporter/CrashReportClient/Private/Mac/CrashReportClientMainMac.cpp
+++ b/Engine/Source/Programs/CrashReporter/CrashReportClient/Private/Mac/CrashReportClientMainMac.cpp
@@ -32,7 +32,7 @@ void CrashReporterCrashHandler(const FGenericCrashContext& GenericContext)
 
 static FString GSavedCommandLine;
 
-@interface UE4AppDelegate : NSObject <NSApplicationDelegate>
+@interface UE4AppDelegate : NSObject <NSApplicationDelegate, NSFileManagerDelegate>
 {
 }
 
diff --git a/Engine/Source/Programs/Mac/UE4EditorServices/Private/Mac/UE4EditorServicesAppDelegate.h b/Engine/Source/Programs/Mac/UE4EditorServices/Private/Mac/UE4EditorServicesAppDelegate.h
--- a/Engine/Source/Programs/Mac/UE4EditorServices/Private/Mac/UE4EditorServicesAppDelegate.h
+++ b/Engine/Source/Programs/Mac/UE4EditorServices/Private/Mac/UE4EditorServicesAppDelegate.h
@@ -2,7 +2,7 @@
 
 #import <Cocoa/Cocoa.h>
 
-@interface FUE4EditorServicesAppDelegate : NSObject <NSApplicationDelegate>
+@interface FUE4EditorServicesAppDelegate : NSObject <NSApplicationDelegate, NSFileManagerDelegate>
 {
 	NSWindow* Window;
 }
diff --git a/Engine/Source/Programs/SlateViewer/Private/Mac/SlateViewerMainMac.cpp b/Engine/Source/Programs/SlateViewer/Private/Mac/SlateViewerMainMac.cpp
--- a/Engine/Source/Programs/SlateViewer/Private/Mac/SlateViewerMainMac.cpp
+++ b/Engine/Source/Programs/SlateViewer/Private/Mac/SlateViewerMainMac.cpp
@@ -6,7 +6,7 @@
 
 static FString GSavedCommandLine;
 
-@interface UE4AppDelegate : NSObject <NSApplicationDelegate>
+@interface UE4AppDelegate : NSObject <NSApplicationDelegate, NSFileManagerDelegate>
 {
 }
 
diff --git a/Engine/Source/Programs/SymbolDebugger/Private/Mac/SymbolDebuggerMainMac.cpp b/Engine/Source/Programs/SymbolDebugger/Private/Mac/SymbolDebuggerMainMac.cpp
--- a/Engine/Source/Programs/SymbolDebugger/Private/Mac/SymbolDebuggerMainMac.cpp
+++ b/Engine/Source/Programs/SymbolDebugger/Private/Mac/SymbolDebuggerMainMac.cpp
@@ -5,7 +5,7 @@
 
 static FString GSavedCommandLine;
 
-@interface UE4AppDelegate : NSObject <NSApplicationDelegate>
+@interface UE4AppDelegate : NSObject <NSApplicationDelegate, NSFileManagerDelegate>
 {
 }
 
diff --git a/Engine/Source/Programs/UnrealBuildTool/Mac/MacToolChain.cs b/Engine/Source/Programs/UnrealBuildTool/Mac/MacToolChain.cs
--- a/Engine/Source/Programs/UnrealBuildTool/Mac/MacToolChain.cs
+++ b/Engine/Source/Programs/UnrealBuildTool/Mac/MacToolChain.cs
@@ -1,6 +1,7 @@
 // Copyright 1998-2014 Epic Games, Inc. All Rights Reserved.
 
 using System;
+using System.Collections;
 using System.Collections.Generic;
 using System.IO;
 using System.Diagnostics;
@@ -24,7 +25,7 @@ namespace UnrealBuildTool
 		 ***********************************************************************
 
 		/** Which version of the Mac OS SDK to target at build time */
-		public static string MacOSSDKVersion = "10.9";
+		public static string MacOSSDKVersion = "latest";
 
 		/** Which version of the Mac OS X to allow at run time */
 		public static string MacOSVersion = "10.9";
@@ -33,7 +34,10 @@ namespace UnrealBuildTool
 		public static string MinMacOSVersion = "10.9.2";
 
 		/** Which developer directory to root from */
-		private static string DeveloperDir = "/Applications/Xcode.app/Contents/Developer/";
+		private static string XcodeDeveloperDir = "/Applications/Xcode.app/Contents/Developer/";
+
+		/** Location of the SDKs */
+		private static string BaseSDKDir;
 
 		/** Which compiler frontend to use */
 		private static string MacCompiler = "clang++";
@@ -57,6 +61,78 @@ namespace UnrealBuildTool
 		{
 			base.SetUpGlobalEnvironment();
 
+			BaseSDKDir = XcodeDeveloperDir + "Platforms/MacOSX.platform/Developer/SDKs";
+
+			if (MacOSSDKVersion == "latest")
+			{
+				try
+				{
+					string[] SubDirs = null;
+					if (Utils.IsRunningOnMono)
+					{
+						// on the Mac, we can just get the directory name
+						SubDirs = System.IO.Directory.GetDirectories(BaseSDKDir);
+						Log.TraceInformation(String.Format("Directories : {0} {1}", SubDirs, SubDirs[0]));
+					}
+					else
+					{
+						Hashtable Results = RPCUtilHelper.Command("/", "ls", BaseSDKDir, null);
+						if (Results != null)
+						{
+							string Result = (string)Results["CommandOutput"];
+							SubDirs = Result.Split("\r\n".ToCharArray(), StringSplitOptions.RemoveEmptyEntries);
+						}
+					}
+
+					// loop over the subdirs and parse out the version
+					float MaxSDKVersion = 0.0f;
+					string MaxSDKVersionString = null;
+					foreach (string SubDir in SubDirs)
+					{
+						string SubDirName = Path.GetFileNameWithoutExtension(SubDir);
+						if (SubDirName.StartsWith("MacOSX10."))
+						{
+							// get the SDK version from the directory name
+							string SDKString = SubDirName.Replace("MacOSX10.", "");
+							float SDKVersion = 0.0f;
+							try
+							{
+								SDKVersion = float.Parse(SDKString, System.Globalization.CultureInfo.InvariantCulture);
+							}
+							catch (Exception)
+							{
+								// weirdly formatted SDKs
+								continue;
+							}
+
+							// update largest SDK version number
+							if (SDKVersion > MaxSDKVersion)
+							{
+								MaxSDKVersion = SDKVersion;
+								MaxSDKVersionString = SDKString;
+							}
+						}
+					}
+
+					// convert back to a string with the exact format
+					if (MaxSDKVersionString != null)
+					{
+						MacOSSDKVersion = "10." + MaxSDKVersionString;
+					}
+				}
+				catch (Exception Ex)
+				{
+					// on any exception, just use the backup version
+					Log.TraceInformation("Triggered an exception while looking for SDK directory in Xcode.app");
+					Log.TraceInformation("{0}", Ex.ToString());
+				}
+
+				if (MacOSSDKVersion == "latest")
+				{
+					throw new BuildException("Unable to determine SDK version from Xcode, we cannot continue");
+				}
+			}
+
 			if (!ProjectFileGenerator.bGenerateProjectFiles)
 			{
 				if (BuildHostPlatform.Current.Platform != UnrealTargetPlatform.Mac)
@@ -102,7 +178,7 @@ namespace UnrealBuildTool
 			Result += " -c";
 
 			Result += " -arch x86_64";
-			Result += " -isysroot " + DeveloperDir + "Platforms/MacOSX.platform/Developer/SDKs/MacOSX" + MacOSSDKVersion + ".sdk";
+			Result += " -isysroot " + BaseSDKDir + "/MacOSX" + MacOSSDKVersion + ".sdk";
 			Result += " -mmacosx-version-min=" + MacOSVersion;
 
 			// Optimize non- debug builds.
@@ -190,7 +266,7 @@ namespace UnrealBuildTool
 			string Result = "";
 
 			Result += " -arch x86_64";
-			Result += " -isysroot " + DeveloperDir + "Platforms/MacOSX.platform/Developer/SDKs/MacOSX" + MacOSSDKVersion + ".sdk";
+			Result += " -isysroot " + BaseSDKDir + "/MacOSX" + MacOSSDKVersion + ".sdk";
 			Result += " -mmacosx-version-min=" + MacOSVersion;
 			Result += " -dead_strip";
 
diff --git a/Engine/Source/Programs/UnrealFrontend/Private/Mac/MacUnrealFrontendMain.cpp b/Engine/Source/Programs/UnrealFrontend/Private/Mac/MacUnrealFrontendMain.cpp
--- a/Engine/Source/Programs/UnrealFrontend/Private/Mac/MacUnrealFrontendMain.cpp
+++ b/Engine/Source/Programs/UnrealFrontend/Private/Mac/MacUnrealFrontendMain.cpp
@@ -10,7 +10,7 @@
 static FString GSavedCommandLine;
 
 
-@interface UE4AppDelegate : NSObject<NSApplicationDelegate>
+@interface UE4AppDelegate : NSObject<NSApplicationDelegate, NSFileManagerDelegate>
 {
 }
 
diff --git a/Engine/Source/Runtime/Core/Private/Mac/MacApplication.cpp b/Engine/Source/Runtime/Core/Private/Mac/MacApplication.cpp
--- a/Engine/Source/Runtime/Core/Private/Mac/MacApplication.cpp
+++ b/Engine/Source/Runtime/Core/Private/Mac/MacApplication.cpp
@@ -52,7 +52,7 @@ NSEvent* FMacApplication::HandleNSEvent(NSEvent* Event)
 		const bool bIsMouseClickOrKeyEvent = [Event type] == NSLeftMouseDown || [Event type] == NSLeftMouseUp
 			|| [Event type] == NSRightMouseDown || [Event type] == NSRightMouseUp
 			|| [Event type] == NSOtherMouseDown || [Event type] == NSOtherMouseUp;
-		const bool bIsResentEvent = [Event type] == NSApplicationDefined && [Event subtype] == FMacApplication::ResentEvent;
+		const bool bIsResentEvent = [Event type] == NSApplicationDefined && (FMacApplicationEventTypes)[Event subtype] == FMacApplication::ResentEvent;
 		
 		if (MacApplication)
 		{
diff --git a/Engine/Source/Runtime/Launch/Private/Mac/LaunchMac.cpp b/Engine/Source/Runtime/Launch/Private/Mac/LaunchMac.cpp
--- a/Engine/Source/Runtime/Launch/Private/Mac/LaunchMac.cpp
+++ b/Engine/Source/Runtime/Launch/Private/Mac/LaunchMac.cpp
@@ -43,7 +43,7 @@ void EngineCrashHandler(const FGenericCrashContext& GenericContext)
 	return Context.GenerateCrashInfoAndLaunchReporter();
 }
 
-@interface UE4AppDelegate : NSObject <NSApplicationDelegate>
+@interface UE4AppDelegate : NSObject <NSApplicationDelegate, NSFileManagerDelegate>
 {
 #if WITH_EDITOR
 	NSString* Filename;
