Removed logging of found SDK directories from MacToolchain

diff --git a/Engine/Source/Programs/UnrealBuildTool/Mac/MacToolChain.cs b/Engine/Source/Programs/UnrealBuildTool/Mac/MacToolChain.cs
--- a/Engine/Source/Programs/UnrealBuildTool/Mac/MacToolChain.cs
+++ b/Engine/Source/Programs/UnrealBuildTool/Mac/MacToolChain.cs
@@ -72,7 +72,6 @@ namespace UnrealBuildTool
 					{
 						// on the Mac, we can just get the directory name
 						SubDirs = System.IO.Directory.GetDirectories(BaseSDKDir);
-						Log.TraceInformation(String.Format("Directories : {0} {1}", SubDirs, SubDirs[0]));
 					}
 					else
 					{
