Set ICU on by default.

Most platforms allow override it to true anyway, but some (Linux) don't, allowing targets to set it up in their SetGlobalEnvironment(). Since you cannot do that in game/editor targets though (see UEBuildTarget.SetupGlobalEnvironment()), default value should be true to get the editor compiled with ICU.

diff --git a/Engine/Source/Programs/UnrealBuildTool/Configuration/UEBuildConfiguration.cs b/Engine/Source/Programs/UnrealBuildTool/Configuration/UEBuildConfiguration.cs
--- a/Engine/Source/Programs/UnrealBuildTool/Configuration/UEBuildConfiguration.cs
+++ b/Engine/Source/Programs/UnrealBuildTool/Configuration/UEBuildConfiguration.cs
@@ -197,6 +197,7 @@ namespace UnrealBuildTool
 			bCompileAPEX = true;
             bRuntimePhysicsCooking = true;
 			bCompileBox2D = true;
+			bCompileICU = true;
 			bBuildDedicatedServer = false;
 			bBuildEditor = true;
 			bBuildRequiresCookedData = false;
