CIS fix for bAutoStartupModule that had problems with creating a file when Intermediate directory wasn't created.

diff --git a/Engine/Source/Programs/UnrealBuildTool/Configuration/UEBuildTarget.cs b/Engine/Source/Programs/UnrealBuildTool/Configuration/UEBuildTarget.cs
--- a/Engine/Source/Programs/UnrealBuildTool/Configuration/UEBuildTarget.cs
+++ b/Engine/Source/Programs/UnrealBuildTool/Configuration/UEBuildTarget.cs
@@ -1825,6 +1825,11 @@ namespace UnrealBuildTool
 		{
 			var ModuleName = "AutoStartupModuleListGetter";
 
+			if (!Directory.Exists(GlobalCompileEnvironment.Config.OutputDirectory))
+			{
+				Directory.CreateDirectory(GlobalCompileEnvironment.Config.OutputDirectory);
+			}
+
 			string SourceFilename = Path.Combine(GlobalCompileEnvironment.Config.OutputDirectory, ModuleName + ".cpp");
 			string HeaderFilename = Path.Combine(GlobalCompileEnvironment.Config.OutputDirectory, ModuleName + ".h");
 
