CIS fix.

diff --git a/Engine/Source/Programs/UnrealBuildTool/Configuration/UEBuildTarget.cs b/Engine/Source/Programs/UnrealBuildTool/Configuration/UEBuildTarget.cs
--- a/Engine/Source/Programs/UnrealBuildTool/Configuration/UEBuildTarget.cs
+++ b/Engine/Source/Programs/UnrealBuildTool/Configuration/UEBuildTarget.cs
@@ -1824,14 +1824,15 @@ namespace UnrealBuildTool
 		private void CreateAutoStartupModuleListGetter()
 		{
 			var ModuleName = "AutoStartupModuleListGetter";
+			var ModuleDir = Path.Combine(GlobalCompileEnvironment.Config.OutputDirectory, ModuleName);
 
-			if (!Directory.Exists(GlobalCompileEnvironment.Config.OutputDirectory))
+			if (!Directory.Exists(ModuleDir))
 			{
-				Directory.CreateDirectory(GlobalCompileEnvironment.Config.OutputDirectory);
+				Directory.CreateDirectory(ModuleDir);
 			}
 
-			string SourceFilename = Path.Combine(GlobalCompileEnvironment.Config.OutputDirectory, ModuleName + ".cpp");
-			string HeaderFilename = Path.Combine(GlobalCompileEnvironment.Config.OutputDirectory, ModuleName + ".h");
+			string SourceFilename = Path.Combine(ModuleDir, ModuleName + ".cpp");
+			string HeaderFilename = Path.Combine(ModuleDir, ModuleName + ".h");
 
 			if (!File.Exists(HeaderFilename))
 			{
@@ -1864,7 +1865,7 @@ namespace UnrealBuildTool
 
 			BindArtificialModuleToBinary(
 				CreateArtificialModule(
-					ModuleName, GlobalCompileEnvironment.Config.OutputDirectory,
+					ModuleName, ModuleDir,
 					new FileItem[] { FileItem.GetItemByPath(SourceFilename) }, new string[] { "Core" }
 				), AppBinaries[0]);
 		}
