Fix UBT file dependency mismatch on case insensitive filesystems.
https://udn.unrealengine.com/questions/222717/unrealbuildtool-file-dependency-mismatch.html#answer-222848
https://answers.unrealengine.com/questions/108671/rebuilding-engine-when-creating-first-project.html

diff --git a/Engine/Source/Programs/UnrealBuildTool/System/ActionGraph.cs b/Engine/Source/Programs/UnrealBuildTool/System/ActionGraph.cs
--- a/Engine/Source/Programs/UnrealBuildTool/System/ActionGraph.cs
+++ b/Engine/Source/Programs/UnrealBuildTool/System/ActionGraph.cs
@@ -717,7 +717,7 @@ namespace UnrealBuildTool
 						string OldProducingCommandLine = "";
 						string NewProducingCommandLine = RootAction.CommandPath + " " + RootAction.CommandArguments;
 						if (!ActionHistory.GetProducingCommandLine(ProducedItem, out OldProducingCommandLine)
-						|| OldProducingCommandLine != NewProducingCommandLine)
+						|| !String.Equals(OldProducingCommandLine, NewProducingCommandLine, StringComparison.InvariantCultureIgnoreCase))
 						{
 							Log.TraceVerbose(
 								"{0}: Produced item \"{1}\" was produced by outdated command-line.\nOld command-line: {2}\nNew command-line: {3}",
diff --git a/Engine/Source/Programs/UnrealBuildTool/System/FileItem.cs b/Engine/Source/Programs/UnrealBuildTool/System/FileItem.cs
--- a/Engine/Source/Programs/UnrealBuildTool/System/FileItem.cs
+++ b/Engine/Source/Programs/UnrealBuildTool/System/FileItem.cs
@@ -278,7 +278,7 @@ namespace UnrealBuildTool
 			Directory.CreateDirectory(Path.GetDirectoryName(AbsolutePath));
 
 			// Only write the file if its contents have changed.
-			if (!File.Exists(AbsolutePath) || Utils.ReadAllText(AbsolutePath) != Contents)
+			if (!File.Exists(AbsolutePath) || !String.Equals(Utils.ReadAllText(AbsolutePath), Contents, StringComparison.InvariantCultureIgnoreCase))
 			{
 				File.WriteAllText(AbsolutePath, Contents, GetEncodingForString(Contents));
 			}
