Couple of fixes for Landscape split out.

diff --git a/Engine/Source/Programs/UnrealBuildTool/Configuration/UEBuildTarget.cs b/Engine/Source/Programs/UnrealBuildTool/Configuration/UEBuildTarget.cs
--- a/Engine/Source/Programs/UnrealBuildTool/Configuration/UEBuildTarget.cs
+++ b/Engine/Source/Programs/UnrealBuildTool/Configuration/UEBuildTarget.cs
@@ -1823,6 +1823,8 @@ namespace UnrealBuildTool
 		/// </summary>
 		private void CreateAutoStartupModuleListGetter()
 		{
+			var AutoStartupModuleArray = GetAutoStartupModuleList().ToArray();
+
 			var ModuleName = "AutoStartupModuleListGetter";
 			var ModuleDir = Path.Combine(GlobalCompileEnvironment.Config.OutputDirectory, ModuleName);
 
@@ -1843,18 +1845,21 @@ namespace UnrealBuildTool
 
 			SourceLines.Add(string.Format("#include \"{0}\"", ModuleName + ".h"));
 			SourceLines.Add("");
-			SourceLines.Add("#include \"Core.h\"");
-			SourceLines.Add("#include \"Array.h\"");
-			SourceLines.Add("#include \"UnrealString.h\"");
-			SourceLines.Add("");
-			SourceLines.Add("void GetAutoStartupModuleList(TArray<FString>& Out)");
+			SourceLines.Add("const char* EnumAutoStartupModuleName(int Index)");
 			SourceLines.Add("{");
 
-			foreach(var AutoStartupModuleName in GetAutoStartupModuleList())
+			SourceLines.Add("\tconst char* Names[] = {");
+
+			foreach (var AutoStartupModuleName in AutoStartupModuleArray)
 			{
-				SourceLines.Add(string.Format("\tOut.Add(TEXT(\"{0}\"));", AutoStartupModuleName));
+				SourceLines.Add(string.Format("\t\t\"{0}\",", AutoStartupModuleName));
 			}
 
+			SourceLines.Add("\t\tnullptr");
+			SourceLines.Add("\t};");
+			SourceLines.Add("");
+			SourceLines.Add("\treturn Names[Index];");
+
 			SourceLines.Add("}");
 			SourceLines.Add("");
 
@@ -1866,7 +1871,7 @@ namespace UnrealBuildTool
 			BindArtificialModuleToBinary(
 				CreateArtificialModule(
 					ModuleName, ModuleDir,
-					new FileItem[] { FileItem.GetItemByPath(SourceFilename) }, new string[] { "Core" }
+					new FileItem[] { FileItem.GetItemByPath(SourceFilename) }, new string[] { }
 				), AppBinaries[0]);
 		}
 
diff --git a/Engine/Source/Runtime/Engine/Public/StaticMeshResources.h b/Engine/Source/Runtime/Engine/Public/StaticMeshResources.h
--- a/Engine/Source/Runtime/Engine/Public/StaticMeshResources.h
+++ b/Engine/Source/Runtime/Engine/Public/StaticMeshResources.h
@@ -11,6 +11,8 @@
 #include "RawIndexBuffer.h"
 #include "TextureLayout3d.h"
 #include "LocalVertexFactory.h"
+#include "PrimitiveSceneProxy.h"
+#include "SceneManagement.h"
 
 /**
  * The LOD settings to use for a group of static meshes.
diff --git a/Engine/Source/Runtime/Landscape/Private/LandscapeCollision.cpp b/Engine/Source/Runtime/Landscape/Private/LandscapeCollision.cpp
--- a/Engine/Source/Runtime/Landscape/Private/LandscapeCollision.cpp
+++ b/Engine/Source/Runtime/Landscape/Private/LandscapeCollision.cpp
@@ -23,6 +23,10 @@ LandscapeCollision.cpp: Landscape collision
 #include "Engine/StaticMesh.h"
 #include "Interfaces/Interface_CollisionDataProvider.h"
 #include "AI/NavigationSystemHelpers.h"
+#include "AI/Navigation/NavigationSystem.h"
+#include "Engine/CollisionProfile.h"
+#include "Engine/Engine.h"
+#include "EngineGlobals.h"
 
 TMap<FGuid, ULandscapeHeightfieldCollisionComponent::FPhysXHeightfieldRef* > GSharedHeightfieldRefs;
 
diff --git a/Engine/Source/Runtime/Landscape/Private/LandscapeEdit.cpp b/Engine/Source/Runtime/Landscape/Private/LandscapeEdit.cpp
--- a/Engine/Source/Runtime/Landscape/Private/LandscapeEdit.cpp
+++ b/Engine/Source/Runtime/Landscape/Private/LandscapeEdit.cpp
@@ -5,6 +5,7 @@ LandscapeEdit.cpp: Landscape editing
 =============================================================================*/
 
 #include "Landscape.h"
+#include "LandscapeProxy.h"
 #include "Materials/MaterialInstanceDynamic.h"
 #include "Materials/MaterialExpressionLandscapeVisibilityMask.h"
 #include "Materials/MaterialExpressionLandscapeLayerWeight.h"
@@ -15,7 +16,6 @@ LandscapeEdit.cpp: Landscape editing
 #include "LandscapeEdit.h"
 #include "LandscapeRender.h"
 #include "LandscapeRenderMobile.h"
-#include "Landscape.h"
 #include "LandscapeInfo.h"
 #include "LandscapeLayerInfoObject.h"
 #include "LandscapeMaterialInstanceConstant.h"
@@ -33,6 +33,10 @@ LandscapeEdit.cpp: Landscape editing
 #include "RawMesh.h"
 #include "ScopedTransaction.h"
 #include "ImageWrapper.h"
+#include "Engine/Level.h"
+#include "EngineUtils.h"
+#include "Engine/Engine.h"
+#include "EngineGlobals.h"
 #endif
 #include "ComponentReregisterContext.h"
 
diff --git a/Engine/Source/Runtime/Landscape/Private/LandscapeRender.cpp b/Engine/Source/Runtime/Landscape/Private/LandscapeRender.cpp
--- a/Engine/Source/Runtime/Landscape/Private/LandscapeRender.cpp
+++ b/Engine/Source/Runtime/Landscape/Private/LandscapeRender.cpp
@@ -6,7 +6,6 @@ LandscapeRender.cpp: New terrain rendering
 
 #include "Landscape.h"
 
-#include "EngineGlobals.h"
 #include "Materials/MaterialExpressionTextureCoordinate.h"
 #include "Materials/MaterialExpressionLandscapeLayerCoords.h"
 #include "ShaderParameters.h"
@@ -20,6 +19,9 @@ LandscapeRender.cpp: New terrain rendering
 #include "RawIndexBuffer.h"
 #include "Engine/LevelStreaming.h"
 #include "Engine/ShadowMapTexture2D.h"
+#include "Engine/Engine.h"
+#include "EngineGlobals.h"
+#include "UnrealEngine.h"
 
 IMPLEMENT_UNIFORM_BUFFER_STRUCT(FLandscapeUniformShaderParameters,TEXT("LandscapeParameters"));
 
diff --git a/Engine/Source/Runtime/Launch/Private/LaunchEngineLoop.cpp b/Engine/Source/Runtime/Launch/Private/LaunchEngineLoop.cpp
--- a/Engine/Source/Runtime/Launch/Private/LaunchEngineLoop.cpp
+++ b/Engine/Source/Runtime/Launch/Private/LaunchEngineLoop.cpp
@@ -1787,14 +1787,31 @@ bool FEngineLoop::LoadStartupCoreModules()
 	return bSuccess;
 }
 
+/**
+ * Loads array of module names that needs to be auto-loaded.
+ *
+ * @param OutModules Output array to fill with module names.
+ */
+void GetAutoStartupModuleList(TArray<FString>& OutModules)
+{
+	// UBT generated function for listing all auto-startup module
+	// this binary is loading. If the list is finished then this function
+	// returns nullptr.
+	extern const ANSICHAR* EnumAutoStartupModuleName(int Index);
+
+	const ANSICHAR* NamePtr = nullptr;
+	int32 Index = 0;
+	while ((NamePtr = EnumAutoStartupModuleName(Index++)) != nullptr)
+	{
+		OutModules.Add(FString(NamePtr));
+	}
+}
+
 /**
  * Calls for each auto-startup module its StartupModule function.
  */
 void InitializeAutoStartupModules()
 {
-	// UBT generated function for listing all auto-startup module
-	// this binary is loading.
-	extern void GetAutoStartupModuleList(TArray<FString>& Out);
 	TArray<FString> Modules;
 	GetAutoStartupModuleList(Modules);
 
diff --git a/Engine/Source/Runtime/UE4Game/UE4Game.Build.cs b/Engine/Source/Runtime/UE4Game/UE4Game.Build.cs
--- a/Engine/Source/Runtime/UE4Game/UE4Game.Build.cs
+++ b/Engine/Source/Runtime/UE4Game/UE4Game.Build.cs
@@ -7,6 +7,7 @@ public class UE4Game : ModuleRules
 	public UE4Game(TargetInfo Target)
 	{
 		PrivateDependencyModuleNames.Add("Core");
+		PrivateDependencyModuleNames.Add("Landscape");
 	
 		if (Target.Platform == UnrealTargetPlatform.IOS)
 		{
