Enabled C++11 support on iOS

diff --git a/Engine/Source/Programs/UnrealBuildTool/IOS/IOSToolChain.cs b/Engine/Source/Programs/UnrealBuildTool/IOS/IOSToolChain.cs
--- a/Engine/Source/Programs/UnrealBuildTool/IOS/IOSToolChain.cs
+++ b/Engine/Source/Programs/UnrealBuildTool/IOS/IOSToolChain.cs
@@ -295,7 +295,8 @@ namespace UnrealBuildTool
 			Result += " -fno-rtti";
 			Result += " -fobjc-abi-version=2";
 			Result += " -fobjc-legacy-dispatch";
-			Result += " -std=c++0x";
+			Result += " -std=c++11";
+			Result += " -stdlib=libc++";
 			return Result;
 		}
 
@@ -306,7 +307,8 @@ namespace UnrealBuildTool
 			Result += " -fobjc-abi-version=2";
 			Result += " -fobjc-legacy-dispatch";
 			Result += " -fno-rtti";
-			Result += " -std=c++0x";
+			Result += " -std=c++11";
+			Result += " -stdlib=libc++";
 			return Result;
 		}
 
@@ -316,7 +318,8 @@ namespace UnrealBuildTool
 			Result += " -x objective-c";
 			Result += " -fobjc-abi-version=2";
 			Result += " -fobjc-legacy-dispatch";
-			Result += " -std=c++0x";
+			Result += " -std=c++11";
+			Result += " -stdlib=libc++";
 			return Result;
 		}
 
@@ -332,7 +335,8 @@ namespace UnrealBuildTool
 			string Result = "";
 			Result += " -x objective-c++-header";
 			Result += " -fno-rtti";
-			Result += " -std=c++0x";
+			Result += " -std=c++11";
+			Result += " -stdlib=libc++";
 			return Result;
 		}
 
