Initial support for SN-DBS distributed build system. Will be used for compiling for PS4 if SN-DBS is installed and Incredibuild isn't available/enabled.

diff --git a/Engine/Source/Programs/UnrealBuildTool/Configuration/BuildConfiguration.cs b/Engine/Source/Programs/UnrealBuildTool/Configuration/BuildConfiguration.cs
--- a/Engine/Source/Programs/UnrealBuildTool/Configuration/BuildConfiguration.cs
+++ b/Engine/Source/Programs/UnrealBuildTool/Configuration/BuildConfiguration.cs
@@ -170,6 +170,12 @@ namespace UnrealBuildTool
 		[XmlConfig]
 		public static bool bAllowRemotelyCompiledPCHs;
 
+        /// <summary>
+        /// Whether SN-DBS may be used.
+        /// </summary>
+        [XmlConfig]
+        public static bool bAllowSNDBS;
+
 		/// <summary>
 		/// Whether or not to delete outdated produced items.
 		/// </summary>
@@ -404,6 +410,7 @@ namespace UnrealBuildTool
 			bAllowLTCG = false;
 			bAllowRemotelyCompiledPCHs = false;
 			bAllowXGE = true;
+            bAllowSNDBS = true;
 
 			// Don't bother to check external (stable) headers for modification.  It slows down UBT's dependency checking.
 			bCheckExternalHeadersForModification = false;
@@ -595,6 +602,11 @@ namespace UnrealBuildTool
 			{
 				bAllowDistcc = false;
 			}
+			
+			if (!BuildPlatform.CanUseSNDBS()) 
+			{
+				bAllowSNDBS = false;
+			}
 		}
 	}
 }
diff --git a/Engine/Source/Programs/UnrealBuildTool/Configuration/UEBuildPlatform.cs b/Engine/Source/Programs/UnrealBuildTool/Configuration/UEBuildPlatform.cs
--- a/Engine/Source/Programs/UnrealBuildTool/Configuration/UEBuildPlatform.cs
+++ b/Engine/Source/Programs/UnrealBuildTool/Configuration/UEBuildPlatform.cs
@@ -31,6 +31,11 @@ namespace UnrealBuildTool
 		 */
 		bool CanUseDistcc();
 
+        /**
+         * If this platform can be compiled with SN-DBS
+         */
+        bool CanUseSNDBS();
+
 		/**
 		 * Register the platform with the UEBuildPlatform class
 		 */
@@ -442,6 +447,14 @@ namespace UnrealBuildTool
 			return false;
 		}
 
+        /**
+         *	If this platform can be compiled with SN-DBS
+         */
+        public virtual bool CanUseSNDBS()
+        {
+            return false;
+        }
+
 		/**
 		 *	Register the platform with the UEBuildPlatform class
 		 */
diff --git a/Engine/Source/Programs/UnrealBuildTool/System/ActionGraph.cs b/Engine/Source/Programs/UnrealBuildTool/System/ActionGraph.cs
--- a/Engine/Source/Programs/UnrealBuildTool/System/ActionGraph.cs
+++ b/Engine/Source/Programs/UnrealBuildTool/System/ActionGraph.cs
@@ -339,6 +339,18 @@ namespace UnrealBuildTool
 					bUsedXGE = true;
 				}
 
+                if (!bUsedXGE && BuildConfiguration.bAllowSNDBS)
+                {
+                    SNDBS.ExecutionResult SNDBSResult = SNDBS.ExecuteActions(ActionsToExecute);
+                    if( SNDBSResult != SNDBS.ExecutionResult.Unavailable )
+                    {
+                        ExecutorName = "SNDBS";
+                        Result = (SNDBSResult == SNDBS.ExecutionResult.TasksSucceeded);
+                        // don't do local compilation
+                        bUsedXGE = true;
+                     }
+                }
+
 				// If XGE is disallowed or unavailable, execute the commands locally.
 				if (!bUsedXGE)
 				{
diff --git a/Engine/Source/Programs/UnrealBuildTool/System/SNDBS.cs b/Engine/Source/Programs/UnrealBuildTool/System/SNDBS.cs
new file mode 100644
--- /dev/null
+++ b/Engine/Source/Programs/UnrealBuildTool/System/SNDBS.cs
@@ -0,0 +1,444 @@
+// Copyright 1998-2014 Epic Games, Inc. All Rights Reserved.
+
+using System;
+using System.Collections.Generic;
+using System.IO;
+using System.Diagnostics;
+using System.Xml;
+using System.Text.RegularExpressions;
+using System.Linq;
+using System.Reflection;
+using System.Threading;
+
+
+namespace UnrealBuildTool
+{
+    public class SNDBS
+    {
+        static private int MaxActionsToExecuteInParallel;
+        static private int JobNumber;
+
+        /** The possible result of executing tasks with SN-DBS. */
+        public enum ExecutionResult
+        {
+            Unavailable,
+            TasksFailed,
+            TasksSucceeded,
+        }
+
+        /**
+         * Used when debugging Actions outputs all action return values to debug out
+         * 
+         * @param	sender		Sending object
+         * @param	e			Event arguments (In this case, the line of string output)
+         */
+        static protected void ActionDebugOutput(object sender, DataReceivedEventArgs e)
+        {
+            var Output = e.Data;
+            if (Output == null)
+            {
+                return;
+            }
+
+            Log.TraceInformation(Output);
+        }
+
+        internal static ExecutionResult ExecuteLocalActions(List<Action> InLocalActions, Dictionary<Action, ActionThread> InActionThreadDictionary, int TotalNumJobs)
+        {
+            // Time to sleep after each iteration of the loop in order to not busy wait.
+            const float LoopSleepTime = 0.1f;
+
+            ExecutionResult LocalActionsResult = ExecutionResult.TasksSucceeded;
+
+            while (true)
+            {
+                // Count the number of pending and still executing actions.
+                int NumUnexecutedActions = 0;
+                int NumExecutingActions = 0;
+                foreach (Action Action in InLocalActions)
+                {
+                    ActionThread ActionThread = null;
+                    bool bFoundActionProcess = InActionThreadDictionary.TryGetValue(Action, out ActionThread);
+                    if (bFoundActionProcess == false)
+                    {
+                        NumUnexecutedActions++;
+                    }
+                    else if (ActionThread != null)
+                    {
+                        if (ActionThread.bComplete == false)
+                        {
+                            NumUnexecutedActions++;
+                            NumExecutingActions++;
+                        }
+                    }
+                }
+
+                // If there aren't any pending actions left, we're done executing.
+                if (NumUnexecutedActions == 0)
+                {
+                    break;
+                }
+
+                // If there are fewer actions executing than the maximum, look for pending actions that don't have any outdated
+                // prerequisites.
+                foreach (Action Action in InLocalActions)
+                {
+                    ActionThread ActionProcess = null;
+                    bool bFoundActionProcess = InActionThreadDictionary.TryGetValue(Action, out ActionProcess);
+                    if (bFoundActionProcess == false)
+                    {
+                        if (NumExecutingActions < Math.Max(1, MaxActionsToExecuteInParallel))
+                        {
+                            // Determine whether there are any prerequisites of the action that are outdated.
+                            bool bHasOutdatedPrerequisites = false;
+                            bool bHasFailedPrerequisites = false;
+                            foreach (FileItem PrerequisiteItem in Action.PrerequisiteItems)
+                            {
+                                if (PrerequisiteItem.ProducingAction != null && InLocalActions.Contains(PrerequisiteItem.ProducingAction))
+                                {
+                                    ActionThread PrerequisiteProcess = null;
+                                    bool bFoundPrerequisiteProcess = InActionThreadDictionary.TryGetValue(PrerequisiteItem.ProducingAction, out PrerequisiteProcess);
+                                    if (bFoundPrerequisiteProcess == true)
+                                    {
+                                        if (PrerequisiteProcess == null)
+                                        {
+                                            bHasFailedPrerequisites = true;
+                                        }
+                                        else if (PrerequisiteProcess.bComplete == false)
+                                        {
+                                            bHasOutdatedPrerequisites = true;
+                                        }
+                                        else if (PrerequisiteProcess.ExitCode != 0)
+                                        {
+                                            bHasFailedPrerequisites = true;
+                                        }
+                                    }
+                                    else
+                                    {
+                                        bHasOutdatedPrerequisites = true;
+                                    }
+                                }
+                            }
+
+                            // If there are any failed prerequisites of this action, don't execute it.
+                            if (bHasFailedPrerequisites)
+                            {
+                                // Add a null entry in the dictionary for this action.
+                                InActionThreadDictionary.Add(Action, null);
+                            }
+                            // If there aren't any outdated prerequisites of this action, execute it.
+                            else if (!bHasOutdatedPrerequisites)
+                            {
+                                ActionThread ActionThread = new ActionThread(Action, JobNumber, TotalNumJobs);
+                                ActionThread.Run();
+
+                                InActionThreadDictionary.Add(Action, ActionThread);
+
+                                NumExecutingActions++;
+                                JobNumber++;
+                            }
+                        }
+                    }
+                }
+
+                System.Threading.Thread.Sleep(TimeSpan.FromSeconds(LoopSleepTime));
+            }
+
+            return LocalActionsResult;
+        }
+
+        internal static ExecutionResult ExecuteActions(List<Action> InActions, Dictionary<Action, ActionThread> InActionThreadDictionary)
+        {
+            // Build the script file that will be executed by SN-DBS
+            StreamWriter ScriptFile;
+            string ScriptFilename = Path.Combine(BuildConfiguration.BaseIntermediatePath, "SNDBS.bat");
+
+            FileStream ScriptFileStream = new FileStream(ScriptFilename, FileMode.Create, FileAccess.ReadWrite, FileShare.Read);
+            ScriptFile = new StreamWriter(ScriptFileStream);
+            ScriptFile.AutoFlush = true;
+
+            int NumScriptedActions = 0;
+            List<Action> LocalActions = new List<Action>();
+            ActionThread DummyActionThread = new ActionThread(null, 1, 1);
+            foreach (Action Action in InActions)
+            {
+                ActionThread ActionProcess = null;
+                bool bFoundActionProcess = InActionThreadDictionary.TryGetValue(Action, out ActionProcess);
+                if (bFoundActionProcess == false)
+                {
+                    // Determine whether there are any prerequisites of the action that are outdated.
+                    bool bHasOutdatedPrerequisites = false;
+                    bool bHasFailedPrerequisites = false;
+                    foreach (FileItem PrerequisiteItem in Action.PrerequisiteItems)
+                    {
+                        if (PrerequisiteItem.ProducingAction != null && InActions.Contains(PrerequisiteItem.ProducingAction))
+                        {
+                            ActionThread PrerequisiteProcess = null;
+                            bool bFoundPrerequisiteProcess = InActionThreadDictionary.TryGetValue(PrerequisiteItem.ProducingAction, out PrerequisiteProcess);
+                            if (bFoundPrerequisiteProcess == true)
+                            {
+                                if (PrerequisiteProcess == null)
+                                {
+                                    bHasFailedPrerequisites = true;
+                                }
+                                else if (PrerequisiteProcess.bComplete == false)
+                                {
+                                    bHasOutdatedPrerequisites = true;
+                                }
+                                else if (PrerequisiteProcess.ExitCode != 0)
+                                {
+                                    bHasFailedPrerequisites = true;
+                                }
+                            }
+                            else
+                            {
+                                bHasOutdatedPrerequisites = true;
+                            }
+                        }
+                    }
+
+                    // If there are any failed prerequisites of this action, don't execute it.
+                    if (bHasFailedPrerequisites)
+                    {
+                        // Add a null entry in the dictionary for this action.
+                        InActionThreadDictionary.Add(Action, null);
+                    }
+                    // If there aren't any outdated prerequisites of this action, execute it.
+                    else if (!bHasOutdatedPrerequisites)
+                    {
+                        if (Action.bCanExecuteRemotely == false)
+                        {
+                            // Execute locally
+                            LocalActions.Add(Action);
+                        }
+                        else
+                        {
+                            // Add to script for execution by SN-DBS
+                            string NewCommandArguments = "\"" + Action.CommandPath + "\"" + " " + Action.CommandArguments;
+                            ScriptFile.WriteLine(ActionThread.ExpandEnvironmentVariables(NewCommandArguments));
+                            InActionThreadDictionary.Add(Action, DummyActionThread);
+                            Action.StartTime = Action.EndTime = DateTimeOffset.Now;
+                            Log.TraceInformation("[{0}/{1}] {2} {3}", JobNumber, InActions.Count, Action.CommandDescription, Action.StatusDescription);
+                            JobNumber++;
+                            NumScriptedActions++;
+                        }
+                    }
+                }
+            }
+
+            ScriptFile.Flush();
+            ScriptFile.Close();
+            ScriptFile.Dispose();
+            ScriptFile = null;
+
+            if( NumScriptedActions > 0 )
+            {
+                // Create the process
+                ProcessStartInfo PSI = new ProcessStartInfo("dbsbuild", "-p UE4 -s " + BuildConfiguration.BaseIntermediatePath + "/sndbs.bat");
+                PSI.RedirectStandardOutput = true;
+                PSI.RedirectStandardError = true;
+                PSI.UseShellExecute = false;
+                PSI.CreateNoWindow = true;
+                PSI.WorkingDirectory = Path.GetFullPath("."); ;
+                Process NewProcess = new Process();
+                NewProcess.StartInfo = PSI;
+                NewProcess.OutputDataReceived += new DataReceivedEventHandler(ActionDebugOutput);
+                NewProcess.ErrorDataReceived += new DataReceivedEventHandler(ActionDebugOutput);
+                DateTimeOffset StartTime = DateTimeOffset.Now;
+                NewProcess.Start();
+                NewProcess.BeginOutputReadLine();
+                NewProcess.BeginErrorReadLine();
+
+                NewProcess.WaitForExit();
+
+                TimeSpan Duration;
+                DateTimeOffset EndTime = DateTimeOffset.Now;
+                if (EndTime == DateTimeOffset.MinValue)
+                {
+                    Duration = DateTimeOffset.Now - StartTime;
+                }
+                else
+                {
+                    Duration = EndTime - StartTime;
+                }
+
+                DummyActionThread.bComplete = true;
+                int ExitCode = NewProcess.ExitCode;
+                if (ExitCode != 0)
+                {
+                    return ExecutionResult.TasksFailed;
+                }
+                else
+                {
+                    UnrealBuildTool.TotalCompileTime += Duration.TotalSeconds;
+                }
+            }
+
+            // Execute local tasks
+            if( LocalActions.Count > 0 )
+            {
+                return ExecuteLocalActions(LocalActions, InActionThreadDictionary, InActions.Count);
+            }
+
+            return ExecutionResult.TasksSucceeded;
+        }
+
+        public static ExecutionResult ExecuteActions(List<Action> Actions)
+        {
+            ExecutionResult SNDBSResult = ExecutionResult.TasksSucceeded;
+            if (Actions.Count > 0)
+            {
+                string SCERoot = Environment.GetEnvironmentVariable("SCE_ROOT_DIR");
+                string SNDBSExecutable = Path.Combine(SCERoot, "Common/SN-DBS/bin/dbsbuild.exe");
+
+                // Check that SN-DBS is available
+                bool bSNDBSExists = File.Exists(SNDBSExecutable);
+                if( bSNDBSExists == false )
+                {
+                    return ExecutionResult.Unavailable;
+                }
+
+                // Use WMI to figure out physical cores, excluding hyper threading.
+                int NumCores = 0;
+                if (!Utils.IsRunningOnMono)
+                {
+                    try
+                    {
+                        using (var Mos = new System.Management.ManagementObjectSearcher("Select * from Win32_Processor"))
+                        {
+                            var MosCollection = Mos.Get();
+                            foreach (var Item in MosCollection)
+                            {
+                                NumCores += int.Parse(Item["NumberOfCores"].ToString());
+                            }
+                        }
+                    }
+                    catch (Exception Ex)
+                    {
+                        Log.TraceWarning("Unable to get the number of Cores: {0}", Ex.ToString());
+                        Log.TraceWarning("Falling back to processor count.");
+                    }
+                }
+                // On some systems this requires a hot fix to work so we fall back to using the (logical) processor count.
+                if (NumCores == 0)
+                {
+                    NumCores = System.Environment.ProcessorCount;
+                }
+                // The number of actions to execute in parallel is trying to keep the CPU busy enough in presence of I/O stalls.
+                MaxActionsToExecuteInParallel = 0;
+                // The CPU has more logical cores than physical ones, aka uses hyper-threading. 
+                if (NumCores < System.Environment.ProcessorCount)
+                {
+                    MaxActionsToExecuteInParallel = (int)(NumCores * BuildConfiguration.ProcessorCountMultiplier);
+                }
+                // No hyper-threading. Only kicking off a task per CPU to keep machine responsive.
+                else
+                {
+                    MaxActionsToExecuteInParallel = NumCores;
+                }
+                MaxActionsToExecuteInParallel = Math.Min(MaxActionsToExecuteInParallel, BuildConfiguration.MaxProcessorCount);
+
+                JobNumber = 1;
+                Dictionary<Action, ActionThread> ActionThreadDictionary = new Dictionary<Action, ActionThread>();
+
+                while( true )
+                {
+                    bool bUnexecutedActions = false;
+                    foreach (Action Action in Actions)
+                    {
+                        ActionThread ActionThread = null;
+                        bool bFoundActionProcess = ActionThreadDictionary.TryGetValue(Action, out ActionThread);
+                        if (bFoundActionProcess == false)
+                        {
+                            bUnexecutedActions = true;
+                            ExecutionResult CompileResult = ExecuteActions(Actions, ActionThreadDictionary);
+                            if (CompileResult != ExecutionResult.TasksSucceeded)
+                            {
+                                return ExecutionResult.TasksFailed;
+                            }
+                            break;
+                        }
+                    }
+
+                    if (bUnexecutedActions == false)
+                    {
+                        break;
+                    }
+                }
+
+                Log.WriteLineIf(BuildConfiguration.bLogDetailedActionStats, TraceEventType.Information, "-------- Begin Detailed Action Stats ----------------------------------------------------------");
+                Log.WriteLineIf(BuildConfiguration.bLogDetailedActionStats, TraceEventType.Information, "^Action Type^Duration (seconds)^Tool^Task^Using PCH");
+
+                double TotalThreadSeconds = 0;
+
+                // Check whether any of the tasks failed and log action stats if wanted.
+                foreach (KeyValuePair<Action, ActionThread> ActionProcess in ActionThreadDictionary)
+                {
+                    Action Action = ActionProcess.Key;
+                    ActionThread ActionThread = ActionProcess.Value;
+
+                    // Check for pending actions, preemptive failure
+                    if (ActionThread == null)
+                    {
+                        SNDBSResult = ExecutionResult.TasksFailed;
+                        continue;
+                    }
+                    // Check for executed action but general failure
+                    if (ActionThread.ExitCode != 0)
+                    {
+                        SNDBSResult = ExecutionResult.TasksFailed;
+                    }
+                    // Log CPU time, tool and task.
+                    double ThreadSeconds = Action.Duration.TotalSeconds;
+
+                    Log.WriteLineIf(BuildConfiguration.bLogDetailedActionStats,
+                        TraceEventType.Information,
+                        "^{0}^{1:0.00}^{2}^{3}^{4}",
+                        Action.ActionType.ToString(),
+                        ThreadSeconds,
+                        Path.GetFileName(Action.CommandPath),
+                        Action.StatusDescription,
+                        Action.bIsUsingPCH);
+
+                    // Update statistics
+                    switch (Action.ActionType)
+                    {
+                        case ActionType.BuildProject:
+                            UnrealBuildTool.TotalBuildProjectTime += ThreadSeconds;
+                            break;
+
+                        case ActionType.Compile:
+                            UnrealBuildTool.TotalCompileTime += ThreadSeconds;
+                            break;
+
+                        case ActionType.CreateAppBundle:
+                            UnrealBuildTool.TotalCreateAppBundleTime += ThreadSeconds;
+                            break;
+
+                        case ActionType.GenerateDebugInfo:
+                            UnrealBuildTool.TotalGenerateDebugInfoTime += ThreadSeconds;
+                            break;
+
+                        case ActionType.Link:
+                            UnrealBuildTool.TotalLinkTime += ThreadSeconds;
+                            break;
+
+                        default:
+                            UnrealBuildTool.TotalOtherActionsTime += ThreadSeconds;
+                            break;
+                    }
+
+                    // Keep track of total thread seconds spent on tasks.
+                    TotalThreadSeconds += ThreadSeconds;
+                }
+
+                Log.TraceInformation("-------- End Detailed Actions Stats -----------------------------------------------------------");
+
+                // Log total CPU seconds and numbers of processors involved in tasks.
+                Log.WriteLineIf(BuildConfiguration.bLogDetailedActionStats || BuildConfiguration.bPrintDebugInfo,
+                    TraceEventType.Information, "Cumulative thread seconds ({0} processors): {1:0.00}", System.Environment.ProcessorCount, TotalThreadSeconds);
+            }
+            return SNDBSResult;
+        }
+    }
+}
diff --git a/Engine/Source/Programs/UnrealBuildTool/System/UnrealBuildTool.cs b/Engine/Source/Programs/UnrealBuildTool/System/UnrealBuildTool.cs
--- a/Engine/Source/Programs/UnrealBuildTool/System/UnrealBuildTool.cs
+++ b/Engine/Source/Programs/UnrealBuildTool/System/UnrealBuildTool.cs
@@ -1754,7 +1754,7 @@ namespace UnrealBuildTool
 
             // Figure out how long we took to execute.
             double BuildDuration = (DateTime.UtcNow - StartTime - MutexWaitTime).TotalSeconds;
-            if (ExecutorName == "Local" || ExecutorName == "Distcc")
+            if (ExecutorName == "Local" || ExecutorName == "Distcc" || ExecutorName == "SNDBS")
             {
                 Log.TraceInformation("Cumulative action seconds ({0} processors): {1:0.00} building projects, {2:0.00} compiling, {3:0.00} creating app bundles, {4:0.00} generating debug info, {5:0.00} linking, {6:0.00} other",
                     System.Environment.ProcessorCount,
diff --git a/Engine/Source/Programs/UnrealBuildTool/UnrealBuildTool.csproj b/Engine/Source/Programs/UnrealBuildTool/UnrealBuildTool.csproj
--- a/Engine/Source/Programs/UnrealBuildTool/UnrealBuildTool.csproj
+++ b/Engine/Source/Programs/UnrealBuildTool/UnrealBuildTool.csproj
@@ -225,6 +225,7 @@
     <Compile Include="System\Distcc.cs" />
     <Compile Include="System\ModuleProcessingException.cs" />
     <Compile Include="System\MakefileGenerator.cs" />
+    <Compile Include="System\SNDBS.cs" />
     <Compile Include="System\VCSolutionOptions.cs" />
     <Compile Include="System\DirectoryLookupCache.cs" />
     <Compile Include="System\CPPEnvironment.cs" />
