Add missing file reference to mono project.

diff --git a/Engine/Source/Programs/UnrealBuildTool/UnrealBuildTool_Mono.csproj b/Engine/Source/Programs/UnrealBuildTool/UnrealBuildTool_Mono.csproj
--- a/Engine/Source/Programs/UnrealBuildTool/UnrealBuildTool_Mono.csproj
+++ b/Engine/Source/Programs/UnrealBuildTool/UnrealBuildTool_Mono.csproj
@@ -144,6 +144,7 @@
     <Compile Include="System\LinkEnvironment.cs" />
     <Compile Include="System\LocalExecutor.cs" />
     <Compile Include="System\MakefileGenerator.cs" />
+    <Compile Include="System\SNDBS.cs" />
     <Compile Include="System\Plugins.cs" />
     <Compile Include="System\Project.cs" />
     <Compile Include="System\ProjectFileGenerator.cs" />
