Fixing hot-reload errors (missing includes).

#change Making sure AutoStartupModuleListGetter is not being created when doing hot-reload as it's not building the entire target

diff --git a/Engine/Source/Programs/UnrealBuildTool/Configuration/UEBuildTarget.cs b/Engine/Source/Programs/UnrealBuildTool/Configuration/UEBuildTarget.cs
--- a/Engine/Source/Programs/UnrealBuildTool/Configuration/UEBuildTarget.cs
+++ b/Engine/Source/Programs/UnrealBuildTool/Configuration/UEBuildTarget.cs
@@ -1508,7 +1508,10 @@ namespace UnrealBuildTool
 				CreateLinkerFixupsCPPFile();
 			}
 
-			CreateAutoStartupModuleListGetter();
+			if (!bEditorRecompile)
+			{
+				CreateAutoStartupModuleListGetter();
+			}
 
 			// Build the target's binaries.
 			foreach (var Binary in AppBinaries)
