Fixed remote building of Mac version from Visual Studio

diff --git a/Engine/Source/Programs/UnrealBuildTool/Mac/MacToolChain.cs b/Engine/Source/Programs/UnrealBuildTool/Mac/MacToolChain.cs
--- a/Engine/Source/Programs/UnrealBuildTool/Mac/MacToolChain.cs
+++ b/Engine/Source/Programs/UnrealBuildTool/Mac/MacToolChain.cs
@@ -959,6 +959,8 @@ namespace UnrealBuildTool
 						CustomResourcesPath = Path.GetDirectoryName(UProjectFilePath) + "/Source/" + GameName + "/Resources/Mac";
 					}
 
+					CustomResourcesPath = ConvertPath(CustomResourcesPath);
+
 					// Copy resources
 					string CustomIcon = CustomResourcesPath + "/" + GameName + ".icns";
 					if (!File.Exists(CustomIcon))
diff --git a/Engine/Source/Programs/UnrealBuildTool/Mac/UEDeployMac.cs b/Engine/Source/Programs/UnrealBuildTool/Mac/UEDeployMac.cs
--- a/Engine/Source/Programs/UnrealBuildTool/Mac/UEDeployMac.cs
+++ b/Engine/Source/Programs/UnrealBuildTool/Mac/UEDeployMac.cs
@@ -35,7 +35,7 @@ namespace UnrealBuildTool.IOS
 			MacToolChain Toolchain = UEToolChain.GetPlatformToolChain(CPPTargetPlatform.Mac) as MacToolChain;
 			
 			string FixDylibDepsScript = Path.Combine(IntermediateDirectory, "FixDylibDependencies.sh");
-			string CreateAppBundleScript = Path.Combine(IntermediateDirectory, "CreateAppBundle.sh");
+			string FinalizeAppBundleScript = Path.Combine(IntermediateDirectory, "FinalizeAppBundle.sh");
 
 			string RemoteWorkingDir = "";
 
@@ -52,9 +52,9 @@ namespace UnrealBuildTool.IOS
 
 					if (!InTarget.GlobalLinkEnvironment.Config.bIsBuildingConsoleApplication)
 					{
-						string RemoteCreateAppBundleScript = Toolchain.ConvertPath(Path.GetFullPath(CreateAppBundleScript));
-						RemoteCreateAppBundleScript = RemoteCreateAppBundleScript.Replace("../../../../", "../../");
-						RPCUtilHelper.CopyFile("../../" + CreateAppBundleScript, RemoteCreateAppBundleScript, true);
+						string RemoteFinalizeAppBundleScript = Toolchain.ConvertPath(Path.GetFullPath(FinalizeAppBundleScript));
+						RemoteFinalizeAppBundleScript = RemoteFinalizeAppBundleScript.Replace("../../../../", "../../");
+						RPCUtilHelper.CopyFile("../../" + FinalizeAppBundleScript, RemoteFinalizeAppBundleScript, true);
 					}
 
 
@@ -74,8 +74,8 @@ namespace UnrealBuildTool.IOS
 
 					if (!InTarget.GlobalLinkEnvironment.Config.bIsBuildingConsoleApplication)
 					{
-						Log.TraceInformation("Running CreateAppBundle.sh...");
-						Results = RPCUtilHelper.Command(RemoteWorkingDir, "/bin/sh", "CreateAppBundle.sh", null);
+						Log.TraceInformation("Running FinalizeAppBundle.sh...");
+						Results = RPCUtilHelper.Command(RemoteWorkingDir, "/bin/sh", "FinalizeAppBundle.sh", null);
 						if (Results != null)
 						{
 							string Result = (string)Results["CommandOutput"];
