Fixed Unix line endings in UBT's UEBuildTarget.CreateAutoStartupModuleListGetter(). Fixes a problem with UBT rebuilding AutoStartupModuleListGetter.cpp every time.

diff --git a/Engine/Source/Programs/UnrealBuildTool/Configuration/UEBuildTarget.cs b/Engine/Source/Programs/UnrealBuildTool/Configuration/UEBuildTarget.cs
--- a/Engine/Source/Programs/UnrealBuildTool/Configuration/UEBuildTarget.cs
+++ b/Engine/Source/Programs/UnrealBuildTool/Configuration/UEBuildTarget.cs
@@ -1866,7 +1866,8 @@ namespace UnrealBuildTool
 			SourceLines.Add("}");
 			SourceLines.Add("");
 
-			if (!File.Exists(SourceFilename) || File.ReadAllText(SourceFilename).Trim() != string.Join("\r\n", SourceLines).Trim())
+			string LineEnd = Environment.OSVersion.Platform == PlatformID.Unix ? "\n" : "\r\n";
+			if (!File.Exists(SourceFilename) || File.ReadAllText(SourceFilename).Trim() != string.Join(LineEnd, SourceLines).Trim())
 			{
 				File.WriteAllLines(SourceFilename, SourceLines);
 			}
