Prevent compilation attempts that are known to fail.

- No clang 3.4.x nor gcc on Linux.

diff --git a/Engine/Source/Programs/UnrealBuildTool/Linux/LinuxToolChain.cs b/Engine/Source/Programs/UnrealBuildTool/Linux/LinuxToolChain.cs
--- a/Engine/Source/Programs/UnrealBuildTool/Linux/LinuxToolChain.cs
+++ b/Engine/Source/Programs/UnrealBuildTool/Linux/LinuxToolChain.cs
@@ -166,7 +166,20 @@ namespace UnrealBuildTool
 
 			if (!DetermineCompilerVersion())
 			{
-				Console.WriteLine("Could not determine version of the compiler, not registering Linux toolchain");
+				Console.WriteLine("\n*** Could not determine version of the compiler, not registering Linux toolchain.\n");
+				return;
+			}
+
+			// refuse to use compilers that we know won't work
+			// disable that only if you are a dev and you know what you are doing
+			if (!UsingClang())
+			{
+				Console.WriteLine("\n*** This version of the engine can only be compiled by clang - refusing to register the Linux toolchain.\n");
+				return;
+			}
+			else if (CompilerVersionMajor == 3 && CompilerVersionMinor == 4)
+			{
+				Console.WriteLine("\n*** clang 3.4.x is known to miscompile the engine - refusing to register the Linux toolchain.\n");
 				return;
 			}
 
