using the wrong standard lib with IOS
#ios

diff --git a/Engine/Source/Programs/UnrealBuildTool/IOS/IOSToolChain.cs b/Engine/Source/Programs/UnrealBuildTool/IOS/IOSToolChain.cs
--- a/Engine/Source/Programs/UnrealBuildTool/IOS/IOSToolChain.cs
+++ b/Engine/Source/Programs/UnrealBuildTool/IOS/IOSToolChain.cs
@@ -296,7 +296,7 @@ namespace UnrealBuildTool
 			Result += " -fobjc-abi-version=2";
 			Result += " -fobjc-legacy-dispatch";
 			Result += " -std=c++11";
-			Result += " -stdlib=libc++";
+			Result += " -stdlib=libstdc++";
 			return Result;
 		}
 
@@ -308,7 +308,7 @@ namespace UnrealBuildTool
 			Result += " -fobjc-legacy-dispatch";
 			Result += " -fno-rtti";
 			Result += " -std=c++11";
-			Result += " -stdlib=libc++";
+			Result += " -stdlib=libstdc++";
 			return Result;
 		}
 
@@ -319,7 +319,7 @@ namespace UnrealBuildTool
 			Result += " -fobjc-abi-version=2";
 			Result += " -fobjc-legacy-dispatch";
 			Result += " -std=c++11";
-			Result += " -stdlib=libc++";
+			Result += " -stdlib=libstdc++";
 			return Result;
 		}
 
@@ -336,7 +336,7 @@ namespace UnrealBuildTool
 			Result += " -x objective-c++-header";
 			Result += " -fno-rtti";
 			Result += " -std=c++11";
-			Result += " -stdlib=libc++";
+			Result += " -stdlib=libstdc++";
 			return Result;
 		}
 
