addition of iOS min version and devices supported to iOS Project Settings
UE-3315, UEPLAT-183, UEPLAT-28
#ios

diff --git a/Engine/Source/Developer/iOS/IOSPlatformEditor/Private/IOSTargetSettingsCustomization.cpp b/Engine/Source/Developer/iOS/IOSPlatformEditor/Private/IOSTargetSettingsCustomization.cpp
--- a/Engine/Source/Developer/iOS/IOSPlatformEditor/Private/IOSTargetSettingsCustomization.cpp
+++ b/Engine/Source/Developer/iOS/IOSPlatformEditor/Private/IOSTargetSettingsCustomization.cpp
@@ -74,6 +74,9 @@ void FIOSTargetSettingsCustomization::BuildPListSection(IDetailLayoutBuilder& De
 	IDetailCategoryBuilder& BundleCategory = DetailLayout.EditCategory(TEXT("Bundle Information"));
 	IDetailCategoryBuilder& OrientationCategory = DetailLayout.EditCategory(TEXT("Orientation"));
 	IDetailCategoryBuilder& RenderCategory = DetailLayout.EditCategory(TEXT("Rendering"));
+	IDetailCategoryBuilder& OSInfoCategory = DetailLayout.EditCategory(TEXT("OS Info"));
+	IDetailCategoryBuilder& DeviceCategory = DetailLayout.EditCategory(TEXT("Devices"));
+	IDetailCategoryBuilder& CookCategory = DetailLayout.EditCategory(TEXT("Cook Settings"));
 
 	TSharedRef<SPlatformSetupMessage> PlatformSetupMessage = SNew(SPlatformSetupMessage, GameInfoPath)
 		.PlatformName(LOCTEXT("iOSPlatformName", "iOS"))
@@ -151,6 +154,11 @@ void FIOSTargetSettingsCustomization::BuildPListSection(IDetailLayoutBuilder& De
 	SETUP_PLIST_PROP(bSupportsMetal, RenderCategory, TEXT("Whether or not to add support for Metal API (requires IOS8 and A7 processors)."));
 	SETUP_PLIST_PROP(bSupportsOpenGLES2, RenderCategory, TEXT("Whether or not to add support for OpenGL ES2 (if this is false, then your game should specify minimum IOS8 version and use \"metal\" instead of \"opengles-2\" in UIRequiredDeviceCapabilities)"));
 
+	SETUP_PLIST_PROP(bSupportsIPad, DeviceCategory, TEXT("Whether or not to add support for iPad devices"));
+	SETUP_PLIST_PROP(bSupportsIPhone, DeviceCategory, TEXT("Whether or not to add support for iPhone devices"));
+
+	SETUP_PLIST_PROP(MinimumiOSVersion, OSInfoCategory, TEXT("WMinimum iOS version this game supports"));
+
 #undef SETUP_PLIST_PROP
 }
 
@@ -292,6 +300,38 @@ void FIOSTargetSettingsCustomization::OnPlistPropertyModified()
 	DeviceCapsArrayBody += TEXT("\t");
 	Updater.ReplaceKey(RequiredDeviceCaps, ClosingArray, DeviceCapsArrayBody);
 
+	// build the replacement device families
+	const FString DeviceFamilyKey(TEXT("<key>UIDeviceFamily</key>"));
+	FString FamilyKeyBody = TEXT("\n\t<array>\n");
+	// automatically add armv7 for now
+	if (Settings.bSupportsIPhone)
+	{
+		FamilyKeyBody += TEXT("\t\t<integer>1</integer>\n");
+	}
+	if (Settings.bSupportsIPad)
+	{
+		FamilyKeyBody += TEXT("\t\t<integer>2</integer>\n");
+	}
+	FamilyKeyBody += TEXT("\t");
+	Updater.ReplaceKey(DeviceFamilyKey, ClosingArray, FamilyKeyBody);
+
+	// build the replacement min iOS version
+	const FString MiniOSVersionKey(TEXT("<key>MinimumOSVersion</key>"));
+	FString iOSVersionBody = TEXT("\n\t<string>");
+	switch (Settings.MinimumiOSVersion)
+	{
+	case EIOSVersion::IOS_6:
+		iOSVersionBody += TEXT("6.0");
+		break;
+	case EIOSVersion::IOS_7:
+		iOSVersionBody += TEXT("7.0");
+		break;
+	case EIOSVersion::IOS_8:
+		iOSVersionBody += TEXT("8.0");
+		break;
+	}
+	Updater.ReplaceKey(MiniOSVersionKey, ClosingString, iOSVersionBody);
+
 	// Write out the updated .plist
 	Updater.Finalize(GameInfoPath, true, FFileHelper::EEncodingOptions::ForceUTF8);
 }
diff --git a/Engine/Source/Programs/AutomationTool/IOS/IOSPlatform.Automation.cs b/Engine/Source/Programs/AutomationTool/IOS/IOSPlatform.Automation.cs
--- a/Engine/Source/Programs/AutomationTool/IOS/IOSPlatform.Automation.cs
+++ b/Engine/Source/Programs/AutomationTool/IOS/IOSPlatform.Automation.cs
@@ -306,7 +306,7 @@ public class IOSPlatform : Platform
 		bWasGenerated = false;
 		string XcodeProj = RawProjectPath.Replace(".uproject", "_IOS.xcodeproj");
 		Console.WriteLine ("Project: " + XcodeProj);
-		if (!Directory.Exists (XcodeProj))
+		//		if (!Directory.Exists (XcodeProj))
 		{
 			// project.xcodeproj doesn't exist, so generate temp project
 			string Arguments = "-project=\"" + RawProjectPath + "\"";
diff --git a/Engine/Source/Programs/UnrealBuildTool/IOS/IOSToolChain.cs b/Engine/Source/Programs/UnrealBuildTool/IOS/IOSToolChain.cs
--- a/Engine/Source/Programs/UnrealBuildTool/IOS/IOSToolChain.cs
+++ b/Engine/Source/Programs/UnrealBuildTool/IOS/IOSToolChain.cs
@@ -47,14 +47,20 @@ namespace UnrealBuildTool
 		[XmlConfig]
 		public static string ShippingArchitectures = "armv7";
 
-		/** Which version of the iOS to allow at run time */
+		/** Which version of the iOS to allow at build time */
 		[XmlConfig]
-		public static string IOSVersion = "6.0";
+		public static string BuildIOSVersion = "6.0";
+
+		/** Which version of the iOS to allow at run time */
+		public static string RunTimeIOSVersion = "6.0";
 
 		/** Which developer directory to root from */
 		[XmlConfig]
 		public static string XcodeDeveloperDir = "/Applications/Xcode.app/Contents/Developer/";
 
+		/** which devices the game is allowed to run on */
+		public static string RunTimeIOSDevices = "1,2";
+
 		/** Location of the SDKs */
 		private static string BaseSDKDir;
 		private static string BaseSDKDirSim;
@@ -83,11 +89,52 @@ namespace UnrealBuildTool
 			BaseSDKDirSim = XcodeDeveloperDir + "Platforms/iPhoneSimulator.platform/Developer/SDKs";
 		}
 
+		public static void ParseProjectSettings()
+		{
+			// look in ini settings for what platforms to compile for
+			ConfigCacheIni Ini = new ConfigCacheIni(UnrealTargetPlatform.IOS, "Engine", UnrealBuildTool.GetUProjectPath());
+			string MinVersion = "IOS_6";
+			if (Ini.GetString("/Script/IOSRuntimeSettings.IOSRuntimeSettings", "MinimumiOSVersion", out MinVersion))
+			{
+				switch (MinVersion)
+				{
+					case "IOS_6":
+						RunTimeIOSVersion = "6.0";
+						break;
+					case "IOS_7":
+						RunTimeIOSVersion = "7.0";
+						break;
+					case "IOS_8":
+						RunTimeIOSVersion = "8.0";
+						break;
+				}
+			}
+
+			bool biPhoneAllowed = true;
+			bool biPadAllowed = true;
+			Ini.GetBool ("/Script/IOSRuntimeSettings.IOSRuntimeSettings", "bSupportsIPhone", out biPhoneAllowed);
+			Ini.GetBool ("/Script/IOSRuntimeSettings.IOSRuntimeSettings", "bSupportsIPad", out biPadAllowed);
+			if (biPhoneAllowed && biPadAllowed)
+			{
+				RunTimeIOSDevices = "1,2";
+			}
+			else if (biPadAllowed)
+			{
+				RunTimeIOSDevices = "2";
+			}
+			else if (biPhoneAllowed)
+			{
+				RunTimeIOSDevices = "1";
+			}
+		}
+
 		/** Hunt down the latest IOS sdk if desired */
 		public override void SetUpGlobalEnvironment()
 		{
 			base.SetUpGlobalEnvironment();
 
+			ParseProjectSettings();
+
 			if (IOSSDKVersion == "latest")
 			{
 				try
@@ -251,7 +298,7 @@ namespace UnrealBuildTool
 				Result += " -isysroot " + BaseSDKDir + "/iPhoneOS" + IOSSDKVersion + ".sdk";
 			}
 
-			Result += " -miphoneos-version-min=" + IOSVersion;
+			Result += " -miphoneos-version-min=" + BuildIOSVersion;
 
 			// Optimize non- debug builds.
 			if (CompileEnvironment.Config.Target.Configuration != CPPTargetConfiguration.Debug)
@@ -416,7 +463,7 @@ namespace UnrealBuildTool
 				Result += " -isysroot " + XcodeDeveloperDir + "Platforms/iPhoneOS.platform/Developer/SDKs/iPhoneOS" + IOSSDKVersion + ".sdk";
 			}
 			Result += " -dead_strip";
-			Result += " -miphoneos-version-min=" + IOSVersion;
+			Result += " -miphoneos-version-min=" + BuildIOSVersion;
 			Result += " -Wl,-no_pie";
 			//			Result += " -v";
 
@@ -1319,7 +1366,14 @@ namespace UnrealBuildTool
 
 		public override string GetPlatformVersion()
 		{
-			return IOSVersion;
+			ParseProjectSettings();
+			return RunTimeIOSVersion;
+		}
+
+		public override string GetPlatformDevices()
+		{
+			ParseProjectSettings();
+			return RunTimeIOSDevices;
 		}
 
 		public static int RunExecutableAndWait( string ExeName, string ArgumentList, out string StdOutResults )
diff --git a/Engine/Source/Programs/UnrealBuildTool/System/XcodeProjectFileGenerator.cs b/Engine/Source/Programs/UnrealBuildTool/System/XcodeProjectFileGenerator.cs
--- a/Engine/Source/Programs/UnrealBuildTool/System/XcodeProjectFileGenerator.cs
+++ b/Engine/Source/Programs/UnrealBuildTool/System/XcodeProjectFileGenerator.cs
@@ -490,7 +490,7 @@ namespace UnrealBuildTool
 					"\t\t\t\t\"PRODUCT_NAME[sdk=iphonesimulator*]\" = \"$(TARGET_NAME)-simulator\";" + ProjectFileGenerator.NewLine +
 					"\t\t\t\tMACOSX_DEPLOYMENT_TARGET = 10.9;" + ProjectFileGenerator.NewLine +
 					"\t\t\t\tIPHONEOS_DEPLOYMENT_TARGET = " + Toolchain.GetPlatformVersion() + ";" + ProjectFileGenerator.NewLine +
-					"\t\t\t\tTARGETED_DEVICE_FAMILY = \"1,2\";" + ProjectFileGenerator.NewLine +
+					"\t\t\t\tTARGETED_DEVICE_FAMILY = \"" + Toolchain.GetPlatformDevices() + "\";" + ProjectFileGenerator.NewLine +
 					"\t\t\t\tCONFIGURATION_BUILD_DIR = \"Engine/Binaries/IOS/Payload\";" + ProjectFileGenerator.NewLine +
 					"\t\t\t\tSDKROOT = macosx;" + ProjectFileGenerator.NewLine +
 					"\t\t\t\t\"SDKROOT[arch=x86_64]\" = macosx;" + ProjectFileGenerator.NewLine +
@@ -516,7 +516,7 @@ namespace UnrealBuildTool
 				"\t\t\t\tSUPPORTED_PLATFORMS = \"iphoneos iphonesimulator\";" + ProjectFileGenerator.NewLine +
 				"\t\t\t\tCONFIGURATION_BUILD_DIR = \"Engine/Binaries/IOS/Payload\";" + ProjectFileGenerator.NewLine +
 				"\t\t\t\tIPHONEOS_DEPLOYMENT_TARGET = " + Toolchain.GetPlatformVersion() + ";" + ProjectFileGenerator.NewLine +
-				"\t\t\t\tTARGETED_DEVICE_FAMILY = \"1,2\";" + ProjectFileGenerator.NewLine +
+				"\t\t\t\tTARGETED_DEVICE_FAMILY = \"" + Toolchain.GetPlatformDevices() + "\";" + ProjectFileGenerator.NewLine +
 				"\t\t\t\tSDKROOT = iphoneos;" + ProjectFileGenerator.NewLine +
 				"\t\t\t};" + ProjectFileGenerator.NewLine +
 				"\t\t\tname = " + ConfigName + ";" + ProjectFileGenerator.NewLine +
@@ -595,7 +595,7 @@ namespace UnrealBuildTool
 				"\t\t\t\t);" + ProjectFileGenerator.NewLine +
 				"\t\t\t\tINFOPLIST_OUTPUT_FORMAT = xml;" + ProjectFileGenerator.NewLine +
 				"\t\t\t\t\"PROVISIONING_PROFILE[sdk=iphoneos*]\" = \"\";" + ProjectFileGenerator.NewLine +
-				"\t\t\t\tTARGETED_DEVICE_FAMILY = \"1,2\";" + ProjectFileGenerator.NewLine +
+				"\t\t\t\tTARGETED_DEVICE_FAMILY = \"" + Toolchain.GetPlatformDevices() + "\";" + ProjectFileGenerator.NewLine +
 				"\t\t\t\tSDKROOT = iphoneos;" + ProjectFileGenerator.NewLine +
 				"\t\t\t};" + ProjectFileGenerator.NewLine +
 				"\t\t\tname = " + ConfigName + ";" + ProjectFileGenerator.NewLine +
diff --git a/Engine/Source/Programs/UnrealBuildTool/ToolChain/UEToolChain.cs b/Engine/Source/Programs/UnrealBuildTool/ToolChain/UEToolChain.cs
--- a/Engine/Source/Programs/UnrealBuildTool/ToolChain/UEToolChain.cs
+++ b/Engine/Source/Programs/UnrealBuildTool/ToolChain/UEToolChain.cs
@@ -45,6 +45,8 @@ namespace UnrealBuildTool
 		void FixBundleBinariesPaths(UEBuildTarget Target, List<UEBuildBinary> Binaries);
 
 		string GetPlatformVersion();
+
+		string GetPlatformDevices();
 	}
 
 	public abstract class UEToolChain : IUEToolChain
@@ -217,5 +219,10 @@ namespace UnrealBuildTool
 		{
 			return "";
 		}
+
+		public virtual string GetPlatformDevices()
+		{
+			return "";
+		}
 	};
 }
\ No newline at end of file
diff --git a/Engine/Source/Runtime/IOS/IOSRuntimeSettings/Classes/IOSRuntimeSettings.h b/Engine/Source/Runtime/IOS/IOSRuntimeSettings/Classes/IOSRuntimeSettings.h
--- a/Engine/Source/Runtime/IOS/IOSRuntimeSettings/Classes/IOSRuntimeSettings.h
+++ b/Engine/Source/Runtime/IOS/IOSRuntimeSettings/Classes/IOSRuntimeSettings.h
@@ -21,6 +21,18 @@ enum class EPowerUsageFrameRateLock : uint8
     PUFRL_60 = 60 UMETA(DisplayName="60 FPS"),
 };
 
+UENUM()
+	enum class EIOSVersion : uint8
+{
+	/** iOS 6 */
+	IOS_6 = 6 UMETA(DisplayName="6.0"),
+
+	/** iOS 7 */
+	IOS_7 = 7 UMETA(DisplayName="7.0"),
+
+	/** iOS 8 */
+	IOS_8 = 8 UMETA(DisplayName="8.0"),
+};
 
 /**
  * Implements the settings for the iOS target platform.
@@ -79,6 +91,22 @@ public:
     UPROPERTY(GlobalConfig, EditAnywhere, Category = PowerUsage)
     TEnumAsByte<EPowerUsageFrameRateLock> FrameRateLock;
 
+	/** Set the minimum iOS setting */
+	UPROPERTY(GlobalConfig, EditAnywhere, Category = OSInfo)
+	TEnumAsByte<EIOSVersion> MinimumiOSVersion;
+
+	/** Does the application support iPad */
+	UPROPERTY(GlobalConfig, EditAnywhere, Category = DeviceUsage)
+	uint32 bSupportsIPad : 1;
+
+	// Does the application support iPhone */
+	UPROPERTY(GlobalConfig, EditAnywhere, Category = DeviceUsage)
+	uint32 bSupportsIPhone : 1;
+
+	// Texture compression value
+	UPROPERTY(GlobalConfig, EditAnywhere, Category = CookSettings)
+	uint32 CompressionLevel;
+
 #if WITH_EDITOR
 	// UObject interface
 	virtual void PostEditChangeProperty(struct FPropertyChangedEvent& PropertyChangedEvent) override;
diff --git a/Engine/Source/Runtime/IOS/IOSRuntimeSettings/Private/IOSRuntimeSettings.cpp b/Engine/Source/Runtime/IOS/IOSRuntimeSettings/Private/IOSRuntimeSettings.cpp
--- a/Engine/Source/Runtime/IOS/IOSRuntimeSettings/Private/IOSRuntimeSettings.cpp
+++ b/Engine/Source/Runtime/IOS/IOSRuntimeSettings/Private/IOSRuntimeSettings.cpp
@@ -14,6 +14,10 @@ UIOSRuntimeSettings::UIOSRuntimeSettings(const FObjectInitializer& ObjectInitial
 	BundleIdentifier = TEXT("com.YourCompany.GameNameNoSpaces");
 	VersionInfo = TEXT("1.0.0");
     FrameRateLock = EPowerUsageFrameRateLock::PUFRL_30;
+	bSupportsIPad = true;
+	bSupportsIPhone = true;
+	MinimumiOSVersion = EIOSVersion::IOS_6;
+	CompressionLevel = 0;
 }
 
 #if WITH_EDITOR
