Fix linking of plugin dependencies in monolithic builds. Two separate issues:

- Unbound dependencies of were plugins assuming that the plugin binary was the executable (rather than a static library)
- Only direct dependencies of the executable were being searched when building up the module list for linking

diff --git a/Engine/Source/Programs/UnrealBuildTool/Configuration/UEBuildBinary.cs b/Engine/Source/Programs/UnrealBuildTool/Configuration/UEBuildBinary.cs
--- a/Engine/Source/Programs/UnrealBuildTool/Configuration/UEBuildBinary.cs
+++ b/Engine/Source/Programs/UnrealBuildTool/Configuration/UEBuildBinary.cs
@@ -263,11 +263,9 @@ namespace UnrealBuildTool
 		/// Process all modules that aren't yet bound, creating binaries for modules that don't yet have one (if needed),
 		/// and updating modules for circular dependencies.
 		/// </summary>
-		/// <param name="Target">The target we are currently building</param>
-		/// <param name="bBuildOnlyModules">True to build only specific modules, false for all</param>
-		/// <param name="ModulesToBuild">The specific modules to build</param>
+		/// <param name="ExecutableBinary">The executable binary, which links against all unbound modules when building monolithically</param>
 		/// <returns>List of newly-created binaries (may be empty)</returns>
-		public virtual List<UEBuildBinary> ProcessUnboundModules()
+		public virtual List<UEBuildBinary> ProcessUnboundModules(UEBuildBinary ExecutableBinary)
 		{
 			return null;
 		}
@@ -419,11 +417,9 @@ namespace UnrealBuildTool
 		/// Process all modules that aren't yet bound, creating binaries for modules that don't yet have one (if needed),
 		/// and updating modules for circular dependencies.
 		/// </summary>
-		/// <param name="Target">The target we are currently building</param>
-		/// <param name="bBuildOnlyModules">True to build only specific modules, false for all</param>
-		/// <param name="ModulesToBuild">The specific modules to build</param>
+		/// <param name="ExecutableBinary">The executable binary, which links against all unbound modules when building monolithically</param>
 		/// <returns>List of newly-created binaries (may be empty)</returns>
-		public override List<UEBuildBinary> ProcessUnboundModules()
+		public override List<UEBuildBinary> ProcessUnboundModules(UEBuildBinary ExecutableBinary)
 		{
 			var Binaries = new Dictionary<string, UEBuildBinary>( StringComparer.InvariantCultureIgnoreCase );
 			if (Config.bHasModuleRules)
@@ -431,7 +427,7 @@ namespace UnrealBuildTool
 				foreach (var ModuleName in ModuleNames)
 				{
 					var Module = Target.FindOrCreateModuleByName(ModuleName);
-					Module.RecursivelyProcessUnboundModules(Target, ref Binaries, this);
+					Module.RecursivelyProcessUnboundModules(Target, ref Binaries, ExecutableBinary);
 				}
 			}
 			else
diff --git a/Engine/Source/Programs/UnrealBuildTool/Configuration/UEBuildTarget.cs b/Engine/Source/Programs/UnrealBuildTool/Configuration/UEBuildTarget.cs
--- a/Engine/Source/Programs/UnrealBuildTool/Configuration/UEBuildTarget.cs
+++ b/Engine/Source/Programs/UnrealBuildTool/Configuration/UEBuildTarget.cs
@@ -1333,13 +1333,18 @@ namespace UnrealBuildTool
 			if (ShouldCompileMonolithic())
 			{
 				var ExecutableBinary = AppBinaries[0];
-				bool bIncludeDynamicallyLoaded = true;
-				var AllReferencedModules = ExecutableBinary.GetAllDependencyModules( bIncludeDynamicallyLoaded, bForceCircular:true );
-				foreach (var CurModule in AllReferencedModules)
+
+				// Search through every binary for dependencies. When building plugin binaries that reference optional engine modules, 
+				// we still need to link them into the executable.
+				foreach (UEBuildBinary Binary in AppBinaries)
 				{
-					if ( CurModule.Binary == null || CurModule.Binary == ExecutableBinary )
+					var AllReferencedModules = Binary.GetAllDependencyModules(bIncludeDynamicallyLoaded: true, bForceCircular: true);
+					foreach (var CurModule in AllReferencedModules)
 					{
-						ExecutableBinary.AddModule(CurModule.Name);
+						if (CurModule.Binary == null || CurModule.Binary == ExecutableBinary)
+						{
+							ExecutableBinary.AddModule(CurModule.Name);
+						}
 					}
 				}
 
@@ -1618,7 +1623,7 @@ namespace UnrealBuildTool
 			foreach (var Binary in AppBinaries)
 			{
 				// Add binaries for all of our dependent modules
-				var FoundBinaries = Binary.ProcessUnboundModules();
+				var FoundBinaries = Binary.ProcessUnboundModules(AppBinaries[0]);
 				if (FoundBinaries != null)
 				{
 					NewBinaries.AddRange(FoundBinaries);
@@ -1636,9 +1641,9 @@ namespace UnrealBuildTool
 			if (OnlyModules.Count > 0)
 			{
 				AppBinaries = FilterOnlyModules();
-							}
+			}
 			else if (UEBuildConfiguration.bHotReloadFromIDE)
-				{
+			{
 				AppBinaries = FilterGameModules();
 			}
 
@@ -2427,44 +2432,44 @@ namespace UnrealBuildTool
 		/** Sets up the plugins for this target */
 		protected virtual void SetupPlugins()
 		{
-				// Filter the plugins list by the current project
-				List<PluginInfo> ValidPlugins = new List<PluginInfo>();
-				foreach(PluginInfo Plugin in Plugins.AllPlugins)
+			// Filter the plugins list by the current project
+			List<PluginInfo> ValidPlugins = new List<PluginInfo>();
+			foreach(PluginInfo Plugin in Plugins.AllPlugins)
+			{
+				if(Plugin.LoadedFrom != PluginInfo.LoadedFromType.GameProject || Plugin.Directory.StartsWith(ProjectDirectory, StringComparison.InvariantCultureIgnoreCase))
 				{
-					if(Plugin.LoadedFrom != PluginInfo.LoadedFromType.GameProject || Plugin.Directory.StartsWith(ProjectDirectory, StringComparison.InvariantCultureIgnoreCase))
-					{
-						ValidPlugins.Add(Plugin);
-					}
+					ValidPlugins.Add(Plugin);
 				}
+			}
 
-				// Build the enabled plugin list
-				if (ShouldCompileMonolithic() || TargetType == TargetRules.TargetType.Program)
-				{
-					var FilterPluginNames = new List<string>(Rules.AdditionalPlugins);
-
-					// Add the list of plugins enabled by default
-					if (UEBuildConfiguration.bCompileAgainstEngine)
-					{
-						FilterPluginNames.AddRange(ValidPlugins.Where(x => x.bEnabledByDefault).Select(x => x.Name));
-					}
-
-					// Update the plugin list for game targets
-					if(TargetType != TargetRules.TargetType.Program && UnrealBuildTool.HasUProjectFile())
-					{
-						// Enable all the game specific plugins by default
-						FilterPluginNames.AddRange(ValidPlugins.Where(x => x.LoadedFrom == PluginInfo.LoadedFromType.GameProject).Select(x => x.Name));
-
-						// Use the project settings to update the plugin list for this target
-						FilterPluginNames = UProjectInfo.GetEnabledPlugins(UnrealBuildTool.GetUProjectFile(), FilterPluginNames, Platform);
-					}
+			// Build the enabled plugin list
+			if (ShouldCompileMonolithic() || TargetType == TargetRules.TargetType.Program)
+			{
+				var FilterPluginNames = new List<string>(Rules.AdditionalPlugins);
 
-					EnabledPlugins.AddRange(ValidPlugins.Where(x => FilterPluginNames.Contains(x.Name)).Distinct());
+				// Add the list of plugins enabled by default
+				if (UEBuildConfiguration.bCompileAgainstEngine)
+				{
+					FilterPluginNames.AddRange(ValidPlugins.Where(x => x.bEnabledByDefault).Select(x => x.Name));
 				}
-				else
+
+				// Update the plugin list for game targets
+				if(TargetType != TargetRules.TargetType.Program && UnrealBuildTool.HasUProjectFile())
 				{
-					EnabledPlugins.AddRange(ValidPlugins);
+					// Enable all the game specific plugins by default
+					FilterPluginNames.AddRange(ValidPlugins.Where(x => x.LoadedFrom == PluginInfo.LoadedFromType.GameProject).Select(x => x.Name));
+
+					// Use the project settings to update the plugin list for this target
+					FilterPluginNames = UProjectInfo.GetEnabledPlugins(UnrealBuildTool.GetUProjectFile(), FilterPluginNames, Platform);
 				}
+
+				EnabledPlugins.AddRange(ValidPlugins.Where(x => FilterPluginNames.Contains(x.Name)).Distinct());
 			}
+			else
+			{
+				EnabledPlugins.AddRange(ValidPlugins);
+			}
+		}
 
 		/** Sets up the binaries for the target. */
 		protected virtual void SetupBinaries()
