- Upgraded to support API 21 (Android 5.0, Lollipop)
- Lots of Android iteration speedups:
   - When compiling from UBT (Visual Studio) perform incremental APK generation - it doesn't wipe the Intermadiate/APK directory between runs
   - Don't update project settings every run:
      - Looks at the project.properties for both game and google play services to make sure they are built with current SDK API version, if not, fixup project.properties, etc
   - Even with a non-incremental build (from Editor/UAT), keep the pre-dexed libraries (currently GooglePlayServices [GPS] only) for HUGE speedup
   - We can't just pre-dex GPS and check-in, that's not going to be safe with different java/SDK versions.
   - Removed the project.properties and proguard-project.txt from p4, as they will be generated if needed. A game can pre-supply those if needed, but we will still check API version and update if required

diff --git a/Engine/Build/Android/Java/proguard-project.txt b/Engine/Build/Android/Java/proguard-project.txt
deleted file mode 100644
diff --git a/Engine/Build/Android/Java/project.properties b/Engine/Build/Android/Java/project.properties
deleted file mode 100644
diff --git a/Engine/Source/Programs/UnrealBuildTool/Android/AndroidToolChain.cs b/Engine/Source/Programs/UnrealBuildTool/Android/AndroidToolChain.cs
--- a/Engine/Source/Programs/UnrealBuildTool/Android/AndroidToolChain.cs
+++ b/Engine/Source/Programs/UnrealBuildTool/Android/AndroidToolChain.cs
@@ -13,6 +13,9 @@ namespace UnrealBuildTool
 	{
 		static private bool bHasNDKExtensionsCompiled = false;
 
+		// the number of the clang version being used to compile 
+		static private float ClangVersionFloat = 0;
+
 		// the list of architectures we will compile for
 		static private string[] Arches = null;
 		// the list of GPU architectures we will compile for
@@ -187,8 +190,13 @@ namespace UnrealBuildTool
 			string ClangVersion = "";
 			string GccVersion = "";
 
-			// prefer clang 3.3, but fall back to 3.1 if needed for now
-			if (Directory.Exists(Path.Combine(NDKPath, @"toolchains\llvm-3.3")))
+			// prefer clang 3.5, but fall back if needed for now
+			if (Directory.Exists(Path.Combine(NDKPath, @"toolchains\llvm-3.5")))
+			{
+				ClangVersion = "3.5";
+				GccVersion = "4.9";
+			}
+			else if (Directory.Exists(Path.Combine(NDKPath, @"toolchains\llvm-3.3")))
 			{
 				ClangVersion = "3.3";
 				GccVersion = "4.8";
@@ -203,6 +211,9 @@ namespace UnrealBuildTool
 				return;
 			}
 
+			ClangVersionFloat = float.Parse(ClangVersion, System.Globalization.CultureInfo.InvariantCulture);
+			// Console.WriteLine("Compiling with clang {0}", ClangVersionFloat);
+
             string ArchitecturePath = "";
             string ArchitecturePathWindows32 = @"prebuilt\windows";
             string ArchitecturePathWindows64 = @"prebuilt\windows-x86_64";
@@ -269,6 +280,13 @@ namespace UnrealBuildTool
 			Result += " -Wno-invalid-offsetof";			// needed to suppress warnings about using offsetof on non-POD types.
 			Result += " -Wno-logical-op-parentheses";	// needed for external headers we can't change
 
+			// new for clang4.5 warnings:
+			if (ClangVersionFloat >= 3.5)
+			{
+				Result += " -Wno-undefined-bool-conversion"; // 'this' pointer cannot be null in well-defined C++ code; pointer may be assumed to always convert to true (if (this))
+
+			}
+
 			// shipping builds will cause this warning with "ensure", so disable only in those case
 			if (CompileEnvironment.Config.Target.Configuration == CPPTargetConfiguration.Shipping)
 			{
diff --git a/Engine/Source/Programs/UnrealBuildTool/Android/UEBuildAndroid.cs b/Engine/Source/Programs/UnrealBuildTool/Android/UEBuildAndroid.cs
--- a/Engine/Source/Programs/UnrealBuildTool/Android/UEBuildAndroid.cs
+++ b/Engine/Source/Programs/UnrealBuildTool/Android/UEBuildAndroid.cs
@@ -73,8 +73,10 @@ namespace UnrealBuildTool
 
             NDKPath = NDKPath.Replace("\"", "");
 
-            // can't find llvm-3.3 or llvm-3.1 in the toolchains
-            if (!Directory.Exists(Path.Combine(NDKPath, @"toolchains\llvm-3.3")) && !Directory.Exists(Path.Combine(NDKPath, @"toolchains\llvm-3.1")))
+            // need a supported llvm
+            if (!Directory.Exists(Path.Combine(NDKPath, @"toolchains\llvm-3.5")) && 
+				!Directory.Exists(Path.Combine(NDKPath, @"toolchains\llvm-3.3")) &&
+				!Directory.Exists(Path.Combine(NDKPath, @"toolchains\llvm-3.1")))
             {
                 return false;
             }
diff --git a/Engine/Source/Programs/UnrealBuildTool/Android/UEDeployAndroid.cs b/Engine/Source/Programs/UnrealBuildTool/Android/UEDeployAndroid.cs
--- a/Engine/Source/Programs/UnrealBuildTool/Android/UEDeployAndroid.cs
+++ b/Engine/Source/Programs/UnrealBuildTool/Android/UEDeployAndroid.cs
@@ -54,51 +54,52 @@ namespace UnrealBuildTool.Android
 			}
 		}
 
-		static private bool bHasPrintedApiLevel = false;
+		static private string CachedSDKLevel = null;
 		private static string GetSdkApiLevel()
 		{
-			// default to looking on disk for latest API level
-			string Target = AndroidPlatform.AndroidSdkApiTarget;
-
-			// if we want to use whatever version the ndk uses, then use that
-			if (Target == "matchndk")
-			{
-				Target = AndroidToolChain.GetNdkApiLevel();
-			}
-
-			// run a command and capture output
-			if (Target == "latest")
+			if (CachedSDKLevel == null)
 			{
-				// we expect there to be one, so use the first one
-				string AndroidCommandPath = Environment.ExpandEnvironmentVariables("%ANDROID_HOME%/tools/android.bat");
+				// default to looking on disk for latest API level
+				string Target = AndroidPlatform.AndroidSdkApiTarget;
 
-				var ExeInfo = new ProcessStartInfo(AndroidCommandPath, "list targets");
-				ExeInfo.UseShellExecute = false;
-				ExeInfo.RedirectStandardOutput = true;
-				using (var GameProcess = Process.Start(ExeInfo))
+				// if we want to use whatever version the ndk uses, then use that
+				if (Target == "matchndk")
 				{
-					PossibleApiLevels = new List<string>();
-					GameProcess.BeginOutputReadLine();
-					GameProcess.OutputDataReceived += ParseApiLevel;
-					GameProcess.WaitForExit();
+					Target = AndroidToolChain.GetNdkApiLevel();
 				}
 
-				if (PossibleApiLevels != null && PossibleApiLevels.Count > 0)
-				{
-					Target = AndroidToolChain.GetLargestApiLevel(PossibleApiLevels.ToArray());
-				}
-				else
+				// run a command and capture output
+				if (Target == "latest")
 				{
-					throw new BuildException("Can't make an APK an API installed (see \"android.bat list targets\")");
+					// we expect there to be one, so use the first one
+					string AndroidCommandPath = Environment.ExpandEnvironmentVariables("%ANDROID_HOME%/tools/android.bat");
+
+					var ExeInfo = new ProcessStartInfo(AndroidCommandPath, "list targets");
+					ExeInfo.UseShellExecute = false;
+					ExeInfo.RedirectStandardOutput = true;
+					using (var GameProcess = Process.Start(ExeInfo))
+					{
+						PossibleApiLevels = new List<string>();
+						GameProcess.BeginOutputReadLine();
+						GameProcess.OutputDataReceived += ParseApiLevel;
+						GameProcess.WaitForExit();
+					}
+
+					if (PossibleApiLevels != null && PossibleApiLevels.Count > 0)
+					{
+						Target = AndroidToolChain.GetLargestApiLevel(PossibleApiLevels.ToArray());
+					}
+					else
+					{
+						throw new BuildException("Can't make an APK an API installed (see \"android.bat list targets\")");
+					}
 				}
-			}
 
-			if (!bHasPrintedApiLevel)
-			{
 				Console.WriteLine("Building with SDK API '{0}'", Target);
-				bHasPrintedApiLevel = true;
+				CachedSDKLevel = Target;
 			}
-			return Target;
+
+			return CachedSDKLevel;
 		}
 
 		private static string GetAntPath()
@@ -195,28 +196,43 @@ namespace UnrealBuildTool.Android
 			}
 		}
 
-		private void DeleteDirectory(string Path)
+		private static void DeleteDirectory(string InPath, string SubDirectoryToKeep="")
 		{
-			Console.WriteLine("\nDeleting: " + Path);
-			if (Directory.Exists(Path))
+			// skip the dir we want to
+			if (String.Compare(Path.GetFileName(InPath), SubDirectoryToKeep, true) == 0)
+			{
+				return;
+			}
+			
+			// delete all files in here
+			string[] Files = Directory.GetFiles(InPath);
+			foreach (string Filename in Files)
 			{
 				try
 				{
-					// first make sure everything it writable
-					string[] Files = Directory.GetFiles(Path, "*.*", SearchOption.AllDirectories);
-					foreach (string Filename in Files)
-					{
-						// remove any read only flags
-						FileInfo FileInfo = new FileInfo(Filename);
-						FileInfo.Attributes = FileInfo.Attributes & ~FileAttributes.ReadOnly;
-					}
+					// remove any read only flags
+					FileInfo FileInfo = new FileInfo(Filename);
+					FileInfo.Attributes = FileInfo.Attributes & ~FileAttributes.ReadOnly;
+					FileInfo.Delete();
+				}
+				catch (Exception)
+				{
+					Log.TraceInformation("Failed to delete all files in directory {0}. Continuing on...", InPath);
+				}
+			}
 
-					// now deleting will work better
-					Directory.Delete(Path, true);
+			string[] Dirs = Directory.GetDirectories(InPath, "*.*", SearchOption.TopDirectoryOnly);
+			foreach (string Dir in Dirs)
+			{
+				DeleteDirectory(Dir, SubDirectoryToKeep);
+				// try to delete the directory, but allow it to fail (due to SubDirectoryToKeep still existing)
+				try
+				{
+					Directory.Delete(Dir);
 				}
 				catch (Exception)
 				{
-					Log.TraceInformation("Failed to delete intermediate cdirectory {0}. Continuing on...", Path);
+					// do nothing
 				}
 			}
 		}
@@ -242,16 +258,13 @@ namespace UnrealBuildTool.Android
             string Setting = OBBinAPK ? "AMAZON" : "DEVELOPMENT";
             if (!File.Exists(FileName) || ShouldWriteJavaBuildSettingsFile(FileName, Setting))
             {
-                StringBuilder BuildSettings = new StringBuilder("package com.epicgames.ue4;\npublic class JavaBuildSettings\n{\n");
+				Log.TraceInformation("{0} ====WRITING JAVABUILDSETTINGS.JAVA========================================================", DateTime.Now.ToString());
+				StringBuilder BuildSettings = new StringBuilder("package com.epicgames.ue4;\npublic class JavaBuildSettings\n{\n");
                 BuildSettings.Append("\tpublic enum PackageType {AMAZON, GOOGLE, DEVELOPMENT};\n");
                 BuildSettings.Append("\tpublic static final PackageType PACKAGING = PackageType." + Setting + ";\n");
                 BuildSettings.Append("}\n");
                 File.WriteAllText(FileName, BuildSettings.ToString());
             }
-            else
-            {
-                Console.WriteLine("::Didn't write config file; contains same data as before.");
-            }
         }
 
         public bool ShouldWriteJavaBuildSettingsFile(string FileName, string setting)
@@ -259,7 +272,7 @@ namespace UnrealBuildTool.Android
             var fileContent = File.ReadAllLines(FileName);
             var packageLine = fileContent[4]; // We know this to be true... because we write it below...
             int location = packageLine.IndexOf("PACKAGING") + 12 + 12; // + (PACKAGING = ) + ("PackageType.")
-            return String.Compare(setting, packageLine.Substring(location)) != 0;
+            return String.Compare(setting, packageLine.Substring(location, setting.Length)) != 0;
         }
 
 		private static string GetNDKArch(string UE4Arch)
@@ -310,7 +323,91 @@ namespace UnrealBuildTool.Android
 //			File.Copy("E:/Dev2/UE4-Clean/Engine/Source/ThirdParty/NVIDIA/TegraGfxDebugger/libTegra_gfx_debugger.so", UE4BuildPath + "/libs/" + Arch + "/libTegra_gfx_debugger.so");
 		}
 
-		private void MakeApk(string ProjectName, string ProjectDirectory, string OutputPath, string EngineDirectory, bool bForDistribution, string CookFlavor, bool bMakeSeparateApks)
+		private static void RunCommandLineProgramAndThrowOnError(string WorkingDirectory, string Command, string Params)
+		{
+			ProcessStartInfo StartInfo = new ProcessStartInfo();
+			StartInfo.WorkingDirectory = WorkingDirectory;
+			StartInfo.FileName = Command;
+			StartInfo.Arguments = Params;
+			StartInfo.UseShellExecute = false;
+			Log.TraceInformation("\nRunning: " + StartInfo.FileName + " " + StartInfo.Arguments);
+			Process Proc = new Process();
+			Proc.StartInfo = StartInfo;
+			Proc.Start();
+			Proc.WaitForExit();
+
+			// android bat failure
+			if (Proc.ExitCode != 0)
+			{
+				throw new BuildException("{0} failed with args {1}", Command, Params);
+			}
+		}
+
+		private static string UpdateGooglePlayServicesAndReturnLocation(string EngineDirectory, string ProjectDirectory)
+		{
+			string AndroidCommandPath = Environment.ExpandEnvironmentVariables("%ANDROID_HOME%/tools/android.bat");
+			string IntermediateAndroidPath = Path.Combine(ProjectDirectory, "Intermediate/Android/");
+			string UE4BuildPath = Path.Combine(IntermediateAndroidPath, "APK");
+
+			// We update into a temporary directory, as changing the API level will modify files
+			string GooglePlayServicesSourcePath = Path.GetFullPath(Path.Combine(EngineDirectory, "Build/Android/Java/google-play-services_lib_rev19/"));
+			string GooglePlayServicesIntermediatePath = Path.GetFullPath(Path.Combine(IntermediateAndroidPath, "google-play-services_lib/"));
+
+			// are we up to date? (version is current and the "bin" directory exists meaning we ran once)
+			bool bIsUpToDate = IsAPIVersionCurrent(GooglePlayServicesIntermediatePath) && Directory.Exists(Path.Combine(GooglePlayServicesIntermediatePath, "bin"));
+
+			if (!bIsUpToDate)
+			{
+				Log.TraceInformation("{0} ====UPDATING GOOGLE SERVICES LIBRARY======================================================", DateTime.Now.ToString());
+				Log.TraceInformation("Google Play Services library will now be updated...");
+
+				// wipe it and start over
+				DeleteDirectory(GooglePlayServicesIntermediatePath);
+				CopyFileDirectory(GooglePlayServicesSourcePath, GooglePlayServicesIntermediatePath, new Dictionary<string, string>());
+
+				// update the project
+				RunCommandLineProgramAndThrowOnError(GooglePlayServicesIntermediatePath, AndroidCommandPath, "update project " + " --path . --target " + GetSdkApiLevel());
+
+				// precompile the lib
+				RunCommandLineProgramAndThrowOnError(GooglePlayServicesIntermediatePath, "cmd.exe", "/c \"" + GetAntPath() + "\" release");
+			}
+
+			//need to create separate run for each lib. Will be added to project.properties in order in which they are added 
+			//the order is important.
+			//as android.library.reference.X=libpath where X = 1 - N
+			//for e.g this one will be added as android.library.reference.1=<EngineDirectory>/Source/ThirdParty/Android/google_play_services_lib
+
+			// Ant seems to need a relative path to work
+			Uri ServicesBuildUri = new Uri(GooglePlayServicesIntermediatePath);
+			Uri ProjectUri = new Uri(UE4BuildPath + "/");
+
+			return ProjectUri.MakeRelativeUri(ServicesBuildUri).ToString();
+		}
+
+		private static bool IsAPIVersionCurrent(string PropertiesFileLocation)
+		{
+			string ProjectProperties = Path.Combine(PropertiesFileLocation, "project.properties");
+			if (File.Exists(ProjectProperties))
+			{
+				string[] Contents = File.ReadAllLines(ProjectProperties);
+				foreach (string Line in Contents)
+				{
+					// look for the special line
+					if (Line.StartsWith("target="))
+					{
+						// skip over target=
+						string Version = Line.Substring(7);
+						// do they match?
+						return string.Compare(Version, GetSdkApiLevel(), true) == 0;
+					}
+				}
+			}
+
+			// if any other case happens, we are not current
+			return false;
+		}
+
+		private void MakeApk(string ProjectName, string ProjectDirectory, string OutputPath, string EngineDirectory, bool bForDistribution, string CookFlavor, bool bMakeSeparateApks, bool bIncrementalPackage)
 		{
 			// cache some tools paths
 			string AndroidCommandPath = Environment.ExpandEnvironmentVariables("%ANDROID_HOME%/tools/android.bat");
@@ -319,7 +416,7 @@ namespace UnrealBuildTool.Android
 
 			// set up some directory info
 			string IntermediateAndroidPath = Path.Combine(ProjectDirectory, "Intermediate/Android/");
-			string UE4BuildPath = IntermediateAndroidPath + "APK";
+			string UE4BuildPath = Path.Combine(IntermediateAndroidPath, "APK");
 			string UE4BuildFilesPath = GetUE4BuildFilePath(EngineDirectory);
 			string GameBuildFilesPath = Path.Combine(ProjectDirectory, "Build/Android");
 	
@@ -331,82 +428,92 @@ namespace UnrealBuildTool.Android
             string UE4JavaBuildSettingsFileName = GetUE4JavaBuildSettingsFileName(EngineDirectory);
             WriteJavaBuildSettingsFile(UE4JavaBuildSettingsFileName, UEBuildConfiguration.bOBBinAPK);
 
-			// first check if all .so's are up to date
-			bool bAllInputsCurrent = true;
-			foreach (string Arch in Arches)
+			// if the last packaged version is wrong, then we can skip the checking and do it
+			if (!IsAPIVersionCurrent(UE4BuildPath))
 			{
-				foreach (string GPUArch in GPUArchitectures)
+				Log.TraceInformation("Output .apk file(s) were made with a different API version, forcing repackage.");
+			}
+			else
+			{
+				// check if so's are up to date against various inputs
+				bool bAllInputsCurrent = true;
+
+				foreach (string Arch in Arches)
 				{
-					string SourceSOName = AndroidToolChain.InlineArchName(OutputPath, Arch, GPUArch);
-					// if the source binary was UE4Game, replace it with the new project name, when re-packaging a binary only build
-					string ApkFilename = Path.GetFileNameWithoutExtension(OutputPath).Replace("UE4Game", ProjectName);
-					string DestApkName = Path.Combine(ProjectDirectory, "Binaries/Android/") + ApkFilename + ".apk";
-	
-					// if we making multiple Apks, we need to put the architecture into the name
-					if (bMakeSeparateApks)
+					foreach (string GPUArch in GPUArchitectures)
 					{
+						string SourceSOName = AndroidToolChain.InlineArchName(OutputPath, Arch, GPUArch);
+						// if the source binary was UE4Game, replace it with the new project name, when re-packaging a binary only build
+						string ApkFilename = Path.GetFileNameWithoutExtension(OutputPath).Replace("UE4Game", ProjectName);
+						string DestApkName = Path.Combine(ProjectDirectory, "Binaries/Android/") + ApkFilename + ".apk";
+
+						// if we making multiple Apks, we need to put the architecture into the name
+						if (bMakeSeparateApks)
+						{
 							DestApkName = AndroidToolChain.InlineArchName(DestApkName, Arch, GPUArch);
-					}
+						}
 
-					// check to see if it's out of date before trying the slow make apk process (look at .so and all Engine and Project build files to be safe)
-					List<String> InputFiles = new List<string>();
-					InputFiles.Add(SourceSOName);
-					InputFiles.AddRange(Directory.EnumerateFiles(UE4BuildFilesPath, "*.*", SearchOption.AllDirectories));
-					if (Directory.Exists(GameBuildFilesPath))
-					{
-						InputFiles.AddRange(Directory.EnumerateFiles(GameBuildFilesPath, "*.*", SearchOption.AllDirectories));
-					}
+						// check to see if it's out of date before trying the slow make apk process (look at .so and all Engine and Project build files to be safe)
+						List<String> InputFiles = new List<string>();
+						InputFiles.Add(SourceSOName);
+						InputFiles.AddRange(Directory.EnumerateFiles(UE4BuildFilesPath, "*.*", SearchOption.AllDirectories));
+						if (Directory.Exists(GameBuildFilesPath))
+						{
+							InputFiles.AddRange(Directory.EnumerateFiles(GameBuildFilesPath, "*.*", SearchOption.AllDirectories));
+						}
 
-					// rebuild if .ini files change
-					// @todo android: programmatically determine if any .ini setting changed?
-					InputFiles.Add(Path.Combine(EngineDirectory, "Config\\BaseEngine.ini"));
-					InputFiles.Add(Path.Combine(ProjectDirectory, "Config\\DefaultEngine.ini"));
-
-                    // make sure changed java settings will rebuild apk
-                    InputFiles.Add(UE4JavaBuildSettingsFileName);
-
-                    // rebuild if .pak files exist for OBB in APK case
-                    if (UEBuildConfiguration.bOBBinAPK)
-                    {
-                        string PAKFileLocation = ProjectDirectory + "/Saved/StagedBuilds/Android" + CookFlavor + "/" + ProjectName + "/Content/Paks";
-                        if (Directory.Exists(PAKFileLocation))
-                        {
-                            var PakFiles = Directory.EnumerateFiles(PAKFileLocation, "*.pak", SearchOption.TopDirectoryOnly);
-                            foreach (var Name in PakFiles)
-                            {
-                                InputFiles.Add(Name);
-                            }
-                        }
-                    }
-
-                    // look for any newer input file
-					DateTime ApkTime = File.GetLastWriteTimeUtc(DestApkName);
-					foreach (var InputFileName in InputFiles)
-					{
-						if (File.Exists(InputFileName))
+						// rebuild if .ini files change
+						// @todo android: programmatically determine if any .ini setting changed?
+						InputFiles.Add(Path.Combine(EngineDirectory, "Config\\BaseEngine.ini"));
+						InputFiles.Add(Path.Combine(ProjectDirectory, "Config\\DefaultEngine.ini"));
+
+						// make sure changed java settings will rebuild apk
+						InputFiles.Add(UE4JavaBuildSettingsFileName);
+
+						// rebuild if .pak files exist for OBB in APK case
+						if (UEBuildConfiguration.bOBBinAPK)
 						{
-							// skip .log files
-							if (Path.GetExtension(InputFileName) == ".log")
+							string PAKFileLocation = ProjectDirectory + "/Saved/StagedBuilds/Android" + CookFlavor + "/" + ProjectName + "/Content/Paks";
+							if (Directory.Exists(PAKFileLocation))
 							{
-								continue;
+								var PakFiles = Directory.EnumerateFiles(PAKFileLocation, "*.pak", SearchOption.TopDirectoryOnly);
+								foreach (var Name in PakFiles)
+								{
+									InputFiles.Add(Name);
+								}
 							}
-							DateTime InputFileTime = File.GetLastWriteTimeUtc(InputFileName);
-							if (InputFileTime.CompareTo(ApkTime) > 0)
+						}
+
+						// look for any newer input file
+						DateTime ApkTime = File.GetLastWriteTimeUtc(DestApkName);
+						foreach (var InputFileName in InputFiles)
+						{
+							if (File.Exists(InputFileName))
 							{
-								bAllInputsCurrent = false;
-								Log.TraceInformation("{0} is out of date due to newer input file {1}", DestApkName, InputFileName);
-								break;	
+								// skip .log files
+								if (Path.GetExtension(InputFileName) == ".log")
+								{
+									continue;
+								}
+								DateTime InputFileTime = File.GetLastWriteTimeUtc(InputFileName);
+								if (InputFileTime.CompareTo(ApkTime) > 0)
+								{
+									bAllInputsCurrent = false;
+									Log.TraceInformation("{0} is out of date due to newer input file {1}", DestApkName, InputFileName);
+									break;
+								}
 							}
 						}
 					}
 				}
-			}
-			if (bAllInputsCurrent)
-			{
-				Log.TraceInformation("Output .apk file(s) are up to date (compared to the .so and .java input files)");
-				return;
+				if (bAllInputsCurrent)
+				{
+					Log.TraceInformation("Output .apk file(s) are up to date (compared to the .so and .java input files)");
+					return;
+				}
 			}
 
+
 			// Once for all arches code:
 
 			// make up a dictionary of strings to replace in the Manifest file
@@ -419,262 +526,203 @@ namespace UnrealBuildTool.Android
 				Replacements.Add("android:debuggable=\"true\"", "android:debuggable=\"false\"");
 			}
 
-			// Update the Google Play services lib with the target platform version currently in use.
-			// This appears to be required for the build to work without errors when Play services are referenced by the game.
-			// This will try to modify existing files, like project.properties, so we copy the entire library into
-			// an intermediate directory and work from there.
-			string GooglePlayServicesSourcePath = Path.GetFullPath(Path.Combine(EngineDirectory, "Build/Android/Java/google-play-services_lib_rev19/"));
-			string GooglePlayServicesIntermediatePath = Path.GetFullPath(Path.Combine(IntermediateAndroidPath, "google-play-services_lib/"));
-
-			DeleteDirectory(GooglePlayServicesIntermediatePath);
-			CopyFileDirectory(GooglePlayServicesSourcePath, GooglePlayServicesIntermediatePath, new Dictionary<string, string>());
+			// make sure that GooglePlayServices is up to date, and get the location so we can import it
+			// and if we have done a recent, incremental build with the same version, then the intermediate GPS stuff is up to date
+			string RelativeGPSLocation = UpdateGooglePlayServicesAndReturnLocation(EngineDirectory, ProjectDirectory);
 
-			ProcessStartInfo AndroidBatStartInfoPlayServicesLib = new ProcessStartInfo();
-			AndroidBatStartInfoPlayServicesLib.WorkingDirectory = GooglePlayServicesIntermediatePath;
-			AndroidBatStartInfoPlayServicesLib.FileName = AndroidCommandPath;
-			AndroidBatStartInfoPlayServicesLib.Arguments = "update project " + " --path . --target " + GetSdkApiLevel();
-			AndroidBatStartInfoPlayServicesLib.UseShellExecute = false;
-			Console.WriteLine("\nRunning: " + AndroidBatStartInfoPlayServicesLib.FileName + " " + AndroidBatStartInfoPlayServicesLib.Arguments);
-			Process AndroidBatPlayServicesLib = new Process();
-			AndroidBatPlayServicesLib.StartInfo = AndroidBatStartInfoPlayServicesLib;
-			AndroidBatPlayServicesLib.Start();
-			AndroidBatPlayServicesLib.WaitForExit();
+			Log.TraceInformation("{0} ====PREPARING TO MAKE APK=================================================================", DateTime.Now.ToString());
 
-			// android bat failure
-			if (AndroidBatPlayServicesLib.ExitCode != 0)
+			if (!bIncrementalPackage)
 			{
-				throw new BuildException("android.bat failed [{0}]", AndroidBatStartInfoPlayServicesLib.Arguments);
+				// Wipe the Intermediate/Build/APK directory first, except for dexedLibs, because Google Services takes FOREVER to predex, and it never changes
+				// so allow the ANT checking to win here - if this grows a bit with extra libs, it's fine, it _should_ only pull in dexedLibs it needs
+				Log.TraceInformation("Performing complete package - wiping {0}, except for predexedLibs", UE4BuildPath);
+				DeleteDirectory(UE4BuildPath, "dexedLibs");
 			}
 
-			//need to create separate run for each lib. Will be added to project.properties in order in which they are added 
-			//the order is important.
-			//as android.library.reference.X=libpath where X = 1 - N
-			//for e.g this one will be added as android.library.reference.1=<EngineDirectory>/Source/ThirdParty/Android/google_play_services_lib
+			// If we are packaging for Amazon then we need to copy the PAK files to the correct location
+			// Currently we'll just support 1 of 'em
+			if (UEBuildConfiguration.bOBBinAPK)
+			{
+				string PAKFileLocation = ProjectDirectory + "/Saved/StagedBuilds/Android" + CookFlavor + "/" + ProjectName + "/Content/Paks";
+				Console.WriteLine("Pak location {0}", PAKFileLocation);
+				string PAKFileDestination = UE4BuildPath + "/assets";
+				Console.WriteLine("Pak destination location {0}", PAKFileDestination);
+				if (Directory.Exists(PAKFileLocation))
+				{
+					Directory.CreateDirectory(UE4BuildPath);
+					Directory.CreateDirectory(PAKFileDestination);
+					Console.WriteLine("PAK file exists...");
+					var PakFiles = Directory.EnumerateFiles(PAKFileLocation, "*.pak", SearchOption.TopDirectoryOnly);
+					foreach (var Name in PakFiles)
+					{
+						Console.WriteLine("Found file {0}", Name);
+					}
 
-			// Ant seems to need a relative path to work
-			Uri ServicesBuildUri = new Uri(GooglePlayServicesIntermediatePath);
-			Uri ProjectUri = new Uri(UE4BuildPath + "/");
-			string RelativeServicesUri = ProjectUri.MakeRelativeUri(ServicesBuildUri).ToString();
+					if (PakFiles.Count() > 0)
+					{
+						var DestFileName = Path.Combine(PAKFileDestination, Path.GetFileName(PakFiles.ElementAt(0)) + ".png"); // Need a rename to turn off compression
+						var SrcFileName = PakFiles.ElementAt(0);
+						if (!File.Exists(DestFileName) || File.GetLastWriteTimeUtc(DestFileName) < File.GetLastWriteTimeUtc(SrcFileName))
+						{
+							Console.WriteLine("Copying {0} to {1}", SrcFileName, DestFileName);
+							File.Copy(SrcFileName, DestFileName);
+						}
+					}
+				}
+				// Do we want to kill the OBB here or not???
+			}
 
-			string FinalNdkBuildABICommand = "";
+			//Copy build files to the intermediate folder in this order (later overrides earlier):
+			//	- Shared Engine
+			//  - Shared Engine NoRedist (for Epic secret files)
+			//  - Game
+			//  - Game NoRedist (for Epic secret files)
+			CopyFileDirectory(UE4BuildFilesPath, UE4BuildPath, Replacements);
+			CopyFileDirectory(UE4BuildFilesPath + "/NotForLicensees", UE4BuildPath, Replacements);
+			CopyFileDirectory(UE4BuildFilesPath + "/NoRedist", UE4BuildPath, Replacements);
+			CopyFileDirectory(GameBuildFilesPath, UE4BuildPath, Replacements);
+			CopyFileDirectory(GameBuildFilesPath + "/NotForLicensees", UE4BuildPath, Replacements);
+			CopyFileDirectory(GameBuildFilesPath + "/NoRedist", UE4BuildPath, Replacements);
+
+			// update metadata files (like project.properties, build.xml) if we are missing a build.xml or if we just overwrote project.properties with a bad version in it (from game/engine dir)
+			string BuildXmlFile = Path.Combine(UE4BuildPath, "build.xml");
+			if (!File.Exists(BuildXmlFile) || !IsAPIVersionCurrent(UE4BuildPath))
+			{
+				Log.TraceInformation("{0} ====UPDATING BUILD.XML AND PROJECT.PROPERTIES (MISSING OR BAD SDK VERSION)================", DateTime.Now.ToString());
+				RunCommandLineProgramAndThrowOnError(UE4BuildPath, AndroidCommandPath, "update project --name " + ProjectName + " --path . --target " + GetSdkApiLevel() + " --library " + RelativeGPSLocation);
+			}
 
 			// now make the apk(s)
+			string FinalNdkBuildABICommand = "";
 			for (int ArchIndex = 0; ArchIndex < Arches.Length; ArchIndex++)
 			{
 				string Arch = Arches[ArchIndex];
 				for (int GPUArchIndex = 0; GPUArchIndex < GPUArchitectures.Length; GPUArchIndex++)
 				{
+					Log.TraceInformation("{0} ====PREPARING NATIVE CODE=================================================================", DateTime.Now.ToString());
+
 					string GPUArchitecture = GPUArchitectures[GPUArchIndex];
 					string SourceSOName = AndroidToolChain.InlineArchName(OutputPath, Arch, GPUArchitecture);
-				// if the source binary was UE4Game, replace it with the new project name, when re-packaging a binary only build
-				string ApkFilename = Path.GetFileNameWithoutExtension(OutputPath).Replace("UE4Game", ProjectName);
-				string DestApkName = Path.Combine(ProjectDirectory, "Binaries/Android/") + ApkFilename + ".apk";
-
-				// if we making multiple Apks, we need to put the architecture into the name
-				if (bMakeSeparateApks)
-				{
-						DestApkName = AndroidToolChain.InlineArchName(DestApkName, Arch, GPUArchitecture);
-				}
-
-				// code in here will run once per .apk (or once if merged apk)
-				if (bMakeSeparateApks || ArchIndex == 0)
-				{
-					//Wipe the Intermediate/Build/APK directory first
-					DeleteDirectory(UE4BuildPath);
+					// if the source binary was UE4Game, replace it with the new project name, when re-packaging a binary only build
+					string ApkFilename = Path.GetFileNameWithoutExtension(OutputPath).Replace("UE4Game", ProjectName);
+					string DestApkName = Path.Combine(ProjectDirectory, "Binaries/Android/") + ApkFilename + ".apk";
 
-					// If we are packaging for Amazon then we need to copy the PAK files to the correct location
-					// Currently we'll just support 1 of 'em
-					if (UEBuildConfiguration.bOBBinAPK)
+					// if we making multiple Apks, we need to put the architecture into the name
+					if (bMakeSeparateApks)
 					{
-						string PAKFileLocation = ProjectDirectory + "/Saved/StagedBuilds/Android" + CookFlavor + "/" + ProjectName + "/Content/Paks";
-						Console.WriteLine("Pak location {0}", PAKFileLocation);
-						string PAKFileDestination = UE4BuildPath + "/assets";
-						Console.WriteLine("Pak destination location {0}", PAKFileDestination);
-						if (Directory.Exists(PAKFileLocation))
-						{
-							Directory.CreateDirectory(UE4BuildPath);
-							Directory.CreateDirectory(PAKFileDestination);
-							Console.WriteLine("PAK file exists...");
-							var PakFiles = Directory.EnumerateFiles(PAKFileLocation, "*.pak", SearchOption.TopDirectoryOnly);
-							foreach (var Name in PakFiles)
-							{
-								Console.WriteLine("Found file {0}", Name);
-							}
-
-							if (PakFiles.Count() > 0)
-							{
-								var DestFileName = Path.Combine(PAKFileDestination, Path.GetFileName(PakFiles.ElementAt(0)) + ".png"); // Need a rename to turn off compression
-								var SrcFileName = PakFiles.ElementAt(0);
-								if (!File.Exists(DestFileName) || File.GetLastWriteTimeUtc(DestFileName) < File.GetLastWriteTimeUtc(SrcFileName))
-								{
-									Console.WriteLine("Copying {0} to {1}", SrcFileName, DestFileName);
-									File.Copy(SrcFileName, DestFileName);
-								}
-							}
-						}
-						// Do we want to kill the OBB here or not???
+						DestApkName = AndroidToolChain.InlineArchName(DestApkName, Arch, GPUArchitecture);
 					}
 
-					//Copy build files to the intermediate folder in this order (later overrides earlier):
-					//	- Shared Engine
-					//  - Shared Engine NoRedist (for Epic secret files)
-					//  - Game
-					//  - Game NoRedist (for Epic secret files)
-					CopyFileDirectory(UE4BuildFilesPath, UE4BuildPath, Replacements);
-					CopyFileDirectory(UE4BuildFilesPath + "/NotForLicensees", UE4BuildPath, Replacements);
-					CopyFileDirectory(UE4BuildFilesPath + "/NoRedist", UE4BuildPath, Replacements);
-					CopyFileDirectory(GameBuildFilesPath, UE4BuildPath, Replacements);
-					CopyFileDirectory(GameBuildFilesPath + "/NotForLicensees", UE4BuildPath, Replacements);
-					CopyFileDirectory(GameBuildFilesPath + "/NoRedist", UE4BuildPath, Replacements);
-
-					//Android.bat for game-specific
-					ProcessStartInfo AndroidBatStartInfoGame = new ProcessStartInfo();
-					AndroidBatStartInfoGame.WorkingDirectory = UE4BuildPath;
-					AndroidBatStartInfoGame.FileName = AndroidCommandPath;
-					AndroidBatStartInfoGame.Arguments = "update project --name " + ProjectName + " --path . --target " + GetSdkApiLevel();
-					AndroidBatStartInfoGame.UseShellExecute = false;
-					Console.WriteLine("\nRunning: " + AndroidBatStartInfoGame.FileName + " " + AndroidBatStartInfoGame.Arguments);
-					Process AndroidBatGame = new Process();
-					AndroidBatGame.StartInfo = AndroidBatStartInfoGame;
-					AndroidBatGame.Start();
-					AndroidBatGame.WaitForExit();
-
-					// android bat failure
-					if (AndroidBatGame.ExitCode != 0)
+					// Copy the generated .so file from the binaries directory to the jni folder
+					if (!File.Exists(SourceSOName))
 					{
-						throw new BuildException("android.bat failed [{0}]", AndroidBatStartInfoGame.Arguments);
+						throw new BuildException("Can't make an APK without the compiled .so [{0}]", SourceSOName);
 					}
-
-					AndroidBatStartInfoGame.Arguments = " update project --name " + ProjectName + " --path .  --target " + GetSdkApiLevel() + " --library " + RelativeServicesUri;
-					Console.WriteLine("\nRunning: " + AndroidBatStartInfoGame.FileName + " " + AndroidBatStartInfoGame.Arguments);
-					AndroidBatGame.Start();
-					AndroidBatGame.WaitForExit();
-
-					if (AndroidBatGame.ExitCode != 0)
+					if (!Directory.Exists(UE4BuildPath + "/jni"))
 					{
-						throw new BuildException("android.bat failed [{0}]", AndroidBatStartInfoGame.Arguments);
+						throw new BuildException("Can't make an APK without the jni directory [{0}/jni]", UE4BuildFilesPath);
 					}
-				}
-
-				// Copy the generated .so file from the binaries directory to the jni folder
-				if (!File.Exists(SourceSOName))
-				{
-					throw new BuildException("Can't make an APK without the compiled .so [{0}]", SourceSOName);
-				}
-				if (!Directory.Exists(UE4BuildPath + "/jni"))
-				{
-					throw new BuildException("Can't make an APK without the jni directory [{0}/jni]", UE4BuildFilesPath);
-				}
 
-				// Use ndk-build to do stuff and move the .so file to the lib folder (only if NDK is installed)
-				string FinalSOName = "";
-				if (File.Exists(NDKBuildPath))
-				{
-					string LibDir = UE4BuildPath + "/jni/" + GetNDKArch(Arch);
-					Directory.CreateDirectory(LibDir);
+					// Use ndk-build to do stuff and move the .so file to the lib folder (only if NDK is installed)
+					string FinalSOName = "";
+					if (File.Exists(NDKBuildPath))
+					{
+						string LibDir = UE4BuildPath + "/jni/" + GetNDKArch(Arch);
+						Directory.CreateDirectory(LibDir);
 
-					// copy the binary to the standard .so location
-					FinalSOName = LibDir + "/libUE4.so";
-					File.Copy(SourceSOName, FinalSOName, true);
+						// copy the binary to the standard .so location
+						FinalSOName = LibDir + "/libUE4.so";
+						File.Copy(SourceSOName, FinalSOName, true);
 
-					FinalNdkBuildABICommand += GetNDKArch(Arch) + " ";
-				}
-				else
-				{
-					// if no NDK, we don't need any of the debugger stuff, so we just copy the .so to where it will end up
-					FinalSOName = UE4BuildPath + "/libs/" + GetNDKArch(Arch) + "/libUE4.so";
-					Directory.CreateDirectory(Path.GetDirectoryName(FinalSOName));
-					File.Copy(SourceSOName, FinalSOName);
-				}
+						FinalNdkBuildABICommand += GetNDKArch(Arch) + " ";
+					}
+					else
+					{
+						// if no NDK, we don't need any of the debugger stuff, so we just copy the .so to where it will end up
+						FinalSOName = UE4BuildPath + "/libs/" + GetNDKArch(Arch) + "/libUE4.so";
+						Directory.CreateDirectory(Path.GetDirectoryName(FinalSOName));
+						File.Copy(SourceSOName, FinalSOName);
+					}
 
-				// now do final stuff per apk (or after all .so's for a shared .apk)
-				if (bMakeSeparateApks || ArchIndex == NumArches - 1)
-				{
-					// if we need to run ndk-build, do it now (if making a shared .apk, we need to wait until all .libs exist)
-					if (!string.IsNullOrEmpty(FinalNdkBuildABICommand))
+					// now do final stuff per apk (or after all .so's for a shared .apk)
+					if (bMakeSeparateApks || ArchIndex == NumArches - 1)
 					{
-						ProcessStartInfo NDKBuildInfo = new ProcessStartInfo();
-						NDKBuildInfo.WorkingDirectory = UE4BuildPath;
-						NDKBuildInfo.FileName = NDKBuildPath;
-						NDKBuildInfo.Arguments = "APP_ABI=\"" + FinalNdkBuildABICommand + "\"";
-						FinalNdkBuildABICommand = "";
-						if (!bForDistribution)
-						{
-							NDKBuildInfo.Arguments += " NDK_DEBUG=1";
-						}
-						NDKBuildInfo.UseShellExecute = true;
-						NDKBuildInfo.WindowStyle = ProcessWindowStyle.Minimized;
-						Console.WriteLine("\nRunning: " + NDKBuildInfo.FileName + " " + NDKBuildInfo.Arguments);
-						Process NDKBuild = new Process();
-						NDKBuild.StartInfo = NDKBuildInfo;
-						NDKBuild.Start();
-						NDKBuild.WaitForExit();
-
-						// ndk build failure
-						if (NDKBuild.ExitCode != 0)
+						// always delete libs up to this point so fat binaries and incremental builds work together (otherwise we might end up with multiple
+						// so files in an apk that doesn't want them)
+						DeleteDirectory(UE4BuildPath + "/libs");
+
+						// if we need to run ndk-build, do it now (if making a shared .apk, we need to wait until all .libs exist)
+						if (!string.IsNullOrEmpty(FinalNdkBuildABICommand))
 						{
-							throw new BuildException("ndk-build failed [{0}]", NDKBuildInfo.Arguments);
+							ProcessStartInfo NDKBuildInfo = new ProcessStartInfo();
+							NDKBuildInfo.WorkingDirectory = UE4BuildPath;
+							NDKBuildInfo.FileName = NDKBuildPath;
+							NDKBuildInfo.Arguments = "APP_ABI=\"" + FinalNdkBuildABICommand + "\"";
+							FinalNdkBuildABICommand = "";
+							if (!bForDistribution)
+							{
+								NDKBuildInfo.Arguments += " NDK_DEBUG=1";
+							}
+							NDKBuildInfo.UseShellExecute = true;
+							NDKBuildInfo.WindowStyle = ProcessWindowStyle.Minimized;
+							Console.WriteLine("\nRunning: " + NDKBuildInfo.FileName + " " + NDKBuildInfo.Arguments);
+							Process NDKBuild = new Process();
+							NDKBuild.StartInfo = NDKBuildInfo;
+							NDKBuild.Start();
+							NDKBuild.WaitForExit();
+
+							// ndk build failure
+							if (NDKBuild.ExitCode != 0)
+							{
+								throw new BuildException("ndk-build failed [{0}]", NDKBuildInfo.Arguments);
+							}
 						}
-					}
 	
-					// after ndk-build is called, we can now copy in the stl .so (ndk-build deletes old files)
-					// copy libgnustl_shared.so to library (use 4.8 if possible, otherwise 4.6)
-					if (bMakeSeparateApks)
-					{
-						CopySTL(UE4BuildPath, Arch);
+						// after ndk-build is called, we can now copy in the stl .so (ndk-build deletes old files)
+						// copy libgnustl_shared.so to library (use 4.8 if possible, otherwise 4.6)
+						if (bMakeSeparateApks)
+						{
+							CopySTL(UE4BuildPath, Arch);
 							CopyGfxDebugger(UE4BuildPath, Arch);
-					}
-					else
-					{
-						foreach (string InnerArch in Arches)
+						}
+						else
 						{
-							CopySTL(UE4BuildPath, InnerArch);
+							foreach (string InnerArch in Arches)
+							{
+								CopySTL(UE4BuildPath, InnerArch);
 								CopyGfxDebugger(UE4BuildPath, InnerArch);
+							}
 						}
-					}
 
 
-					// remove any read only flags
-					FileInfo DestFileInfo = new FileInfo(FinalSOName);
-					DestFileInfo.Attributes = DestFileInfo.Attributes & ~FileAttributes.ReadOnly;
+						// remove any read only flags
+						FileInfo DestFileInfo = new FileInfo(FinalSOName);
+						DestFileInfo.Attributes = DestFileInfo.Attributes & ~FileAttributes.ReadOnly;
 
-					// Use ant debug to build the .apk file
-					ProcessStartInfo CallAntStartInfo = new ProcessStartInfo();
-					CallAntStartInfo.WorkingDirectory = UE4BuildPath;
-					CallAntStartInfo.FileName = "cmd.exe";
-					CallAntStartInfo.Arguments = "/c \"" + AntBuildPath + "\"  " + (bForDistribution ? "release" : "debug");
-					CallAntStartInfo.UseShellExecute = false;
-					Console.WriteLine("\nRunning: " + CallAntStartInfo.Arguments);
-					Process CallAnt = new Process();
-					CallAnt.StartInfo = CallAntStartInfo;
-					CallAnt.Start();
-					CallAnt.WaitForExit();
-
-					// ant failure
-					if (CallAnt.ExitCode != 0)
-					{
-						throw new BuildException("ant.bat failed [{0}]", CallAntStartInfo.Arguments);
-					}
+						// Use ant debug to build the .apk file
+						Log.TraceInformation("{0} ====PERFORMING FINAL APK PACKAGE OPERATION=(with no library dependency compiling)=========", DateTime.Now.ToString());
+						RunCommandLineProgramAndThrowOnError(UE4BuildPath, "cmd.exe", "/c \"" + GetAntPath() + "\" nodeps " + (bForDistribution ? "release" : "debug"));
 
-					// make sure destination exists
-					Directory.CreateDirectory(Path.GetDirectoryName(DestApkName));
+						// make sure destination exists
+						Directory.CreateDirectory(Path.GetDirectoryName(DestApkName));
 
-					// do we need to sign for distro?
-					if (bForDistribution)
-					{
-						// use diffeent source and dest apk's for signed mode
-						string SourceApkName = UE4BuildPath + "/bin/" + ProjectName + "-release-unsigned.apk";
-						SignApk(UE4BuildPath + "/SigningConfig.xml", SourceApkName, DestApkName);
-					}
-					else
-					{
-						// now copy to the final location
-						File.Copy(UE4BuildPath + "/bin/" + ProjectName + "-debug" + ".apk", DestApkName, true);
+						// do we need to sign for distro?
+						if (bForDistribution)
+						{
+							// use diffeent source and dest apk's for signed mode
+							string SourceApkName = UE4BuildPath + "/bin/" + ProjectName + "-release-unsigned.apk";
+							SignApk(UE4BuildPath + "/SigningConfig.xml", SourceApkName, DestApkName);
+						}
+						else
+						{
+							// now copy to the final location
+							File.Copy(UE4BuildPath + "/bin/" + ProjectName + "-debug" + ".apk", DestApkName, true);
+						}
 					}
 				}
 			}
 		}
-		}
 
 		private void SignApk(string ConfigFilePath, string SourceApk, string DestApk)
 		{
@@ -777,7 +825,7 @@ namespace UnrealBuildTool.Android
 			// we need to strip architecture from any of the output paths
 			string BaseSoName = AndroidToolChain.RemoveArchName(InTarget.OutputPaths[0]);
 			// this always makes a merged .apk since for debugging, there's no way to know which one to run
-			MakeApk(InTarget.AppName, InTarget.ProjectDirectory, BaseSoName, BuildConfiguration.RelativeEnginePath, bForDistribution:false, CookFlavor:"", bMakeSeparateApks:false);
+			MakeApk(InTarget.AppName, InTarget.ProjectDirectory, BaseSoName, BuildConfiguration.RelativeEnginePath, bForDistribution:false, CookFlavor:"", bMakeSeparateApks:false, bIncrementalPackage:true);
 			return true;
 		}
 
@@ -793,7 +841,7 @@ namespace UnrealBuildTool.Android
 
 		public override bool PrepForUATPackageOrDeploy(string ProjectName, string ProjectDirectory, string ExecutablePath, string EngineDirectory, bool bForDistribution, string CookFlavor)
 		{
-			MakeApk(ProjectName, ProjectDirectory, ExecutablePath, EngineDirectory, bForDistribution, CookFlavor, ShouldMakeSeparateApks());
+			MakeApk(ProjectName, ProjectDirectory, ExecutablePath, EngineDirectory, bForDistribution, CookFlavor, ShouldMakeSeparateApks(), bIncrementalPackage:false);
 			return true;
 		}
 
diff --git a/Engine/Source/Runtime/Core/Private/Android/AndroidMemory.cpp b/Engine/Source/Runtime/Core/Private/Android/AndroidMemory.cpp
--- a/Engine/Source/Runtime/Core/Private/Android/AndroidMemory.cpp
+++ b/Engine/Source/Runtime/Core/Private/Android/AndroidMemory.cpp
@@ -95,7 +95,8 @@ FMalloc* FAndroidPlatformMemory::BaseAllocator()
 
 void* FAndroidPlatformMemory::BinnedAllocFromOS( SIZE_T Size )
 {
-	return valloc(Size);
+	// valloc was deprecated, this is a functional equivalent, for SDK 21
+	return memalign(sysconf(_SC_PAGESIZE), Size);
 }
 
 void FAndroidPlatformMemory::BinnedFreeToOS( void* Ptr )
diff --git a/Engine/Source/Runtime/Core/Private/HAL/PThreadRunnableThread.h b/Engine/Source/Runtime/Core/Private/HAL/PThreadRunnableThread.h
--- a/Engine/Source/Runtime/Core/Private/HAL/PThreadRunnableThread.h
+++ b/Engine/Source/Runtime/Core/Private/HAL/PThreadRunnableThread.h
@@ -4,7 +4,7 @@
 
 #include <pthread.h>
 
-#if PLATFORM_LINUX
+#if PLATFORM_LINUX || PLATFORM_ANDROID
 #define PTHREAD_NULL -1
 #else
 #define PTHREAD_NULL NULL
diff --git a/Engine/Source/Runtime/Core/Public/Android/AndroidCompilerSetup.h b/Engine/Source/Runtime/Core/Public/Android/AndroidCompilerSetup.h
--- a/Engine/Source/Runtime/Core/Public/Android/AndroidCompilerSetup.h
+++ b/Engine/Source/Runtime/Core/Public/Android/AndroidCompilerSetup.h
@@ -7,3 +7,5 @@
 
 #pragma once
 
+// The register keyword is deprecated
+#define register
\ No newline at end of file
