Fix cached include paths being modified by CompileCPPFiles() for Mac and iOS.

diff --git a/Engine/Source/Programs/UnrealBuildTool/IOS/IOSToolChain.cs b/Engine/Source/Programs/UnrealBuildTool/IOS/IOSToolChain.cs
--- a/Engine/Source/Programs/UnrealBuildTool/IOS/IOSToolChain.cs
+++ b/Engine/Source/Programs/UnrealBuildTool/IOS/IOSToolChain.cs
@@ -513,7 +513,7 @@ namespace UnrealBuildTool
 			}
 
 			// Add include paths to the argument list.
-			HashSet<string> AllIncludes = CompileEnvironment.Config.CPPIncludeInfo.IncludePaths;
+			HashSet<string> AllIncludes = new HashSet<string>(CompileEnvironment.Config.CPPIncludeInfo.IncludePaths);
 			AllIncludes.UnionWith(CompileEnvironment.Config.CPPIncludeInfo.SystemIncludePaths);
 			foreach (string IncludePath in AllIncludes)
 			{
diff --git a/Engine/Source/Programs/UnrealBuildTool/Mac/MacToolChain.cs b/Engine/Source/Programs/UnrealBuildTool/Mac/MacToolChain.cs
--- a/Engine/Source/Programs/UnrealBuildTool/Mac/MacToolChain.cs
+++ b/Engine/Source/Programs/UnrealBuildTool/Mac/MacToolChain.cs
@@ -320,7 +320,7 @@ namespace UnrealBuildTool
 			}
 
 			// Add include paths to the argument list.
-			HashSet<string> AllIncludes = CompileEnvironment.Config.CPPIncludeInfo.IncludePaths;
+			HashSet<string> AllIncludes = new HashSet<string>(CompileEnvironment.Config.CPPIncludeInfo.IncludePaths);
 			AllIncludes.UnionWith(CompileEnvironment.Config.CPPIncludeInfo.SystemIncludePaths);
 			foreach (string IncludePath in AllIncludes)
 			{
