Workaround for a problem with Mac build products being placed in wrong subfolder of an app bundle due to duplicate entries in Binaries array in MacToolChain.FixBundleBinariesPaths()

diff --git a/Engine/Source/Programs/UnrealBuildTool/Mac/MacToolChain.cs b/Engine/Source/Programs/UnrealBuildTool/Mac/MacToolChain.cs
--- a/Engine/Source/Programs/UnrealBuildTool/Mac/MacToolChain.cs
+++ b/Engine/Source/Programs/UnrealBuildTool/Mac/MacToolChain.cs
@@ -1157,7 +1157,7 @@ namespace UnrealBuildTool
 						Binary.Config.OutputFilePaths[0] = BundleContentsPath + "MacOS" + SubDir + "/" + BinaryFileName;
 					}
 				}
-				else if (!BinaryFileName.EndsWith(".a"))
+				else if (!BinaryFileName.EndsWith(".a") && !BinaryFileName.Contains(".app/Contents/MacOS/")) // Binaries can contain duplicates
 				{
 					Binary.Config.OutputFilePaths[0] += ".app/Contents/MacOS/" + BinaryFileName;
 				}
