Fix mono warning (unused variable).

diff --git a/Engine/Source/Programs/UnrealBuildTool/Android/UEDeployAndroid.cs b/Engine/Source/Programs/UnrealBuildTool/Android/UEDeployAndroid.cs
--- a/Engine/Source/Programs/UnrealBuildTool/Android/UEDeployAndroid.cs
+++ b/Engine/Source/Programs/UnrealBuildTool/Android/UEDeployAndroid.cs
@@ -412,7 +412,6 @@ namespace UnrealBuildTool.Android
 			// cache some tools paths
 			string AndroidCommandPath = Environment.ExpandEnvironmentVariables("%ANDROID_HOME%/tools/android.bat");
 			string NDKBuildPath = Environment.ExpandEnvironmentVariables("%NDKROOT%/ndk-build.cmd");
-			string AntBuildPath = GetAntPath();
 
 			// set up some directory info
 			string IntermediateAndroidPath = Path.Combine(ProjectDirectory, "Intermediate/Android/");
