Fix linking to libraries located in path with spaces.

- PR #551 contributed by [REDACTED_PERSON].
- Fixes a quite common case of a project with plugins being located in "Unreal Projects" directory.
- UE-4651

diff --git a/Engine/Source/Programs/UnrealBuildTool/Linux/LinuxToolChain.cs b/Engine/Source/Programs/UnrealBuildTool/Linux/LinuxToolChain.cs
--- a/Engine/Source/Programs/UnrealBuildTool/Linux/LinuxToolChain.cs
+++ b/Engine/Source/Programs/UnrealBuildTool/Linux/LinuxToolChain.cs
@@ -854,6 +854,10 @@ namespace UnrealBuildTool
                 {
                     // static library passed in, pass it along but make path absolute, because FixDependencies script may be executed in a different directory
                     string AbsoluteAdditionalLibrary = Path.GetFullPath(AdditionalLibrary);
+					if (AbsoluteAdditionalLibrary.Contains(" "))
+					{
+						AbsoluteAdditionalLibrary = string.Format("\"{0}\"", AbsoluteAdditionalLibrary);
+					}
                     LinkAction.CommandArguments += (" " + AbsoluteAdditionalLibrary);
                     LinkAction.PrerequisiteItems.Add(FileItem.GetItemByPath(AdditionalLibrary));
                 }
