- Fixing xbox packaging issues [UE-1251]

diff --git a/Engine/Source/Programs/AutomationTool/DeploymentContext.cs b/Engine/Source/Programs/AutomationTool/DeploymentContext.cs
--- a/Engine/Source/Programs/AutomationTool/DeploymentContext.cs
+++ b/Engine/Source/Programs/AutomationTool/DeploymentContext.cs
@@ -306,6 +306,8 @@ public class DeploymentContext //: ProjectParams
 	public int StageFiles(StagedFileType FileType, string InPath, string Wildcard = "*", bool bRecursive = true, string[] ExcludeWildcard = null, string NewPath = null, bool bAllowNone = false, bool bRemap = true, string NewName = null)
 	{
 		int FilesAdded = 0;
+		// make sure any ..'s are removed
+		Utils.CollapseRelativeDirectories(ref InPath);
 
 		if (CommandUtils.DirectoryExists(InPath))
 		{
diff --git a/Engine/Source/Programs/UnrealBuildTool/System/Utils.cs b/Engine/Source/Programs/UnrealBuildTool/System/Utils.cs
--- a/Engine/Source/Programs/UnrealBuildTool/System/Utils.cs
+++ b/Engine/Source/Programs/UnrealBuildTool/System/Utils.cs
@@ -736,21 +736,30 @@ namespace UnrealBuildTool
 		/// <returns>true if the path could be collapsed, false otherwise.</returns>
 		public static bool CollapseRelativeDirectories(ref string InPath)
 		{
+			string LocalString = InPath;
+			bool bHadBackSlashes = false;
+			// look to see what kind of slashes we had
+			if (LocalString.IndexOf("\\") != -1)
+			{
+				LocalString = LocalString.Replace("\\", "/");
+				bHadBackSlashes = true;
+			}
+
 			string ParentDir       = "/..";
 			int    ParentDirLength = ParentDir.Length;
 
 			for (;;)
 			{
 				// An empty path is finished
-				if (string.IsNullOrEmpty(InPath))
+				if (string.IsNullOrEmpty(LocalString))
 					break;
 
 				// Consider empty paths or paths which start with .. or /.. as invalid
-				if (InPath.StartsWith("..") || InPath.StartsWith(ParentDir))
+				if (LocalString.StartsWith("..") || LocalString.StartsWith(ParentDir))
 					return false;
 
 				// If there are no "/.."s left then we're done
-				int Index = InPath.IndexOf(ParentDir);
+				int Index = LocalString.IndexOf(ParentDir);
 				if (Index == -1)
 					break;
 
@@ -758,27 +767,35 @@ namespace UnrealBuildTool
 				for (;;)
 				{
 					// Find the previous slash
-					PreviousSeparatorIndex = Math.Max(0, InPath.LastIndexOf("/", PreviousSeparatorIndex - 1));
+					PreviousSeparatorIndex = Math.Max(0, LocalString.LastIndexOf("/", PreviousSeparatorIndex - 1));
 
 					// Stop if we've hit the start of the string
 					if (PreviousSeparatorIndex == 0)
 						break;
 
 					// Stop if we've found a directory that isn't "/./"
-					if ((Index - PreviousSeparatorIndex) > 1 && (InPath[PreviousSeparatorIndex + 1] != '.' || InPath[PreviousSeparatorIndex + 2] != '/'))
+					if ((Index - PreviousSeparatorIndex) > 1 && (LocalString[PreviousSeparatorIndex + 1] != '.' || LocalString[PreviousSeparatorIndex + 2] != '/'))
 						break;
 				}
 
 				// If we're attempting to remove the drive letter, that's illegal
-				int Colon = InPath.IndexOf(":", PreviousSeparatorIndex);
+				int Colon = LocalString.IndexOf(":", PreviousSeparatorIndex);
 				if (Colon >= 0 && Colon < Index)
 					return false;
 
-				InPath = InPath.Substring(0, PreviousSeparatorIndex) + InPath.Substring(Index + ParentDirLength);
+				LocalString = LocalString.Substring(0, PreviousSeparatorIndex) + LocalString.Substring(Index + ParentDirLength);
 			}
 
-			InPath = InPath.Replace("./", "");
+			LocalString = LocalString.Replace("./", "");
+
+			// restore back slashes now
+			if (bHadBackSlashes)
+			{
+				LocalString = LocalString.Replace("/", "\\");
+			}
 
+			// and pass back out
+			InPath = LocalString;
 			return true;
 		}
 
