[UE-4798] AutostartupModules is not working properly

The auto-startup module list is now compiled to:
1. Exe if the target is a program.
2. Game module binary if the target is game, editor etc.

diff --git a/Engine/Source/Programs/UnrealBuildTool/Configuration/UEBuildModule.cs b/Engine/Source/Programs/UnrealBuildTool/Configuration/UEBuildModule.cs
--- a/Engine/Source/Programs/UnrealBuildTool/Configuration/UEBuildModule.cs
+++ b/Engine/Source/Programs/UnrealBuildTool/Configuration/UEBuildModule.cs
@@ -931,6 +931,16 @@ namespace UnrealBuildTool
 			}
 		}
 
+		/**
+		 * Adds additional source cpp files for this module.
+		 *
+		 * @param Files Files to add.
+		 */
+		public void AddAdditionalCPPFiles(IEnumerable<FileItem> Files)
+		{
+			SourceFilesToBuild.CPPFiles.AddRange(Files);
+		}
+
 		/** A list of the absolute paths of source files to be built in this module. */
 		public readonly SourceFilesClass SourceFilesToBuild = new SourceFilesClass();
 
diff --git a/Engine/Source/Programs/UnrealBuildTool/Configuration/UEBuildTarget.cs b/Engine/Source/Programs/UnrealBuildTool/Configuration/UEBuildTarget.cs
--- a/Engine/Source/Programs/UnrealBuildTool/Configuration/UEBuildTarget.cs
+++ b/Engine/Source/Programs/UnrealBuildTool/Configuration/UEBuildTarget.cs
@@ -1826,6 +1826,16 @@ namespace UnrealBuildTool
 			return Output;
 		}
 
+		/// <summary>
+		/// Tests if the module provided is a game module.
+		/// </summary>
+		/// <param name="Module">Module to check.</param>
+		/// <returns>True if provided module is a game module. False otherwise.</returns>
+		private static bool IsGameModule(UEBuildModule Module)
+		{
+			return Module.Type == UEBuildModuleType.Game || (Module.Type == UEBuildModuleType.Runtime && Module.Name == "UE4Game");
+		}
+
 		/// <summary>
 		/// Creates a getter module that outputs auto-startup module list.
 		/// </summary>
@@ -1834,25 +1844,31 @@ namespace UnrealBuildTool
 			var AutoStartupModuleArray = GetAutoStartupModuleList().ToArray();
 
 			var ModuleName = "AutoStartupModuleListGetter";
-			var ModuleDir = Path.Combine(GlobalCompileEnvironment.Config.OutputDirectory, ModuleName);
 
-			if (!Directory.Exists(ModuleDir))
+			UEBuildModuleCPP Module = null;
+
+			if (TargetType != TargetRules.TargetType.Program)
 			{
-				Directory.CreateDirectory(ModuleDir);
+				Module = ExtraModuleNames.Select(Name => GetModuleByName(Name)).Where(
+					InnerModule => IsGameModule(InnerModule)
+				).First() as UEBuildModuleCPP;
+			}
+			else
+			{
+				Module = GetModuleByName(AppName) as UEBuildModuleCPP;
 			}
 
-			string SourceFilename = Path.Combine(ModuleDir, ModuleName + ".cpp");
-			string HeaderFilename = Path.Combine(ModuleDir, ModuleName + ".h");
+			var ModuleDir = Path.Combine(Module.Target.ProjectIntermediateDirectory, ModuleName);
 
-			if (!File.Exists(HeaderFilename))
+			if (!Directory.Exists(ModuleDir))
 			{
-				File.WriteAllText(HeaderFilename, "");
+				Directory.CreateDirectory(ModuleDir);
 			}
 
+			string SourceFilename = Path.Combine(ModuleDir, ModuleName + ".generated.cpp");
+
 			var SourceLines = new List<string>();
 
-			SourceLines.Add(string.Format("#include \"{0}\"", ModuleName + ".h"));
-			SourceLines.Add("");
 			SourceLines.Add("const char* EnumAutoStartupModuleName(int Index)");
 			SourceLines.Add("{");
 
@@ -1877,11 +1893,7 @@ namespace UnrealBuildTool
 				File.WriteAllLines(SourceFilename, SourceLines);
 			}
 
-			BindArtificialModuleToBinary(
-				CreateArtificialModule(
-					ModuleName, ModuleDir,
-					new FileItem[] { FileItem.GetItemByPath(SourceFilename) }, new string[] { }
-				), AppBinaries[0]);
+			(Module as UEBuildModuleCPP).AddAdditionalCPPFiles(new FileItem[] { FileItem.GetItemByPath(SourceFilename) });
 		}
 
 		/// <summary>
@@ -2246,28 +2258,36 @@ namespace UnrealBuildTool
 			// Add extra modules that will either link into the main binary (monolithic), or be linked into separate DLL files (modular)
 			foreach (var ModuleName in ExtraModuleNames)
 			{
-				if (ShouldCompileMonolithic())
-				{
-					// Add this module to the executable's list of included modules
-					var ExecutableBinary = AppBinaries[0];
-					ExecutableBinary.AddModule(ModuleName);
-				}
-				else
-				{
-					// Create a DLL binary for this module
-					string[] OutputFilePaths = MakeBinaryPaths(ModuleName, GetAppName() + "-" + ModuleName, UEBuildBinaryType.DynamicLinkLibrary, TargetType, null, AppName);
-					UEBuildBinaryConfiguration Config = new UEBuildBinaryConfiguration( InType: UEBuildBinaryType.DynamicLinkLibrary,
-																						InOutputFilePaths: OutputFilePaths,
-																						InIntermediateDirectory: RulesCompiler.IsGameModule(ModuleName) ? ProjectIntermediateDirectory : EngineIntermediateDirectory,
-																						bInAllowExports: true,
-																						InModuleNames: new List<string> { ModuleName } );
-
-					// Tell the target about this new binary
-					AppBinaries.Add(new UEBuildBinaryCPP(this, Config));
-				}
+				AddExtraModule(ModuleName);
 			}
 		}
 
+		/// <summary>
+		/// Adds extra module to the target.
+		/// </summary>
+		/// <param name="ModuleName">Name of the module.</param>
+		protected void AddExtraModule(string ModuleName)
+		{
+			if (ShouldCompileMonolithic())
+			{
+				// Add this module to the executable's list of included modules
+				var ExecutableBinary = AppBinaries[0];
+				ExecutableBinary.AddModule(ModuleName);
+			}
+			else
+			{
+				// Create a DLL binary for this module
+				string[] OutputFilePaths = MakeBinaryPaths(ModuleName, GetAppName() + "-" + ModuleName, UEBuildBinaryType.DynamicLinkLibrary, TargetType, null, AppName);
+				UEBuildBinaryConfiguration Config = new UEBuildBinaryConfiguration(InType: UEBuildBinaryType.DynamicLinkLibrary,
+																					InOutputFilePaths: OutputFilePaths,
+																					InIntermediateDirectory: RulesCompiler.IsGameModule(ModuleName) ? ProjectIntermediateDirectory : EngineIntermediateDirectory,
+																					bInAllowExports: true,
+																					InModuleNames: new List<string> { ModuleName });
+
+				// Tell the target about this new binary
+				AppBinaries.Add(new UEBuildBinaryCPP(this, Config));
+			}
+		}
 
 		/**
 		 * @return true if debug information is created, false otherwise.
diff --git a/Engine/Source/Runtime/Core/Private/Modules/ModuleManager.cpp b/Engine/Source/Runtime/Core/Private/Modules/ModuleManager.cpp
--- a/Engine/Source/Runtime/Core/Private/Modules/ModuleManager.cpp
+++ b/Engine/Source/Runtime/Core/Private/Modules/ModuleManager.cpp
@@ -986,3 +986,28 @@ bool FModuleManager::DoesLoadedModuleHaveUObjects( const FName ModuleName )
 
 	return false;
 }
+
+const ANSICHAR* (*GEnumAutoStartupModuleName)(int32 Index) = nullptr;
+
+void FModuleManager::InitializeAutoStartupModules()
+{
+	TArray<FString> Modules;
+	GetAutoStartupModuleList(Modules);
+
+	for (auto ModuleName : Modules)
+	{
+		LoadModuleChecked<IModuleInterface>(*ModuleName);
+	}
+}
+
+void FModuleManager::GetAutoStartupModuleList(TArray<FString>& OutModules) const
+{
+	checkf(GEnumAutoStartupModuleName != nullptr, TEXT("This callback should be set by the target binary."));
+
+	const ANSICHAR* NamePtr = nullptr;
+	int32 Index = 0;
+	while ((NamePtr = GEnumAutoStartupModuleName(Index++)) != nullptr)
+	{
+		OutModules.Add(FString(NamePtr));
+	}
+}
diff --git a/Engine/Source/Runtime/Core/Public/Modules/ModuleManager.h b/Engine/Source/Runtime/Core/Public/Modules/ModuleManager.h
--- a/Engine/Source/Runtime/Core/Public/Modules/ModuleManager.h
+++ b/Engine/Source/Runtime/Core/Public/Modules/ModuleManager.h
@@ -437,6 +437,11 @@ public:
 
 	virtual bool Exec( UWorld* Inworld, const TCHAR* Cmd, FOutputDevice& Ar ) override;
 
+	/**
+	 * Calls for each auto-startup module its StartupModule function.
+	 */
+	void InitializeAutoStartupModules();
+
 protected:
 
 	/**
@@ -510,6 +515,13 @@ private:
 	/** Finds modules matching a given name wildcard within a given directory. */
 	void FindModulePathsInDirectory(const FString &DirectoryName, bool bIsGameDirectory, const TCHAR *NamePattern, TMap<FName, FString> &OutModulePaths) const;
 
+	/**
+	 * Loads array of module names that needs to be auto-loaded.
+	 *
+	 * @param OutModules Output array to fill with module names.
+	 */
+	void GetAutoStartupModuleList(TArray<FString>& OutModules) const;
+
 private:
 
 	/** Map of all modules.  Maps the case-insensitive module name to information about that module, loaded or not. */
@@ -614,6 +626,17 @@ class FDefaultGameModuleImpl
 	}
 };
 
+/**
+ * Callback to retrieve auto-startup modules list.
+ *
+ * It should return string ptr to auto-startup module for consecutive indices
+ * until the list is finished. Then this function should return nullptr.
+ *
+ * @param Index Auto-startup module index to return.
+ *
+ * @returns Auto-startup module name.
+ */
+extern CORE_API const ANSICHAR* (*GEnumAutoStartupModuleName)(int32 Index);
 
 /**
  * Module implementation boilerplate for regular modules.
@@ -702,6 +725,16 @@ class FDefaultGameModuleImpl
 	#define IMPLEMENT_DEBUGGAME()
 #endif 
 
+#define REGISTER_AUTO_STARTUP_MODULE_LIST_GETTER() \
+	struct FRegisterAutoStartupModuleListEnumerator \
+	{ \
+		FRegisterAutoStartupModuleListEnumerator() \
+		{ \
+			extern const ANSICHAR* EnumAutoStartupModuleName(int Index); \
+			GEnumAutoStartupModuleName = &EnumAutoStartupModuleName; \
+		} \
+	} RegisterAutoStartupModuleListEnumerator;
+
 #if IS_PROGRAM
 
 	#if IS_MONOLITHIC
@@ -712,7 +745,8 @@ class FDefaultGameModuleImpl
 			IMPLEMENT_FOREIGN_ENGINE_DIR() \
 			IMPLEMENT_GAME_MODULE(FDefaultGameModuleImpl, ModuleName) \
 			PER_MODULE_BOILERPLATE \
-			FEngineLoop GEngineLoop;
+			FEngineLoop GEngineLoop; \
+			REGISTER_AUTO_STARTUP_MODULE_LIST_GETTER();
 
 	#else		
 
@@ -727,7 +761,8 @@ class FDefaultGameModuleImpl
 			} AutoSet##ModuleName; \
 			PER_MODULE_BOILERPLATE \
 			PER_MODULE_BOILERPLATE_ANYLINK(FDefaultGameModuleImpl, ModuleName) \
-			FEngineLoop GEngineLoop;
+			FEngineLoop GEngineLoop; \
+			REGISTER_AUTO_STARTUP_MODULE_LIST_GETTER();
 	#endif
 
 #else
@@ -751,7 +786,8 @@ class FDefaultGameModuleImpl
 			{ \
 				extern void UELinkerFixups(); \
 				UELinkerFixups(); \
-			}
+			} \
+			REGISTER_AUTO_STARTUP_MODULE_LIST_GETTER();
 
 	#else	//PLATFORM_DESKTOP
 
@@ -764,6 +800,7 @@ class FDefaultGameModuleImpl
 			IMPLEMENT_GAME_MODULE( ModuleImplClass, ModuleName ) \
 			/* Implement the GIsGameAgnosticExe variable (See Core.h). */ \
 			bool GIsGameAgnosticExe = false; \
+			REGISTER_AUTO_STARTUP_MODULE_LIST_GETTER();
 
 	#endif	//PLATFORM_DESKTOP
 
@@ -771,7 +808,8 @@ class FDefaultGameModuleImpl
 
 	#define IMPLEMENT_PRIMARY_GAME_MODULE( ModuleImplClass, ModuleName, GameName ) \
 		/* Nothing special to do for modular builds.  The game name will be set via the command-line */ \
-		IMPLEMENT_GAME_MODULE( ModuleImplClass, ModuleName )
+		IMPLEMENT_GAME_MODULE( ModuleImplClass, ModuleName ) \
+		REGISTER_AUTO_STARTUP_MODULE_LIST_GETTER();
 #endif	//IS_MONOLITHIC
 
 #endif
diff --git a/Engine/Source/Runtime/Launch/Private/LaunchEngineLoop.cpp b/Engine/Source/Runtime/Launch/Private/LaunchEngineLoop.cpp
--- a/Engine/Source/Runtime/Launch/Private/LaunchEngineLoop.cpp
+++ b/Engine/Source/Runtime/Launch/Private/LaunchEngineLoop.cpp
@@ -13,6 +13,8 @@
 #include "UProjectInfo.h"
 #include "EngineVersion.h"
 
+#include "ModuleManager.h"
+
 #if WITH_EDITOR
 	#include "EditorStyle.h"
 	#include "AutomationController.h"
@@ -652,40 +654,6 @@ bool IsServerDelegateForOSS(FName WorldContextHandle)
 }
 #endif
 
-/**
- * Loads array of module names that needs to be auto-loaded.
- *
- * @param OutModules Output array to fill with module names.
- */
-void GetAutoStartupModuleList(TArray<FString>& OutModules)
-{
-	// UBT generated function for listing all auto-startup module
-	// this binary is loading. If the list is finished then this function
-	// returns nullptr.
-	extern const ANSICHAR* EnumAutoStartupModuleName(int Index);
-
-	const ANSICHAR* NamePtr = nullptr;
-	int32 Index = 0;
-	while ((NamePtr = EnumAutoStartupModuleName(Index++)) != nullptr)
-	{
-		OutModules.Add(FString(NamePtr));
-	}
-}
-
-/**
- * Calls for each auto-startup module its StartupModule function.
- */
-void InitializeAutoStartupModules()
-{
-	TArray<FString> Modules;
-	GetAutoStartupModuleList(Modules);
-
-	for (auto ModuleName : Modules)
-	{
-		FModuleManager::LoadModuleChecked<IModuleInterface>(*ModuleName);
-	}
-}
-
 int32 FEngineLoop::PreInit( const TCHAR* CmdLine )
 {
 	DECLARE_SCOPE_CYCLE_COUNTER(TEXT("Engine Pre-Initialized"), STAT_PreInit, STATGROUP_LoadTime);
@@ -1407,9 +1375,7 @@ int32 FEngineLoop::PreInit( const TCHAR* CmdLine )
 	MarkObjectsToDisregardForGC(); 
 	GUObjectArray.CloseDisregardForGC();
 
-#if WITH_ENGINE
 	SetIsServerForOnlineSubsystemsDelegate(FQueryIsRunningServer::CreateStatic(&IsServerDelegateForOSS));
-#endif
 
 	SlowTask.EnterProgressFrame(50);
 
@@ -1521,7 +1487,7 @@ int32 FEngineLoop::PreInit( const TCHAR* CmdLine )
 			// Execute the commandlet.
 			double CommandletExecutionStartTime = FPlatformTime::Seconds();
 
-			InitializeAutoStartupModules();
+			FModuleManager::Get().InitializeAutoStartupModules();
 
 			// Commandlets don't always handle -run= properly in the commandline so we'll provide them
 			// with a custom version that doesn't have it.
@@ -1689,6 +1655,7 @@ int32 FEngineLoop::PreInit( const TCHAR* CmdLine )
 
 
 #else // WITH_ENGINE
+	FModuleManager::Get().InitializeAutoStartupModules();
 	EndInitTextLocalization();
 	FPlatformMisc::PlatformPostInit();
 #endif // WITH_ENGINE
@@ -1950,7 +1917,7 @@ int32 FEngineLoop::Init()
 		return 1;
 	}
 
-	InitializeAutoStartupModules();
+	FModuleManager::Get().InitializeAutoStartupModules();
 
 	GetMoviePlayer()->WaitForMovieToFinish();
 
diff --git a/Engine/Source/Runtime/Projects/Private/ProjectManager.cpp b/Engine/Source/Runtime/Projects/Private/ProjectManager.cpp
--- a/Engine/Source/Runtime/Projects/Private/ProjectManager.cpp
+++ b/Engine/Source/Runtime/Projects/Private/ProjectManager.cpp
@@ -93,6 +93,12 @@ bool FProjectManager::LoadModulesForProject( const ELoadingPhase::Type LoadingPh
 		}
 	}
 
+	if (LoadingPhase == ELoadingPhase::PreDefault && (!CurrentProject.IsValid() || CurrentProject->Modules.Num() == 0))
+	{
+		// If this is not a code-project load default game project module.
+		FModuleManager::Get().LoadModuleChecked("UE4Game");
+	}
+
 	return bSuccess;
 }
 
diff --git a/Engine/Source/Runtime/UE4Game/Private/UE4Game.cpp b/Engine/Source/Runtime/UE4Game/Private/UE4Game.cpp
--- a/Engine/Source/Runtime/UE4Game/Private/UE4Game.cpp
+++ b/Engine/Source/Runtime/UE4Game/Private/UE4Game.cpp
@@ -12,3 +12,5 @@ TCHAR GInternalGameName[64] = TEXT("");
 IMPLEMENT_DEBUGGAME()
 IMPLEMENT_FOREIGN_ENGINE_DIR()
 #endif
+
+REGISTER_AUTO_STARTUP_MODULE_LIST_GETTER();
diff --git a/Engine/Source/UE4Editor.Target.cs b/Engine/Source/UE4Editor.Target.cs
--- a/Engine/Source/UE4Editor.Target.cs
+++ b/Engine/Source/UE4Editor.Target.cs
@@ -16,6 +16,7 @@ public class UE4EditorTarget : TargetRules
 		ref List<string> OutExtraModuleNames
 		)
 	{
+		OutExtraModuleNames.Add("UE4Game");
 		OutExtraModuleNames.Add("GameMenuBuilder");
 		if ((Target.Platform == UnrealTargetPlatform.Win32) || (Target.Platform == UnrealTargetPlatform.Win64))
 		{
