addition of arm64 and armv7s architecture to the IOS project settings
UEPLAT-24
#ios

diff --git a/Engine/Source/Programs/UnrealBuildTool/IOS/IOSToolChain.cs b/Engine/Source/Programs/UnrealBuildTool/IOS/IOSToolChain.cs
--- a/Engine/Source/Programs/UnrealBuildTool/IOS/IOSToolChain.cs
+++ b/Engine/Source/Programs/UnrealBuildTool/IOS/IOSToolChain.cs
@@ -45,7 +45,7 @@ namespace UnrealBuildTool
 		[XmlConfig]
 		public static string NonShippingArchitectures = "armv7";
 		[XmlConfig]
-		public static string ShippingArchitectures = "armv7";
+		public static string ShippingArchitectures = "armv7,arm64";
 
 		/** Which version of the iOS to allow at build time */
 		[XmlConfig]
@@ -128,12 +128,70 @@ namespace UnrealBuildTool
 			}
 		}
 
+		public static void ParseArchitectures()
+		{
+			// look in ini settings for what platforms to compile for
+			ConfigCacheIni Ini = new ConfigCacheIni(UnrealTargetPlatform.Android, "Engine", UnrealBuildTool.GetUProjectPath());
+			List<string> ProjectArches = new List<string>();
+			bool bBuild = true;
+			if (Ini.GetBool("/Script/IOSRuntimeSettings.IOSRuntimeSettings", "bDevForArmV7", out bBuild) && bBuild)
+			{
+				ProjectArches.Add("armv7");
+			}
+			if (Ini.GetBool("/Script/IOSRuntimeSettings.IOSRuntimeSettings", "bDevForArm64", out bBuild) && bBuild)
+			{
+				ProjectArches.Add("arm64");
+			}
+			if (Ini.GetBool("/Script/IOSRuntimeSettings.IOSRuntimeSettings", "bDevForArmV7S", out bBuild) && bBuild)
+			{
+				ProjectArches.Add("armv7s");
+			}
+
+			// force armv7 if something went wrong
+			if (ProjectArches.Count == 0)
+			{
+				ProjectArches.Add("armv7");
+			}
+			NonShippingArchitectures = ProjectArches[0];
+			for (int Index = 1; Index < ProjectArches.Count; ++Index)
+			{
+				NonShippingArchitectures += "," + ProjectArches[Index];
+			}
+
+			ProjectArches.Clear();
+			if (Ini.GetBool("/Script/IOSRuntimeSettings.IOSRuntimeSettings", "bShipForArmV7", out bBuild) && bBuild)
+			{
+				ProjectArches.Add("armv7");
+			}
+			if (Ini.GetBool("/Script/IOSRuntimeSettings.IOSRuntimeSettings", "bShipForArm64", out bBuild) && bBuild)
+			{
+				ProjectArches.Add("arm64");
+			}
+			if (Ini.GetBool("/Script/IOSRuntimeSettings.IOSRuntimeSettings", "bShipForArmV7S", out bBuild) && bBuild)
+			{
+				ProjectArches.Add("armv7s");
+			}
+
+			// force armv7 if something went wrong
+			if (ProjectArches.Count == 0)
+			{
+				ProjectArches.Add("armv7");
+				ProjectArches.Add("arm64");
+			}
+			ShippingArchitectures = ProjectArches[0];
+			for (int Index = 1; Index < ProjectArches.Count; ++Index)
+			{
+				ShippingArchitectures += "," + ProjectArches[Index];
+			}
+		}
+
 		/** Hunt down the latest IOS sdk if desired */
 		public override void SetUpGlobalEnvironment()
 		{
 			base.SetUpGlobalEnvironment();
 
 			ParseProjectSettings();
+			ParseArchitectures();
 
 			if (IOSSDKVersion == "latest")
 			{
diff --git a/Engine/Source/Runtime/IOS/IOSRuntimeSettings/Classes/IOSRuntimeSettings.h b/Engine/Source/Runtime/IOS/IOSRuntimeSettings/Classes/IOSRuntimeSettings.h
--- a/Engine/Source/Runtime/IOS/IOSRuntimeSettings/Classes/IOSRuntimeSettings.h
+++ b/Engine/Source/Runtime/IOS/IOSRuntimeSettings/Classes/IOSRuntimeSettings.h
@@ -55,6 +55,30 @@ public:
 	UPROPERTY(GlobalConfig, EditAnywhere, Category = Rendering)
 	bool bSupportsOpenGLES2;
 
+	// Enable ArmV7 support? (this will be used if all type are unchecked)
+	UPROPERTY(GlobalConfig, EditAnywhere, Category = Build, meta = (DisplayName = "Support armv7 in Development"))
+	bool bDevForArmV7;
+
+	// Enable Arm64 support?
+	UPROPERTY(GlobalConfig, EditAnywhere, Category = Build, meta = (DisplayName = "Support arm64 in Development"))
+	bool bDevForArm64;
+
+	// Enable x86 support? [CURRENTLY FOR FULL SOURCE GAMES ONLY]
+	UPROPERTY(GlobalConfig, EditAnywhere, Category = Build, meta = (DisplayName = "Support armv7s in Development"))
+	bool bDevForArmV7S;
+
+	// Enable ArmV7 support? (this will be used if all type are unchecked)
+	UPROPERTY(GlobalConfig, EditAnywhere, Category = Build, meta = (DisplayName = "Support armv7 in Shipping"))
+	bool bShipForArmV7;
+
+	// Enable Arm64 support?
+	UPROPERTY(GlobalConfig, EditAnywhere, Category = Build, meta = (DisplayName = "Support arm64 in Shipping"))
+	bool bShipForArm64;
+
+	// Enable x86 support? [CURRENTLY FOR FULL SOURCE GAMES ONLY]
+	UPROPERTY(GlobalConfig, EditAnywhere, Category = Build, meta = (DisplayName = "Support armv7s in Shipping"))
+	bool bShipForArmV7S;
+
 	// Does the application support portrait orientation?
 	UPROPERTY(GlobalConfig, EditAnywhere, Category = DeviceOrientations)
 	uint32 bSupportsPortraitOrientation : 1;
diff --git a/Engine/Source/Runtime/IOS/IOSRuntimeSettings/Private/IOSRuntimeSettings.cpp b/Engine/Source/Runtime/IOS/IOSRuntimeSettings/Private/IOSRuntimeSettings.cpp
--- a/Engine/Source/Runtime/IOS/IOSRuntimeSettings/Private/IOSRuntimeSettings.cpp
+++ b/Engine/Source/Runtime/IOS/IOSRuntimeSettings/Private/IOSRuntimeSettings.cpp
@@ -18,6 +18,12 @@ UIOSRuntimeSettings::UIOSRuntimeSettings(const FObjectInitializer& ObjectInitial
 	bSupportsIPhone = true;
 	MinimumiOSVersion = EIOSVersion::IOS_6;
 	CompressionLevel = 0;
+	bDevForArmV7 = true;
+	bDevForArm64 = false;
+	bDevForArmV7S = false;
+	bShipForArmV7 = true;
+	bShipForArm64 = true;
+	bShipForArmV7S = false;
 }
 
 #if WITH_EDITOR
@@ -36,5 +42,15 @@ void UIOSRuntimeSettings::PostEditChangeProperty(struct FPropertyChangedEvent& P
 	{
 		bSupportsOpenGLES2 = true;
 	}
+
+	// Ensure that at least armv7 is selected for shipping and dev
+	if (!bDevForArmV7 && !bDevForArm64 && !bDevForArmV7S)
+	{
+		bDevForArmV7 = true;
+	}
+	if (!bShipForArmV7 && !bShipForArm64 && !bShipForArmV7S)
+	{
+		bShipForArmV7 = true;
+	}
 }
 #endif
