Fix for ShouldWriteJavaBuildSettingsFile substring length overrun (and protect against corrupt file)
#android
#ue4

diff --git a/Engine/Source/Programs/UnrealBuildTool/Android/UEDeployAndroid.cs b/Engine/Source/Programs/UnrealBuildTool/Android/UEDeployAndroid.cs
--- a/Engine/Source/Programs/UnrealBuildTool/Android/UEDeployAndroid.cs
+++ b/Engine/Source/Programs/UnrealBuildTool/Android/UEDeployAndroid.cs
@@ -270,9 +270,13 @@ namespace UnrealBuildTool.Android
         public bool ShouldWriteJavaBuildSettingsFile(string FileName, string setting)
         {
             var fileContent = File.ReadAllLines(FileName);
+            if (fileContent.Length < 5)
+                return true;
             var packageLine = fileContent[4]; // We know this to be true... because we write it below...
-            int location = packageLine.IndexOf("PACKAGING") + 12 + 12; // + (PACKAGING = ) + ("PackageType.")
-            return String.Compare(setting, packageLine.Substring(location, setting.Length)) != 0;
+            int location = packageLine.IndexOf("PACKAGING") + 12 + 12; // + ("PACKAGING = ") + ("PackageType.")
+            if (location == -1)
+                return true;
+            return String.Compare(setting, packageLine.Substring(location, Math.Min(packageLine.Length - location, setting.Length))) != 0;
         }
 
 		private static string GetNDKArch(string UE4Arch)
