Linux: support compiling for other architectures than x86_64.

diff --git a/Engine/Source/Programs/UnrealBuildTool/Linux/LinuxToolChain.cs b/Engine/Source/Programs/UnrealBuildTool/Linux/LinuxToolChain.cs
--- a/Engine/Source/Programs/UnrealBuildTool/Linux/LinuxToolChain.cs
+++ b/Engine/Source/Programs/UnrealBuildTool/Linux/LinuxToolChain.cs
@@ -157,12 +157,13 @@ namespace UnrealBuildTool
 
                 BaseLinuxPath = BaseLinuxPath.Replace("\"", "");
 
-                // set up the path to our toolchains (FIXME: support switching per architecture)
+                // set up the path to our toolchains
                 GCCPath = "";
-                ClangPath = Path.Combine(BaseLinuxPath, @"bin\Clang++.exe");
-                ArPath = Path.Combine(BaseLinuxPath, @"bin\x86_64-unknown-linux-gnu-ar.exe");
-                RanlibPath = Path.Combine(BaseLinuxPath, @"bin\x86_64-unknown-linux-gnu-ranlib.exe");
-            }
+                ClangPath = Path.Combine(BaseLinuxPath, @"bin\clang++.exe");
+				// ar and ranlib will be switched later to match the architecture
+				ArPath = "ar.exe";
+				RanlibPath = "ranlib.exe";
+			}
 
 			if (!DetermineCompilerVersion())
 			{
@@ -199,7 +200,61 @@ namespace UnrealBuildTool
 				(CompilerVersionMajor == Major && CompilerVersionMinor == Minor && CompilerVersionPatch >= Patch);
 		}
 
-        static string GetCLArguments_Global(CPPEnvironment CompileEnvironment)
+		/** Architecture-specific compiler switches */
+		static string ArchitectureSpecificSwitches(string Architecture)
+		{
+			string Result = "";
+
+			if (Architecture.StartsWith("x86_64"))
+			{
+				Result += " -mmmx -msse -msse2";
+			}
+			else if (Architecture.StartsWith("arm"))
+			{
+				Result += " -fsigned-char";
+
+				// FreeType2's ftconfig.h will trigger this
+				Result += " -Wno-deprecated-register";
+			}
+
+			return Result;
+		}
+
+		static string ArchitectureSpecificDefines(string Architecture)
+		{
+			string Result = "";
+
+			if (Architecture.StartsWith("x86_64"))
+			{
+				Result += " -D_LINUX64";
+			}
+
+			return Result;
+		}
+
+		/** Gets architecture-specific ar paths */
+		private static string GetArPath(string Architecture)
+		{
+			if (CrossCompiling())
+			{
+				return Path.Combine(Path.Combine(BaseLinuxPath, String.Format("bin/{0}-{1}", Architecture, ArPath)));
+			}
+
+			return ArPath;
+		}
+
+		/** Gets architecture-specific ranlib paths */
+		private static string GetRanlibPath(string Architecture)
+		{
+			if (CrossCompiling())
+			{
+				return Path.Combine(Path.Combine(BaseLinuxPath, String.Format("bin/{0}-{1}", Architecture, RanlibPath)));
+			}
+
+			return RanlibPath;
+		}
+
+		static string GetCLArguments_Global(CPPEnvironment CompileEnvironment)
         {
             string Result = "";
 
@@ -225,8 +280,10 @@ namespace UnrealBuildTool
             Result += " -Wsequence-point";              // additional warning not normally included in Wall: warns if order of operations is ambigious
             //Result += " -Wunreachable-code";            // additional warning not normally included in Wall: warns if there is code that will never be executed - not helpful due to bIsGCC and similar 
             //Result += " -Wshadow";                      // additional warning not normally included in Wall: warns if there variable/typedef shadows some other variable - not helpful because we have gobs of code that shadows variables
-            Result += " -mmmx -msse -msse2";            // allows use of SIMD intrinsics
-            Result += " -fno-math-errno";               // do not assume that math ops have side effects
+
+			Result += ArchitectureSpecificSwitches(CompileEnvironment.Config.Target.Architecture);
+
+			Result += " -fno-math-errno";               // do not assume that math ops have side effects
             Result += " -fno-rtti";                     // no run-time type info
 
             // these needs to be removed ASAP
@@ -325,10 +382,8 @@ namespace UnrealBuildTool
 
             //Result += " -v";                            // for better error diagnosis
 
-            // assume we will not perform 32 bit builds on Linux, so define
-            // _LINUX64 in any case
-            Result += " -D_LINUX64";
-            if (CrossCompiling())
+			Result += ArchitectureSpecificDefines(CompileEnvironment.Config.Target.Architecture);
+			if (CrossCompiling())
             {
                 if (UsingClang())
                 {
@@ -668,7 +723,7 @@ namespace UnrealBuildTool
             ArchiveAction.ProducedItems.Add(OutputFile);
             ArchiveAction.CommandDescription = "Archive";
             ArchiveAction.StatusDescription = Path.GetFileName(OutputFile.AbsolutePath);
-            ArchiveAction.CommandArguments += string.Format("{0} {1} \"{2}\"",  ArPath, GetArchiveArguments(LinkEnvironment), OutputFile.AbsolutePath);
+            ArchiveAction.CommandArguments += string.Format("{0} {1} \"{2}\"",  GetArPath(LinkEnvironment.Config.Target.Architecture), GetArchiveArguments(LinkEnvironment), OutputFile.AbsolutePath);
 
             // Add the input files to a response file, and pass the response file on the command-line.
             List<string> InputFileNames = new List<string>();
@@ -681,7 +736,7 @@ namespace UnrealBuildTool
             }
 
             // add ranlib
-            ArchiveAction.CommandArguments += string.Format(" && {0} \"{1}\"", RanlibPath, OutputFile.AbsolutePath);
+            ArchiveAction.CommandArguments += string.Format(" && {0} \"{1}\"", GetRanlibPath(LinkEnvironment.Config.Target.Architecture), OutputFile.AbsolutePath);
 
             // Add the additional arguments specified by the environment.
             ArchiveAction.CommandArguments += LinkEnvironment.Config.AdditionalArguments;
diff --git a/Engine/Source/Programs/UnrealBuildTool/Linux/UEBuildLinux.cs b/Engine/Source/Programs/UnrealBuildTool/Linux/UEBuildLinux.cs
--- a/Engine/Source/Programs/UnrealBuildTool/Linux/UEBuildLinux.cs
+++ b/Engine/Source/Programs/UnrealBuildTool/Linux/UEBuildLinux.cs
@@ -11,13 +11,19 @@ namespace UnrealBuildTool
     class LinuxPlatform : UEBuildPlatform
     {
         /** This is the SDK version we support */
-		static private string ExpectedSDKVersion = "v4_clang-3.5.0_ld-2.24_glibc-2.12.2";
+		static private Dictionary<string, string> ExpectedSDKVersions = new Dictionary<string, string>()
+		{
+			{ "x86_64-unknown-linux-gnu", "v4_clang-3.5.0_ld-2.24_glibc-2.12.2" },
+			{ "arm-unknown-linux-gnueabihf", "arm-unknown-linux-gnueabihf_v5_clang-3.5.0-ld-2.23.1-glibc-2.13" },
+		};
 
         /** Platform name (embeds architecture for now) */
         static private string TargetPlatformName = "Linux_x64";
 
         /** Linux architecture (compiler target triplet) */
-        static private string DefaultArchitecture = "x86_64-unknown-linux-gnu";
+		// FIXME: for now switching between architectures is hard-coded
+		static private string DefaultArchitecture = "x86_64-unknown-linux-gnu";
+		//static private string DefaultArchitecture = "arm-unknown-linux-gnueabihf";
 
         /** The current architecture */
         public override string GetActiveArchitecture()
@@ -61,7 +67,13 @@ namespace UnrealBuildTool
          */
         protected override string GetRequiredSDKString()
         {
-            return ExpectedSDKVersion;
+			string SDKString;
+			if (!ExpectedSDKVersions.TryGetValue(GetActiveArchitecture(), out SDKString))
+			{
+				throw new BuildException("LinuxPlatform::GetRequiredSDKString: no toolchain set up for architecture '{0}'", GetActiveArchitecture());
+			}
+
+			return SDKString;
         }
 
         protected override String GetRequiredScriptVersionString()
diff --git a/Engine/Source/Runtime/Core/Private/Linux/LinuxPlatformCrashContext.cpp b/Engine/Source/Runtime/Core/Private/Linux/LinuxPlatformCrashContext.cpp
--- a/Engine/Source/Runtime/Core/Private/Linux/LinuxPlatformCrashContext.cpp
+++ b/Engine/Source/Runtime/Core/Private/Linux/LinuxPlatformCrashContext.cpp
@@ -336,7 +336,7 @@ bool FLinuxCrashContext::GetInfoForAddress(void * Address, const char **OutFunct
 
 		if (NumLines >= kMaxBufferLinesAllowed)
 		{
-			fprintf(stderr, "Number of lines associated with a DIE looks unreasonable (%ld), early quitting.\n", NumLines);
+			fprintf(stderr, "Number of lines associated with a DIE looks unreasonable (%d), early quitting.\n", static_cast<int32>(NumLines));
 			ReturnCode = DW_DLE_DIE_NO_CU_CONTEXT;	// invalidate
 			break;
 		}
diff --git a/Engine/Source/Runtime/Core/Public/Linux/LinuxPlatform.h b/Engine/Source/Runtime/Core/Public/Linux/LinuxPlatform.h
--- a/Engine/Source/Runtime/Core/Public/Linux/LinuxPlatform.h
+++ b/Engine/Source/Runtime/Core/Public/Linux/LinuxPlatform.h
@@ -44,8 +44,13 @@ typedef FLinuxPlatformTypes FPlatformTypes;
 #define PLATFORM_HAS_BSD_IPV6_SOCKETS				1
 
 #define PLATFORM_USES_DYNAMIC_RHI					1
-#define PLATFORM_ENABLE_VECTORINTRINSICS			1
 
+// only enable vectorintrinsics on x86 for now
+#if defined(_M_IX86) || defined(__i386__)
+	#define PLATFORM_ENABLE_VECTORINTRINSICS		1
+#else
+	#define PLATFORM_ENABLE_VECTORINTRINSICS		0
+#endif // defined(_M_IX86) || defined(__i386__)
 
 // Function type macros.
 #define VARARGS													/* Functions with variable arguments */
diff --git a/Engine/Source/Runtime/OpenGLDrv/Private/OpenGLShaders.cpp b/Engine/Source/Runtime/OpenGLDrv/Private/OpenGLShaders.cpp
--- a/Engine/Source/Runtime/OpenGLDrv/Private/OpenGLShaders.cpp
+++ b/Engine/Source/Runtime/OpenGLDrv/Private/OpenGLShaders.cpp
@@ -10,7 +10,7 @@
 
 #define CHECK_FOR_GL_SHADERS_TO_REPLACE 0
 
-#if PLATFORM_WINDOWS || PLATFORM_LINUX
+#if PLATFORM_WINDOWS
 #include <mmintrin.h>
 #elif PLATFORM_MAC
 #include <xmmintrin.h>
