dSYM fixup for IOS deployment from XCode

diff --git a/Engine/Source/Programs/UnrealBuildTool/IOS/IOSToolChain.cs b/Engine/Source/Programs/UnrealBuildTool/IOS/IOSToolChain.cs
--- a/Engine/Source/Programs/UnrealBuildTool/IOS/IOSToolChain.cs
+++ b/Engine/Source/Programs/UnrealBuildTool/IOS/IOSToolChain.cs
@@ -867,7 +867,7 @@ namespace UnrealBuildTool
 			LinkAction.StatusDescription = string.Format("{0}", OutputFile.AbsolutePath);
 			LinkAction.OutputEventHandler = new DataReceivedEventHandler(RemoteOutputReceivedEventHandler);
 			// For iPhone, generate the dSYM file if the config file is set to do so
-			if (BuildConfiguration.bGeneratedSYMFile == true)
+			if (BuildConfiguration.bGeneratedSYMFile == true && Path.GetExtension(OutputFile.AbsolutePath) != ".a")
 			{
 				Log.TraceInformation("Generating the dSYM file - this will add some time to your build...");
 				RemoteOutputFile = GenerateDebugInfo(RemoteOutputFile);
@@ -886,7 +886,16 @@ namespace UnrealBuildTool
 			// Make a file item for the source and destination files
 			string FullDestPathRoot = Executable.AbsolutePath + ".app.dSYM";
 			string FullDestPath = FullDestPathRoot;
-			FileItem DestFile = FileItem.GetRemoteItemByPath(FullDestPath, UnrealTargetPlatform.IOS);
+
+			FileItem DestFile;
+			if (!Utils.IsRunningOnMono && BuildHostPlatform.Current.Platform != UnrealTargetPlatform.Mac) 
+			{
+				DestFile = FileItem.GetRemoteItemByPath (FullDestPath, UnrealTargetPlatform.IOS);
+			}
+			else 
+			{
+				DestFile = FileItem.GetItemByPath (FullDestPath);
+			}
 
 			// Make the compile action
 			Action GenDebugAction = new Action(ActionType.GenerateDebugInfo);
@@ -900,11 +909,20 @@ namespace UnrealBuildTool
 			GenDebugAction.CommandPath = "sh";
 
 			// note that the source and dest are switched from a copy command
-			GenDebugAction.CommandArguments = string.Format("-c '{0}/usr/bin/dsymutil {1} -o {2}; cd {2}/..; zip -r -y -1 {3}.app.dSYM.zip {3}.app.dSYM'",
-				XcodeDeveloperDir,
-				Executable.AbsolutePath,
-				FullDestPathRoot,
-				Path.GetFileName(Executable.AbsolutePath));
+			if (!Utils.IsRunningOnMono && BuildHostPlatform.Current.Platform != UnrealTargetPlatform.Mac) 
+			{
+				GenDebugAction.CommandArguments = string.Format("-c '/usr/bin/dsymutil {0} -o {1}; cd {1}/..; zip -r -y -1 {2}.app.dSYM.zip {2}.app.dSYM'",
+					Executable.AbsolutePath,
+					FullDestPathRoot,
+					Path.GetFileName(Executable.AbsolutePath));
+			} 
+			else
+			{
+				GenDebugAction.CommandArguments = string.Format("-c '/usr/bin/dsymutil {0} -o {1}'",
+					Executable.AbsolutePath,
+					FullDestPathRoot);
+			}
+
 			GenDebugAction.PrerequisiteItems.Add(Executable);
 			GenDebugAction.ProducedItems.Add(DestFile);
 			GenDebugAction.StatusDescription = GenDebugAction.CommandArguments;// string.Format("Generating debug info for {0}", Path.GetFileName(Executable.AbsolutePath));
@@ -1056,13 +1074,15 @@ namespace UnrealBuildTool
 					RunExecutableAndWait( "unzip", String.Format( "-o \"{0}\" -d \"{1}\"", ZipSrcPath, ZipDstPath ), out ResultsText );
 					continue;
 				}
-
-				// Use RPC utility if the zip is on remote mac
-				Hashtable Result = RPCUtilHelper.Command( "/", String.Format( "unzip -o \"{0}\" -d \"{1}\"", ZipSrcPath, ZipDstPath ), "", null );
-
-				foreach ( DictionaryEntry Entry in Result )
+				else
 				{
-					Log.TraceInformation( "{0}", Entry.Value );
+					// Use RPC utility if the zip is on remote mac
+					Hashtable Result = RPCUtilHelper.Command( "/", String.Format( "unzip -o \"{0}\" -d \"{1}\"", ZipSrcPath, ZipDstPath ), "", null );
+
+					foreach ( DictionaryEntry Entry in Result )
+					{
+						Log.TraceInformation( "{0}", Entry.Value );
+					}
 				}
 			}
 		}
@@ -1253,7 +1273,7 @@ namespace UnrealBuildTool
 					string FinalRemoteExecutablePath = String.Format("{0}/Payload/{1}.app/{1}", RemoteShadowDirectoryMac, Target.GameName);
 					RPCUtilHelper.Command("/", String.Format("cp -f {0} {1}", RemoteExecutablePath, FinalRemoteExecutablePath), "", null);
 				}
-				else
+				else if(!Utils.IsRunningOnMono && BuildHostPlatform.Current.Platform != UnrealTargetPlatform.Mac)
 				{
 					RPCUtilHelper.CopyFile(RemoteExecutablePath, Target.OutputPath, false);
 
