Fix for AutoStartupModuleList, so it always builds UE4Game shared library in non-monolithic builds which is now required.

diff --git a/Engine/Source/Programs/UnrealBuildTool/Configuration/UEBuildEditor.cs b/Engine/Source/Programs/UnrealBuildTool/Configuration/UEBuildEditor.cs
--- a/Engine/Source/Programs/UnrealBuildTool/Configuration/UEBuildEditor.cs
+++ b/Engine/Source/Programs/UnrealBuildTool/Configuration/UEBuildEditor.cs
@@ -58,6 +58,13 @@ namespace UnrealBuildTool
 																					bInAllowExports:!ShouldCompileMonolithic(),
 																					InModuleNames: new List<string>() { "Launch" } );
 
+				// This "default" game module is required by the editor executable in case
+				// you want to run it in content-only mode.
+				if (!ExtraModuleNames.Contains("UE4Game"))
+				{
+					ExtraModuleNames.Add("UE4Game");
+				}
+
 				AppBinaries.Add( new UEBuildBinaryCPP( this, Config ) );
 			}
 
diff --git a/Engine/Source/Programs/UnrealBuildTool/Configuration/UEBuildTarget.cs b/Engine/Source/Programs/UnrealBuildTool/Configuration/UEBuildTarget.cs
--- a/Engine/Source/Programs/UnrealBuildTool/Configuration/UEBuildTarget.cs
+++ b/Engine/Source/Programs/UnrealBuildTool/Configuration/UEBuildTarget.cs
@@ -1810,16 +1810,27 @@ namespace UnrealBuildTool
 		/// <summary>
 		/// Gets auto-startup module list.
 		/// </summary>
+		/// <param name="Module">Game module to gather auto-startup list.</param>
 		/// <returns>Auto-startup module name list.</returns>
-		private IEnumerable<string> GetAutoStartupModuleList()
+		private IEnumerable<string> GetAutoStartupModuleList(UEBuildModuleCPP Module)
 		{
 			var Output = new List<string>();
 
-			foreach(var Module in AppBinaries[0].GetAllDependencyModules(false, true))
+			UEBuildBinary Binary = null;
+			if(Module.Name == "UE4Game")
 			{
-				if (Module is UEBuildModuleCPP && (Module as UEBuildModuleCPP).bIsAutoStartupModule)
+				Binary = AppBinaries[0];
+			}
+			else
+			{
+				Binary = Module.Binary;
+			}
+
+			foreach (var DepModule in Binary.GetAllDependencyModules(false, true))
+			{
+				if (DepModule is UEBuildModuleCPP && (DepModule as UEBuildModuleCPP).bIsAutoStartupModule)
 				{
-					Output.Add(Module.Name);
+					Output.Add(DepModule.Name);
 				}
 			}
 
@@ -1837,35 +1848,44 @@ namespace UnrealBuildTool
 		}
 
 		/// <summary>
-		/// Creates a getter module that outputs auto-startup module list.
+		/// Creates a getter that outputs auto-startup module list.
 		/// </summary>
 		private void CreateAutoStartupModuleListGetter()
 		{
-			var AutoStartupModuleArray = GetAutoStartupModuleList().ToArray();
-
-			var ModuleName = "AutoStartupModuleListGetter";
-
-			UEBuildModuleCPP Module = null;
-
 			if (TargetType != TargetRules.TargetType.Program)
 			{
-				Module = ExtraModuleNames.Select(Name => GetModuleByName(Name)).Where(
-					InnerModule => IsGameModule(InnerModule)
-				).First() as UEBuildModuleCPP;
+				foreach(var Module in ExtraModuleNames.Select(Name => GetModuleByName(Name)).Where(
+						InnerModule => IsGameModule(InnerModule)
+					).Cast<UEBuildModuleCPP>())
+				{
+					CreateAutoGeneratedModulesListForSpecificModule(Module);
+				}
 			}
 			else
 			{
-				Module = GetModuleByName(AppName) as UEBuildModuleCPP;
+				CreateAutoGeneratedModulesListForSpecificModule(
+					GetModuleByName(AppName) as UEBuildModuleCPP);
 			}
+		}
 
-			var ModuleDir = Path.Combine(Module.Target.ProjectIntermediateDirectory, ModuleName);
+		/// <summary>
+		/// Creates a getter that outputs auto-startup module list.
+		/// </summary>
+		/// <param name="Module">Module to create auto-startup list for.</param>
+		private void CreateAutoGeneratedModulesListForSpecificModule(UEBuildModuleCPP Module)
+		{
+			var AutoStartupModuleArray = GetAutoStartupModuleList(Module).ToArray();
+
+			var BaseFileName = "AutoStartupModuleListGetter";
+
+			var ModuleDir = Path.Combine(Module.Binary.Config.IntermediateDirectory, BaseFileName);
 
 			if (!Directory.Exists(ModuleDir))
 			{
 				Directory.CreateDirectory(ModuleDir);
 			}
 
-			string SourceFilename = Path.Combine(ModuleDir, ModuleName + ".generated.cpp");
+			string SourceFilename = Path.Combine(ModuleDir, BaseFileName + ".generated.cpp");
 
 			var SourceLines = new List<string>();
 
@@ -1893,9 +1913,8 @@ namespace UnrealBuildTool
 				File.WriteAllLines(SourceFilename, SourceLines);
 			}
 
-			(Module as UEBuildModuleCPP).AddAdditionalCPPFiles(new FileItem[] { FileItem.GetItemByPath(SourceFilename) });
+			Module.AddAdditionalCPPFiles(new FileItem[] { FileItem.GetItemByPath(SourceFilename) });
 		}
-
 		/// <summary>
 		/// All non-program monolithic binaries implicitly depend on all static plugin libraries so they are always linked appropriately
 		/// In order to do this, we create a new module here with a cpp file we emit that invokes an empty function in each library.
