Unable to perform hot-reload due to compilation errors. Fixed.

Additionally fixed some monolithic compilation errors on CIS.

UE-4845

diff --git a/Engine/Source/Programs/UnrealBuildTool/Configuration/UEBuildTarget.cs b/Engine/Source/Programs/UnrealBuildTool/Configuration/UEBuildTarget.cs
--- a/Engine/Source/Programs/UnrealBuildTool/Configuration/UEBuildTarget.cs
+++ b/Engine/Source/Programs/UnrealBuildTool/Configuration/UEBuildTarget.cs
@@ -1513,10 +1513,7 @@ namespace UnrealBuildTool
 				CreateLinkerFixupsCPPFile();
 			}
 
-			if (!bEditorRecompile && !UEBuildConfiguration.bHotReloadFromIDE)
-			{
-				CreateAutoStartupModuleListGetter();
-			}
+			CreateAutoStartupModuleListGetters();
 
 			// Build the target's binaries.
 			foreach (var Binary in AppBinaries)
@@ -1810,22 +1807,12 @@ namespace UnrealBuildTool
 		/// <summary>
 		/// Gets auto-startup module list.
 		/// </summary>
-		/// <param name="Module">Game module to gather auto-startup list.</param>
+		/// <param name="Binary">Binary to gather auto-startup list.</param>
 		/// <returns>Auto-startup module name list.</returns>
-		private IEnumerable<string> GetAutoStartupModuleList(UEBuildModuleCPP Module)
+		private IEnumerable<string> GetAutoStartupModuleList(UEBuildBinary Binary)
 		{
 			var Output = new List<string>();
 
-			UEBuildBinary Binary = null;
-			if(Module.Name == "UE4Game")
-			{
-				Binary = AppBinaries[0];
-			}
-			else
-			{
-				Binary = Module.Binary;
-			}
-
 			foreach (var DepModule in Binary.GetAllDependencyModules(false, true))
 			{
 				if (DepModule is UEBuildModuleCPP && (DepModule as UEBuildModuleCPP).bIsAutoStartupModule)
@@ -1848,53 +1835,84 @@ namespace UnrealBuildTool
 		}
 
 		/// <summary>
-		/// Creates a getter that outputs auto-startup module list.
+		/// Creates getters that passes auto-startup module list to the engine.
 		/// </summary>
-		private void CreateAutoStartupModuleListGetter()
+		private void CreateAutoStartupModuleListGetters()
 		{
-			if (TargetType != TargetRules.TargetType.Program)
+			if (TargetType != TargetRules.TargetType.Program && !ShouldCompileMonolithic())
 			{
-				foreach(var Module in ExtraModuleNames.Select(Name => GetModuleByName(Name)).Where(
-						InnerModule => IsGameModule(InnerModule)
-					).Cast<UEBuildModuleCPP>())
+				foreach (var Module in ExtraModuleNames.Select(Name => GetModuleByName(Name)).Where(
+						InnerModule => IsGameModule(InnerModule)).Cast<UEBuildModuleCPP>())
 				{
-					CreateAutoGeneratedModulesListForSpecificModule(Module);
+					UEBuildBinary Binary = null;
+					if (Module.Name == "UE4Game")
+					{
+						Binary = AppBinaries[0];
+					}
+					else
+					{
+						Binary = Module.Binary;
+					}
+
+					Module.AddAdditionalCPPFiles(new FileItem[] { FileItem.GetItemByPath(
+							CreateAutoGeneratedModulesListGetterCPPFile(
+								GetAutoStartupModuleList(Binary),
+								Module.Binary.Config.IntermediateDirectory)
+						)
+					});
 				}
 			}
 			else
 			{
-				CreateAutoGeneratedModulesListForSpecificModule(
-					GetModuleByName(AppName) as UEBuildModuleCPP);
+				var SourceFilename = CreateAutoGeneratedModulesListGetterCPPFile(
+					GetAutoStartupModuleList(AppBinaries[0]),
+					AppBinaries[0].Config.IntermediateDirectory);
+
+				BindArtificialModuleToBinary(
+					CreateArtificialModule(
+						"AutoStartupModulesListGetter", AppBinaries[0].Config.IntermediateDirectory,
+						new FileItem[] { FileItem.GetItemByPath(SourceFilename) }, new string[] { }
+					), AppBinaries[0]);
 			}
 		}
 
 		/// <summary>
 		/// Creates a getter that outputs auto-startup module list.
 		/// </summary>
-		/// <param name="Module">Module to create auto-startup list for.</param>
-		private void CreateAutoGeneratedModulesListForSpecificModule(UEBuildModuleCPP Module)
+		/// <param name="AutoStartupModuleEnumerable">A list of auto-startup modules.</param>
+		/// <param name="DestPath">Destination path for the file.</param>
+		/// <returns>Created source file path.</returns>
+		private string CreateAutoGeneratedModulesListGetterCPPFile(IEnumerable<string> AutoStartupModuleEnumerable, string DestPath)
 		{
-			var AutoStartupModuleArray = GetAutoStartupModuleList(Module).ToArray();
-
 			var BaseFileName = "AutoStartupModuleListGetter";
 
-			var ModuleDir = Path.Combine(Module.Binary.Config.IntermediateDirectory, BaseFileName);
+			var ModuleDir = Path.Combine(DestPath, BaseFileName);
 
 			if (!Directory.Exists(ModuleDir))
 			{
 				Directory.CreateDirectory(ModuleDir);
 			}
 
+			var HeaderName = BaseFileName + ".h";
+			var HeaderPath = Path.Combine(ModuleDir, HeaderName);
+
+			if(!File.Exists(HeaderPath))
+			{
+				File.WriteAllText(HeaderPath, "");
+			}
+
 			string SourceFilename = Path.Combine(ModuleDir, BaseFileName + ".generated.cpp");
 
 			var SourceLines = new List<string>();
 
+			SourceLines.Add(string.Format("#include \"{0}\"", HeaderName));
+			SourceLines.Add("");
 			SourceLines.Add("const char* EnumAutoStartupModuleName(int Index)");
 			SourceLines.Add("{");
 
 			SourceLines.Add("\tconst char* Names[] = {");
 
-			foreach (var AutoStartupModuleName in AutoStartupModuleArray)
+			foreach (var AutoStartupModuleName in AutoStartupModuleEnumerable.OrderBy(ModuleName => ModuleName))
 			{
 				SourceLines.Add(string.Format("\t\t\"{0}\",", AutoStartupModuleName));
 			}
@@ -1913,8 +1931,9 @@ namespace UnrealBuildTool
 				File.WriteAllLines(SourceFilename, SourceLines);
 			}
 
-			Module.AddAdditionalCPPFiles(new FileItem[] { FileItem.GetItemByPath(SourceFilename) });
+			return SourceFilename;
 		}
+
 		/// <summary>
 		/// All non-program monolithic binaries implicitly depend on all static plugin libraries so they are always linked appropriately
 		/// In order to do this, we create a new module here with a cpp file we emit that invokes an empty function in each library.
