Don't use exceptions to check whether HTML5 is not installed - it prevents running through Visual Studio with break-on-exceptions turned on.

diff --git a/Engine/Source/Programs/UnrealBuildTool/HTML5/HTML5SDKInfo.cs b/Engine/Source/Programs/UnrealBuildTool/HTML5/HTML5SDKInfo.cs
--- a/Engine/Source/Programs/UnrealBuildTool/HTML5/HTML5SDKInfo.cs
+++ b/Engine/Source/Programs/UnrealBuildTool/HTML5/HTML5SDKInfo.cs
@@ -139,15 +139,7 @@ namespace UnrealBuildTool
 
         public static bool IsSDKInstalled()
         {
-            try 
-            {
-                EmscriptenVersion();
-                return true; 
-            }
-            catch (Exception /*ex*/ )
-            {
-                return false; 
-            }
+			return File.Exists(GetVersionInfoPath());
         }
 
         public static bool IsPythonInstalled()
@@ -159,11 +151,16 @@ namespace UnrealBuildTool
 
         public static string EmscriptenVersion()
         {
-            string BaseSDKPath = HTML5SDKInfo.EmscriptenSDKPath(); 
-            string VersionInfo = File.ReadAllText( Path.Combine(BaseSDKPath,"emscripten-version.txt"));
+            string VersionInfo = File.ReadAllText(GetVersionInfoPath());
             VersionInfo = VersionInfo.Trim();
             return VersionInfo; 
         }
+
+		static string GetVersionInfoPath()
+		{
+            string BaseSDKPath = HTML5SDKInfo.EmscriptenSDKPath(); 
+            return Path.Combine(BaseSDKPath,"emscripten-version.txt");
+		}
     }
 }
  
\ No newline at end of file
diff --git a/Engine/Source/Programs/UnrealBuildTool/HTML5/UEBuildHTML5.cs b/Engine/Source/Programs/UnrealBuildTool/HTML5/UEBuildHTML5.cs
--- a/Engine/Source/Programs/UnrealBuildTool/HTML5/UEBuildHTML5.cs
+++ b/Engine/Source/Programs/UnrealBuildTool/HTML5/UEBuildHTML5.cs
@@ -80,6 +80,10 @@ namespace UnrealBuildTool
 			{
 				return SDKStatus.Invalid;
             }
+			if (!HTML5SDKInfo.IsSDKInstalled())
+			{
+				return SDKStatus.Invalid;
+			}
             try
             {
                 if (HTML5SDKInfo.EmscriptenVersion().Contains(ExpectedSDKVersion))
