Fix for attempting to delete nonexistant directory during packaging for Android
#android
#ue4
#codeview Josh.Adams

diff --git a/Engine/Source/Programs/UnrealBuildTool/Android/UEDeployAndroid.cs b/Engine/Source/Programs/UnrealBuildTool/Android/UEDeployAndroid.cs
--- a/Engine/Source/Programs/UnrealBuildTool/Android/UEDeployAndroid.cs
+++ b/Engine/Source/Programs/UnrealBuildTool/Android/UEDeployAndroid.cs
@@ -205,7 +205,16 @@ namespace UnrealBuildTool.Android
 			}
 			
 			// delete all files in here
-			string[] Files = Directory.GetFiles(InPath);
+            string[] Files;
+            try
+            {
+                Files = Directory.GetFiles(InPath);
+            }
+            catch (Exception)
+            {
+                // directory doesn't exist so all is good
+                return;
+            }
 			foreach (string Filename in Files)
 			{
 				try
