fix for duplicate scheme names when building IOS
UE-4506
#ios

diff --git a/Engine/Source/Programs/UnrealBuildTool/System/XcodeProjectFileGenerator.cs b/Engine/Source/Programs/UnrealBuildTool/System/XcodeProjectFileGenerator.cs
--- a/Engine/Source/Programs/UnrealBuildTool/System/XcodeProjectFileGenerator.cs
+++ b/Engine/Source/Programs/UnrealBuildTool/System/XcodeProjectFileGenerator.cs
@@ -1111,7 +1111,7 @@ namespace UnrealBuildTool
 
 					// if the project is not an engine project check to make sure we have the correct name
 					string DisplayName = TargetName;
-					if (!IsEngineTarget && TargetRulesObject.Type != TargetRules.TargetType.Program)
+					if (!IsEngineTarget && TargetRulesObject.Type != TargetRules.TargetType.Program && TargetRulesObject.Type != TargetRules.TargetType.Client)
 					{
 						List<UProjectInfo> AllGames = UProjectInfo.FilterGameProjects(true, bGeneratingGameProjectFiles ? GameProjectName : null);
 						UProjectInfo ProjectInfo = FindGameContainingFile (AllGames, TargetFilePath);
