UBT: Changes to project generators.

- Separates CMake generator from makefile one and adds qmake generator.
- Contributed by [REDACTED_PERSON] (PR #560).

diff --git a/Engine/Build/BatchFiles/Linux/GenerateProjectFiles.sh b/Engine/Build/BatchFiles/Linux/GenerateProjectFiles.sh
--- a/Engine/Build/BatchFiles/Linux/GenerateProjectFiles.sh
+++ b/Engine/Build/BatchFiles/Linux/GenerateProjectFiles.sh
@@ -175,5 +175,5 @@ xbuild Source/Programs/AutomationTool/HTML5/HTML5.Automation.csproj \
   /p:Configuration="Development"
 
 # pass all parameters to UBT
-mono Binaries/DotNET/UnrealBuildTool.exe -makefile "$@"
+mono Binaries/DotNET/UnrealBuildTool.exe -makefile -qmakefile -cmakefile "$@"
 set +x
diff --git a/Engine/Source/Programs/UnrealBuildTool/Configuration/UEBuildTarget.cs b/Engine/Source/Programs/UnrealBuildTool/Configuration/UEBuildTarget.cs
--- a/Engine/Source/Programs/UnrealBuildTool/Configuration/UEBuildTarget.cs
+++ b/Engine/Source/Programs/UnrealBuildTool/Configuration/UEBuildTarget.cs
@@ -353,6 +353,26 @@ namespace UnrealBuildTool
                         }
                         break;
 
+					    case "-CMAKEFILE":
+						    {
+							    // Force platform to Linux for building IntelliSense files
+							    Platform = UnrealTargetPlatform.Linux;
+
+							    // Force configuration to Development for IntelliSense
+							    Configuration = UnrealTargetConfiguration.Development;
+						    }
+						    break;
+
+					    case "-QMAKEFILE":
+						    {
+							    // Force platform to Linux for building IntelliSense files
+							    Platform = UnrealTargetPlatform.Linux;
+
+							    // Force configuration to Development for IntelliSense
+							    Configuration = UnrealTargetConfiguration.Development;
+						    }
+						    break;
+
                         case "-EDITORRECOMPILE":
 							{
 								bIsEditorRecompile = true;
@@ -533,6 +553,18 @@ namespace UnrealBuildTool
                             Platform = UnrealTargetPlatform.Linux;
                             Configuration = UnrealTargetConfiguration.Development;
                             break;
+
+						case "-CMAKEFILE":
+						    // Force platform to Linux and configuration to Development for building IntelliSense files
+						    Platform = UnrealTargetPlatform.Linux;
+						    Configuration = UnrealTargetConfiguration.Development;
+						    break;
+
+						case "-QMAKEFILE":
+						    // Force platform to Linux and configuration to Development for building IntelliSense files
+						    Platform = UnrealTargetPlatform.Linux;
+						    Configuration = UnrealTargetConfiguration.Development;
+						    break;
                     }
 				}
 			}
diff --git a/Engine/Source/Programs/UnrealBuildTool/System/CMakefileGenerator.cs b/Engine/Source/Programs/UnrealBuildTool/System/CMakefileGenerator.cs
new file mode 100644
--- /dev/null
+++ b/Engine/Source/Programs/UnrealBuildTool/System/CMakefileGenerator.cs
@@ -0,0 +1,265 @@
+﻿// Copyright 1998-2014 Epic Games, Inc. All Rights Reserved.
+
+using System;
+using System.Collections.Generic;
+using System.Text;
+using System.IO;
+
+namespace UnrealBuildTool
+{
+	/// <summary>
+	/// Represents a folder within the master project (e.g. Visual Studio solution)
+	/// </summary>
+	public class CMakefileFolder : MasterProjectFolder
+	{
+		/// <summary>
+		/// Constructor
+		/// </summary>
+		public CMakefileFolder( ProjectFileGenerator InitOwnerProjectFileGenerator, string InitFolderName )
+			: base(InitOwnerProjectFileGenerator, InitFolderName)
+		{
+		}
+	}
+
+	public class CMakefileProjectFile : ProjectFile
+	{
+		public CMakefileProjectFile( string InitFilePath )
+			: base(InitFilePath)
+		{
+		}
+	}
+	/// <summary>
+	/// CMakefile project file generator implementation
+	/// </summary>
+	public class CMakefileGenerator : ProjectFileGenerator
+	{
+		/// File extension for project files we'll be generating (e.g. ".vcxproj")
+		override public string ProjectFileExtension
+		{
+			get
+			{
+				return ".txt";
+			}
+		}
+
+		protected override bool WriteMasterProjectFile( ProjectFile UBTProject )
+		{
+			bool bSuccess = true;
+			return bSuccess;
+		}
+
+		private bool WriteCMakeLists()
+		{
+			string BuildCommand = "";
+			var FileName = "CMakeLists.txt";
+			var CMakefileContent = new StringBuilder();
+
+			var CMakeSectionEnd = " )\n\n";
+
+			var CMakeSourceFilesList = "set(SOURCE_FILES \n";
+			var CMakeHeaderFilesList = "set(HEADER_FILES \n";
+			var CMakeConfigFilesList = "set(CONFIG_FILES \n";
+
+			var CMakeGameRootPath = "";
+			var CMakeUE4RootPath = "set(UE4_ROOT_PATH " + Path.GetFullPath(ProjectFileGenerator.RootRelativePath) + ")\n";
+
+			string GameProjectPath = "";
+			string GameProjectFile = "";
+
+			string CMakeGameProjectFile = "";
+
+			if (!String.IsNullOrEmpty (GameProjectName)) 
+			{
+				CMakeGameRootPath = "set(GAME_ROOT_PATH \"" + UnrealBuildTool.GetUProjectPath() + "\")\n";
+
+				GameProjectPath = UnrealBuildTool.GetUProjectPath ();
+				GameProjectFile = UnrealBuildTool.GetUProjectFile ();
+
+				CMakeGameProjectFile = "set(GAME_PROJECT_FILE \"" + GameProjectFile + "\")\n";
+
+				BuildCommand = "set(BUILD mono ${UE4_ROOT_PATH}/Engine/Binaries/DotNET/UnrealBuildTool.exe )\n"; // -project=\"\\\"" + UnrealBuildTool.GetUProjectPath () + "/" + GameProjectName + ".uproject\\\"\")\n";
+			}
+			else
+			{
+				BuildCommand = "set(BUILD bash ${UE4_ROOT_PATH}/Engine/Build/BatchFiles/Linux/Build.sh)\n";
+			}
+
+			CMakefileContent.Append(
+				"# Makefile generated by CMakefileGenerator.cs\n" +
+				"# *DO NOT EDIT*\n\n" +
+				"cmake_minimum_required (VERSION 2.6)\n" +
+				"project (UE4)\n\n" +
+				CMakeUE4RootPath +
+				CMakeGameProjectFile +
+				BuildCommand +
+				CMakeGameRootPath + "\n"
+			);
+
+			// Create SourceFiles, HeaderFiles, and ConfigFiles sections.
+			var AllModuleFiles = DiscoverModules();
+			foreach (string CurModuleFile in AllModuleFiles)
+			{
+				var FoundFiles = SourceFileSearch.FindModuleSourceFiles(CurModuleFile, ExcludeNoRedistFiles: bExcludeNoRedistFiles);
+				foreach (string CurSourceFile in FoundFiles)
+				{
+
+					string SourceFileRelativeToRoot = Utils.MakePathRelativeTo(CurSourceFile, Path.Combine(EngineRelativePath));
+					// Exclude Windows, Mac, and iOS only files/folders.
+					// This got ugly quick.
+					if (!SourceFileRelativeToRoot.Contains ("Source/ThirdParty/") && 
+						!SourceFileRelativeToRoot.Contains ("/Windows/") && 
+						!SourceFileRelativeToRoot.Contains ("/Mac/") &&
+						!SourceFileRelativeToRoot.Contains ("/IOS/") &&
+						!SourceFileRelativeToRoot.Contains ("/iOS/") &&
+						!SourceFileRelativeToRoot.Contains ("/VisualStudioSourceCodeAccess/") &&
+						!SourceFileRelativeToRoot.Contains ("/XCodeSourceCodeAccess/") &&
+						!SourceFileRelativeToRoot.Contains ("/WmfMedia/") &&
+						!SourceFileRelativeToRoot.Contains ("/IOSDeviceProfileSelector/") &&
+						!SourceFileRelativeToRoot.Contains ("/WindowsDeviceProfileSelector/") &&
+						!SourceFileRelativeToRoot.Contains ("/WindowsMoviePlayer/") &&
+						!SourceFileRelativeToRoot.Contains ("/AppleMoviePlayer/") &&
+						!SourceFileRelativeToRoot.Contains ("/MacGraphicsSwitching/") &&
+						!SourceFileRelativeToRoot.Contains ("/Apple/") &&
+						!SourceFileRelativeToRoot.Contains ("/WinRT/")
+					) 
+					{
+						if (SourceFileRelativeToRoot.EndsWith (".cpp"))
+						{
+							if (!SourceFileRelativeToRoot.StartsWith ("..") && !Path.IsPathRooted (SourceFileRelativeToRoot))
+							{
+								// SourceFileRelativeToRoot = "Engine/" + SourceFileRelativeToRoot;
+								CMakeSourceFilesList += ("\t\"${UE4_ROOT_PATH}/Engine/" + SourceFileRelativeToRoot + "\"\n");
+							} 
+							else 
+							{
+								if (String.IsNullOrEmpty (GameProjectName)) 
+								{
+									// SourceFileRelativeToRoot = SourceFileRelativeToRoot.Substring (3);
+									CMakeSourceFilesList += ("\t\"" + SourceFileRelativeToRoot.Substring (3) + "\"\n");
+								}
+								else
+								{
+									CMakeSourceFilesList += ("\t\"${GAME_ROOT_PATH}/" + Utils.MakePathRelativeTo (CurSourceFile, GameProjectPath) + "\"\n");
+								}
+							}
+						}
+						if (SourceFileRelativeToRoot.EndsWith (".h")) 
+						{
+							if (!SourceFileRelativeToRoot.StartsWith ("..") && !Path.IsPathRooted (SourceFileRelativeToRoot)) 
+							{
+								// SourceFileRelativeToRoot = "Engine/" + SourceFileRelativeToRoot;
+								CMakeHeaderFilesList += ("\t\"${UE4_ROOT_PATH}/Engine/" +  SourceFileRelativeToRoot + "\"\n");
+							} 
+							else 
+							{
+								if (String.IsNullOrEmpty (GameProjectName)) 
+								{
+									// SourceFileRelativeToRoot = SourceFileRelativeToRoot.Substring (3);
+									CMakeHeaderFilesList += ("\t\"" + SourceFileRelativeToRoot.Substring (3) + "\"\n");
+								}
+								else
+								{
+									CMakeHeaderFilesList += ("\t\"${GAME_ROOT_PATH}/" + Utils.MakePathRelativeTo (CurSourceFile, GameProjectPath) + "\"\n");
+								}
+							}
+						}
+						if (SourceFileRelativeToRoot.EndsWith (".cs"))
+						{
+							if (!SourceFileRelativeToRoot.StartsWith ("..") && !Path.IsPathRooted (SourceFileRelativeToRoot)) 
+							{
+								// SourceFileRelativeToRoot = "Engine/" + SourceFileRelativeToRoot;
+								CMakeConfigFilesList += ("\t\"${UE4_ROOT_PATH}/Engine/" + SourceFileRelativeToRoot + "\"\n");
+
+							} 
+							else
+							{
+								if (String.IsNullOrEmpty (GameProjectName)) 
+								{
+									// SourceFileRelativeToRoot = SourceFileRelativeToRoot.Substring (3);
+									CMakeConfigFilesList += ("\t\"" + SourceFileRelativeToRoot.Substring (3) + "\"\n");
+								}
+								else
+								{
+									CMakeConfigFilesList += ("\t\"${GAME_ROOT_PATH}/" + Utils.MakePathRelativeTo (CurSourceFile, GameProjectPath) + "\"\n");
+								};
+							}
+						}
+					}
+				}
+
+			}
+
+			// Add section end to section strings;
+			CMakeSourceFilesList += CMakeSectionEnd;
+			CMakeHeaderFilesList += CMakeSectionEnd;
+			CMakeConfigFilesList += CMakeSectionEnd;
+
+			// Append sections to the CMakeLists.txt file
+			CMakefileContent.Append (CMakeSourceFilesList);
+			CMakefileContent.Append (CMakeHeaderFilesList);
+			CMakefileContent.Append (CMakeConfigFilesList);
+
+			string CMakeProjectCmdArg = "";
+
+			foreach (string TargetFilePath in DiscoverTargets())
+			{
+				var TargetName = Utils.GetFilenameWithoutAnyExtensions(TargetFilePath);		// Remove both ".cs" and ".
+
+				foreach (UnrealTargetConfiguration CurConfiguration in Enum.GetValues(typeof(UnrealTargetConfiguration)))
+				{
+					if (CurConfiguration != UnrealTargetConfiguration.Unknown && CurConfiguration != UnrealTargetConfiguration.Development)
+					{
+						if (UnrealBuildTool.IsValidConfiguration(CurConfiguration))
+						{
+							if (TargetName == GameProjectName || TargetName == (GameProjectName + "Editor")) 
+							{
+								CMakeProjectCmdArg = " -project=\"\\\"${GAME_PROJECT_FILE}\\\"\"";
+							}
+							var ConfName = Enum.GetName(typeof(UnrealTargetConfiguration), CurConfiguration);
+							CMakefileContent.Append(String.Format("add_custom_target({0}-Linux-{1} ${{BUILD}} {2} {0} Linux {1} $(ARGS))\n", TargetName, ConfName, CMakeProjectCmdArg));
+						}
+					}
+				}
+
+				if (TargetName == GameProjectName || TargetName == (GameProjectName + "Editor")) 
+				{
+					CMakeProjectCmdArg = " -project=\"\\\"${GAME_PROJECT_FILE}\\\"\"";
+				}
+				CMakefileContent.Append(String.Format("add_custom_target({0} ${{BUILD}} {1} {0} Linux Development $(ARGS) SOURCES ${{SOURCE_FILES}} ${{HEADER_FILES}} ${{CONFIG_FILES}})\n\n", TargetName, CMakeProjectCmdArg));
+			}
+
+			var FullFileName = Path.Combine(MasterProjectRelativePath, FileName);
+			return WriteFileIfChanged(FullFileName, CMakefileContent.ToString());
+		}
+
+		
+		/// ProjectFileGenerator interface
+		//protected override bool WriteMasterProjectFile( ProjectFile UBTProject )
+		protected override bool WriteProjectFiles ()
+		{
+			return WriteCMakeLists();
+		}
+
+		/// ProjectFileGenerator interface
+		public override MasterProjectFolder AllocateMasterProjectFolder( ProjectFileGenerator InitOwnerProjectFileGenerator, string InitFolderName )
+		{
+			return new CMakefileFolder( InitOwnerProjectFileGenerator, InitFolderName );
+		}
+
+		/// ProjectFileGenerator interface
+		/// <summary>
+		/// Allocates a generator-specific project file object
+		/// </summary>
+		/// <param name="InitFilePath">Path to the project file</param>
+		/// <returns>The newly allocated project file object</returns>
+		protected override ProjectFile AllocateProjectFile( string InitFilePath )
+		{
+			return new CMakefileProjectFile( InitFilePath );
+		}
+
+		/// ProjectFileGenerator interface
+		public override void CleanProjectFiles(string InMasterProjectRelativePath, string InMasterProjectName, string InIntermediateProjectFilesPath)
+		{
+		}
+	}
+}
diff --git a/Engine/Source/Programs/UnrealBuildTool/System/MakefileGenerator.cs b/Engine/Source/Programs/UnrealBuildTool/System/MakefileGenerator.cs
--- a/Engine/Source/Programs/UnrealBuildTool/System/MakefileGenerator.cs
+++ b/Engine/Source/Programs/UnrealBuildTool/System/MakefileGenerator.cs
@@ -101,69 +101,11 @@ namespace UnrealBuildTool
             return WriteFileIfChanged(FullFileName, MakefileContent.ToString());
         }
 
-        private bool WriteCMakeLists()
-        {
-            var FileName = "CMakeLists.txt";
-            var CMakefileContent = new StringBuilder();
-
-            CMakefileContent.Append(
-                "# Makefile generated by MakefileGenerator.cs\n" +
-                "# *DO NOT EDIT*\n\n" +
-                "cmake_minimum_required (VERSION 2.6)\n" +
-                "project (UE4)\n\n" +
-                "set(BUILD ${CMAKE_SOURCE_DIR}/Engine/Build/BatchFiles/Linux/Build.sh)\n\n" +
-                "set(SOURCE_FILES"
-                );
-
-            var AllModuleFiles = DiscoverModules();
-            foreach (string CurModuleFile in AllModuleFiles)
-            {
-                CMakefileContent.Append("\n\t");
-
-                var FoundFiles = SourceFileSearch.FindModuleSourceFiles(CurModuleFile, ExcludeNoRedistFiles: bExcludeNoRedistFiles);
-                foreach (string CurSourceFile in FoundFiles)
-                {
-                    string SourceFileRelativeToRoot = Utils.MakePathRelativeTo(CurSourceFile, Path.Combine(EngineRelativePath));
-                    if (!SourceFileRelativeToRoot.StartsWith("..") && !Path.IsPathRooted(SourceFileRelativeToRoot))
-                    {
-                        SourceFileRelativeToRoot = "Engine/" + SourceFileRelativeToRoot;
-                    }
-                    else
-                    {
-                        SourceFileRelativeToRoot = SourceFileRelativeToRoot.Substring(3);
-                    }
-
-                    CMakefileContent.Append(String.Format("\"{0}\" ", SourceFileRelativeToRoot));
-                }
-            }
-            CMakefileContent.Append("\n)\n\n");
-            foreach (string TargetFilePath in DiscoverTargets())
-            {
-                var TargetName = Utils.GetFilenameWithoutAnyExtensions(TargetFilePath);		// Remove both ".cs" and ".
-
-                foreach (UnrealTargetConfiguration CurConfiguration in Enum.GetValues(typeof(UnrealTargetConfiguration)))
-                {
-                    if (CurConfiguration != UnrealTargetConfiguration.Unknown && CurConfiguration != UnrealTargetConfiguration.Development)
-                    {
-                        if (UnrealBuildTool.IsValidConfiguration(CurConfiguration))
-                        {
-                            var ConfName = Enum.GetName(typeof(UnrealTargetConfiguration), CurConfiguration);
-                            CMakefileContent.Append(String.Format("add_custom_target({0}-Linux-{1} ${{BUILD}} {0} Linux {1})\n", TargetName, ConfName));
-                        }
-                    }
-                }
-
-                CMakefileContent.Append(String.Format("add_custom_target({0} ${{BUILD}} {0} Linux Development SOURCES ${{SOURCE_FILES}})\n\n", TargetName));
-            }
-            var FullFileName = Path.Combine(MasterProjectRelativePath, FileName);
-            return WriteFileIfChanged(FullFileName, CMakefileContent.ToString());
-        }
-
         /// ProjectFileGenerator interface
         //protected override bool WriteMasterProjectFile( ProjectFile UBTProject )
         protected override bool WriteProjectFiles ()
         {
-            return WriteMakefile() && WriteCMakeLists();
+			return WriteMakefile();
         }
 
         /// ProjectFileGenerator interface
diff --git a/Engine/Source/Programs/UnrealBuildTool/System/QMakefileGenerator.cs b/Engine/Source/Programs/UnrealBuildTool/System/QMakefileGenerator.cs
new file mode 100644
--- /dev/null
+++ b/Engine/Source/Programs/UnrealBuildTool/System/QMakefileGenerator.cs
@@ -0,0 +1,289 @@
+﻿// Copyright 1998-2014 Epic Games, Inc. All Rights Reserved.
+
+using System;
+using System.Collections.Generic;
+using System.Text;
+using System.IO;
+
+namespace UnrealBuildTool
+{
+	/// <summary>
+	/// Represents a folder within the master project (e.g. Visual Studio solution)
+	/// </summary>
+	public class QMakefileFolder : MasterProjectFolder
+	{
+		/// <summary>
+		/// Constructor
+		/// </summary>
+		public QMakefileFolder( ProjectFileGenerator InitOwnerProjectFileGenerator, string InitFolderName )
+			: base(InitOwnerProjectFileGenerator, InitFolderName)
+		{
+		}
+	}
+
+	public class QMakefileProjectFile : ProjectFile
+	{
+		public QMakefileProjectFile( string InitFilePath )
+			: base(InitFilePath)
+		{
+		}
+	}
+
+	/// <summary>
+	/// QMakefile project file generator implementation
+	/// </summary>
+	public class QMakefileGenerator : ProjectFileGenerator
+	{
+		/// File extension for project files we'll be generating (e.g. ".vcxproj")
+		override public string ProjectFileExtension
+		{
+			get
+			{
+				return ".pro";
+			}
+		}
+
+		protected override bool WriteMasterProjectFile( ProjectFile UBTProject )
+		{
+			bool bSuccess = true;
+			return bSuccess;
+		}
+
+		private bool WriteQMakePro()
+		{
+			string GameProjectPath = "";
+			string GameProjectFile = "";
+			string GameProjectRootPath = "";
+
+			string BuildCommand = "";
+
+			string QMakeGameProjectFile = "";
+			
+			if (!String.IsNullOrEmpty (GameProjectName))
+			{
+				GameProjectPath = UnrealBuildTool.GetUProjectPath ();
+				GameProjectFile = UnrealBuildTool.GetUProjectFile ();
+				QMakeGameProjectFile = "gameProjectFile=" + GameProjectFile + "\n";
+				BuildCommand = "build=mono $$unrealRootPath/Engine/Binaries/DotNET/UnrealBuildTool.exe\n\n";
+			}
+			else
+			{
+				BuildCommand = "build=bash $$unrealRootPath/Engine/Build/BatchFiles/Linux/Build.sh\n";
+			}
+
+			var UnrealRootPath = Path.GetFullPath(ProjectFileGenerator.RootRelativePath);
+
+			var FileName = MasterProjectName + ".pro";
+			var QMakeFileContent = new StringBuilder();
+
+			var QMakeSectionEnd = " \n\n";
+
+			var QMakeSourceFilesList = "SOURCES += \\ \n";
+			var QMakeHeaderFilesList = "HEADERS += \\ \n";
+			var QMakeConfigFilesList = "OTHER_FILES += \\ \n";
+			var QMakeTargetList = "QMAKE_EXTRA_TARGETS += \\ \n";
+
+			if (!String.IsNullOrEmpty (GameProjectName)) {
+				GameProjectRootPath = GameProjectName + "RootPath=" + GameProjectPath + "\n\n";
+			}
+
+			QMakeFileContent.Append(
+				"# UnrealEngine.pro generated by QMakefileGenerator.cs\n" +
+				"# *DO NOT EDIT*\n\n" +
+				"TEMPLATE = aux\n" +
+				"CONFIG -= console\n" +
+				"CONFIG -= app_bundle\n" +
+				"CONFIG -= qt\n\n" +
+				"TARGET = UE4 \n\n" +
+				"unrealRootPath=" + UnrealRootPath + "\n" +
+				GameProjectRootPath +
+				QMakeGameProjectFile +
+				BuildCommand +
+				"args=$(ARGS)\n\n"
+			);
+
+			// Create SourceFiles, HeaderFiles, and ConfigFiles sections.
+			var AllModuleFiles = DiscoverModules();
+			foreach (string CurModuleFile in AllModuleFiles)
+			{
+				var FoundFiles = SourceFileSearch.FindModuleSourceFiles(CurModuleFile, ExcludeNoRedistFiles: bExcludeNoRedistFiles);
+				foreach (string CurSourceFile in FoundFiles)
+				{
+
+					string SourceFileRelativeToRoot = Utils.MakePathRelativeTo(CurSourceFile, Path.Combine(EngineRelativePath));
+					// Exclude Windows, Mac, and iOS only files/folders.
+					// This got ugly quick.
+					if (!SourceFileRelativeToRoot.Contains ("Source/ThirdParty/") && 
+						!SourceFileRelativeToRoot.Contains ("/Windows/") && 
+						!SourceFileRelativeToRoot.Contains ("/Mac/") &&
+						!SourceFileRelativeToRoot.Contains ("/IOS/") &&
+						!SourceFileRelativeToRoot.Contains ("/iOS/") &&
+						!SourceFileRelativeToRoot.Contains ("/VisualStudioSourceCodeAccess/") &&
+						!SourceFileRelativeToRoot.Contains ("/XCodeSourceCodeAccess/") &&
+						!SourceFileRelativeToRoot.Contains ("/WmfMedia/") &&
+						!SourceFileRelativeToRoot.Contains ("/IOSDeviceProfileSelector/") &&
+						!SourceFileRelativeToRoot.Contains ("/WindowsDeviceProfileSelector/") &&
+						!SourceFileRelativeToRoot.Contains ("/WindowsMoviePlayer/") &&
+						!SourceFileRelativeToRoot.Contains ("/AppleMoviePlayer/") &&
+						!SourceFileRelativeToRoot.Contains ("/MacGraphicsSwitching/") &&
+						!SourceFileRelativeToRoot.Contains ("/Apple/") &&
+						!SourceFileRelativeToRoot.Contains ("/WinRT/")
+					) 
+					{
+						if (SourceFileRelativeToRoot.EndsWith (".cpp"))
+						{
+							if (!SourceFileRelativeToRoot.StartsWith ("..") && !Path.IsPathRooted (SourceFileRelativeToRoot))
+							{
+								// SourceFileRelativeToRoot = "Engine/" + SourceFileRelativeToRoot;
+								QMakeSourceFilesList += ("\t\"" + "$$unrealRootPath/Engine/" + SourceFileRelativeToRoot + "\" \\\n");
+							} 
+							else 
+							{
+								if (String.IsNullOrEmpty (GameProjectName)) 
+								{
+									// SourceFileRelativeToRoot = SourceFileRelativeToRoot.Substring (3);
+									QMakeSourceFilesList += ("\t\"" + SourceFileRelativeToRoot.Substring (3) + "\" \\\n");
+								} 
+								else if (!String.IsNullOrEmpty (GameProjectName)) 
+								{
+									QMakeSourceFilesList += ("\t\"$$" + GameProjectName + "RootPath/" + Utils.MakePathRelativeTo (CurSourceFile, GameProjectPath) + "\" \\\n");
+								} 
+								else 
+								{
+									System.Console.WriteLine ("Error!, you should not be here.");
+								}
+							}
+						}
+						if (SourceFileRelativeToRoot.EndsWith (".h")) 
+						{
+							if (!SourceFileRelativeToRoot.StartsWith ("..") && !Path.IsPathRooted (SourceFileRelativeToRoot)) 
+							{
+								// SourceFileRelativeToRoot = "Engine/" + SourceFileRelativeToRoot;
+								QMakeHeaderFilesList += ("\t\"" + "$$unrealRootPath/Engine/" + SourceFileRelativeToRoot + "\" \\\n");
+							} 
+							else 
+							{
+								if (String.IsNullOrEmpty (GameProjectName)) 
+								{
+									// SourceFileRelativeToRoot = SourceFileRelativeToRoot.Substring (3);
+									QMakeHeaderFilesList += ("\t\"" + SourceFileRelativeToRoot.Substring (3) + "\" \\\n");
+								} 
+								else if (!String.IsNullOrEmpty (GameProjectName)) 
+								{
+									QMakeHeaderFilesList += ("\t\"$$" + GameProjectName + "RootPath/" + Utils.MakePathRelativeTo (CurSourceFile, GameProjectPath) + "\" \\\n");
+								} 
+								else 
+								{
+									System.Console.WriteLine ("Error!, you should not be here.");
+								}
+							}
+						}
+						if (SourceFileRelativeToRoot.EndsWith (".cs"))
+						{
+							if (!SourceFileRelativeToRoot.StartsWith ("..") && !Path.IsPathRooted (SourceFileRelativeToRoot)) 
+							{
+								// SourceFileRelativeToRoot = "Engine/" + SourceFileRelativeToRoot;
+								QMakeConfigFilesList += ("\t\"" + "$$unrealRootPath/Engine/" + SourceFileRelativeToRoot + "\" \\\n");
+
+							} 
+							else
+							{
+								if (String.IsNullOrEmpty (GameProjectName)) 
+								{
+									// SourceFileRelativeToRoot = SourceFileRelativeToRoot.Substring (3);
+									QMakeConfigFilesList += ("\t\"" + SourceFileRelativeToRoot.Substring (3) + "\" \\\n");
+								} 
+								else if (!String.IsNullOrEmpty (GameProjectName)) 
+								{
+									QMakeConfigFilesList += ("\t\"$$" + GameProjectName + "RootPath/" + Utils.MakePathRelativeTo (CurSourceFile, GameProjectPath) + "\" \\\n");
+								} 
+								else 
+								{
+									System.Console.WriteLine ("Error!, you should not be here.");
+								}
+							}
+						}
+					}
+				}
+
+			}
+
+			// Add section end to section strings;
+			QMakeSourceFilesList += QMakeSectionEnd;
+			QMakeHeaderFilesList += QMakeSectionEnd;
+			QMakeConfigFilesList += QMakeSectionEnd;
+
+			// Append sections to the QMakeLists.txt file
+			QMakeFileContent.Append (QMakeSourceFilesList);
+			QMakeFileContent.Append (QMakeHeaderFilesList);
+			QMakeFileContent.Append (QMakeConfigFilesList);
+
+			string QMakeProjectCmdArg = "";
+
+			foreach (string TargetFilePath in DiscoverTargets())
+			{
+				var TargetName = Utils.GetFilenameWithoutAnyExtensions(TargetFilePath);		// Remove both ".cs" and ".
+
+				foreach (UnrealTargetConfiguration CurConfiguration in Enum.GetValues(typeof(UnrealTargetConfiguration)))
+				{
+					if (CurConfiguration != UnrealTargetConfiguration.Unknown && CurConfiguration != UnrealTargetConfiguration.Development)
+					{
+						if (UnrealBuildTool.IsValidConfiguration(CurConfiguration))
+						{
+
+							if (TargetName == GameProjectName || TargetName == (GameProjectName + "Editor")) 
+							{
+								QMakeProjectCmdArg = " -project=\"\\\"$$gameProjectFile\\\"\"";
+							}
+							var ConfName = Enum.GetName(typeof(UnrealTargetConfiguration), CurConfiguration);
+							QMakeFileContent.Append(String.Format("{0}-Linux-{1}.commands = $$build {2} {0} Linux {1} $$args\n", TargetName, ConfName, QMakeProjectCmdArg));
+							QMakeTargetList += "\t" + TargetName + "-Linux-" + ConfName + " \\\n"; // , TargetName, ConfName);
+						}
+					}
+				}
+
+				if (TargetName == GameProjectName || TargetName == (GameProjectName + "Editor")) 
+				{
+					QMakeProjectCmdArg = " -project=\"\\\"$$gameProjectFile\\\"\"";
+				}
+
+				QMakeFileContent.Append(String.Format("{0}.commands = $$build {1} {0} Linux Development $$args\n\n", TargetName, QMakeProjectCmdArg));
+				QMakeTargetList += "\t" + TargetName + " \\\n";
+			}
+
+			QMakeFileContent.Append (QMakeTargetList.TrimEnd('\\'));
+
+			var FullFileName = Path.Combine(MasterProjectRelativePath, FileName);
+			return WriteFileIfChanged(FullFileName, QMakeFileContent.ToString());
+		}
+
+		/// ProjectFileGenerator interface
+		//protected override bool WriteMasterProjectFile( ProjectFile UBTProject )
+		protected override bool WriteProjectFiles ()
+		{
+			return WriteQMakePro();
+		}
+
+		/// ProjectFileGenerator interface
+		public override MasterProjectFolder AllocateMasterProjectFolder( ProjectFileGenerator InitOwnerProjectFileGenerator, string InitFolderName )
+		{
+			return new QMakefileFolder( InitOwnerProjectFileGenerator, InitFolderName );
+		}
+
+		/// ProjectFileGenerator interface
+		/// <summary>
+		/// Allocates a generator-specific project file object
+		/// </summary>
+		/// <param name="InitFilePath">Path to the project file</param>
+		/// <returns>The newly allocated project file object</returns>
+		protected override ProjectFile AllocateProjectFile( string InitFilePath )
+		{
+			return new QMakefileProjectFile( InitFilePath );
+		}
+
+		/// ProjectFileGenerator interface
+		public override void CleanProjectFiles(string InMasterProjectRelativePath, string InMasterProjectName, string InIntermediateProjectFilesPath)
+		{
+		}
+	}
+}
diff --git a/Engine/Source/Programs/UnrealBuildTool/System/UnrealBuildTool.cs b/Engine/Source/Programs/UnrealBuildTool/System/UnrealBuildTool.cs
--- a/Engine/Source/Programs/UnrealBuildTool/System/UnrealBuildTool.cs
+++ b/Engine/Source/Programs/UnrealBuildTool/System/UnrealBuildTool.cs
@@ -541,8 +541,16 @@ namespace UnrealBuildTool
             }
             else if (LowercaseArg.StartsWith("-project="))
             {
-                UProjectFile = InArg.Substring(9);
-                bSetupProject = true;
+				if (BuildHostPlatform.Current.Platform == UnrealTargetPlatform.Linux) 
+				{
+					UProjectFile = InArg.Substring(9).Trim(new Char[] { ' ', '"', '\'' });
+				}
+				else 
+				{
+					UProjectFile = InArg.Substring(9);
+				}
+
+				bSetupProject = true;
             }
             else if (LowercaseArg.EndsWith(".uproject"))
             {
@@ -759,6 +767,8 @@ namespace UnrealBuildTool
                     var bGenerateVCProjectFiles = false;
                     var bGenerateXcodeProjectFiles = false;
                     var bGenerateMakefiles = false;
+					var bGenerateCMakefiles = false;
+					var bGenerateQMakefiles = false;
                     var bValidPlatformsOnly = false;
                     var bSpecificModulesOnly = false;
 
@@ -818,6 +828,16 @@ namespace UnrealBuildTool
                             bGenerateMakefiles = true;
                             ProjectFileGenerator.bGenerateProjectFiles = true;
                         }
+						else if (LowercaseArg.StartsWith("-cmakefile"))
+						{
+							bGenerateCMakefiles = true;
+							ProjectFileGenerator.bGenerateProjectFiles = true;
+						}
+						else if (LowercaseArg.StartsWith("-qmakefile"))
+						{
+							bGenerateQMakefiles = true;
+							ProjectFileGenerator.bGenerateProjectFiles = true;
+						}
                         else if (LowercaseArg.StartsWith("-projectfile"))
                         {
                             bGenerateVCProjectFiles = true;
@@ -984,6 +1004,7 @@ namespace UnrealBuildTool
                     }
 
                     // Send an event with basic usage dimensions
+                    // [RCL] 2014-11-03 FIXME: this is incorrect since we can have more than one action
                     Telemetry.SendEvent("CommonAttributes.2",
                         // @todo No platform independent way to do processor speed and count. Defer for now. Mono actually has ways to do this but needs to be tested.
                         "ProcessorCount", Environment.ProcessorCount.ToString(),
@@ -997,11 +1018,15 @@ namespace UnrealBuildTool
                                 ? "GenerateXcodeProjectFiles" 
                                 : bRunCopyrightVerification
                                     ? "RunCopyrightVerification"
-                                        : bGenerateMakefiles
+                                    : bGenerateMakefiles
                                         ? "GenerateMakefiles"
-                                            : bValidatePlatforms 
-                                                ? "ValidatePlatfomrs"
-                                                : "Build",
+										: bGenerateCMakefiles
+										    ? "GenerateCMakeFiles"
+											: bGenerateQMakefiles
+											    ? "GenerateQMakefiles"
+												: bValidatePlatforms 
+                                                    ? "ValidatePlatfomrs"
+                                                	: "Build",
                         "Platform", CheckPlatform.ToString(),
                         "Configuration", CheckConfiguration.ToString(),
                         "IsRocket", bRunningRocket.ToString(),
@@ -1034,7 +1059,7 @@ namespace UnrealBuildTool
                         JunkDeleter.DeleteJunk();
                     }
 
-                    if (bGenerateVCProjectFiles || bGenerateXcodeProjectFiles || bGenerateMakefiles)
+					if (bGenerateVCProjectFiles || bGenerateXcodeProjectFiles || bGenerateMakefiles || bGenerateCMakefiles || bGenerateQMakefiles)
                     {
                         bool bGenerationSuccess = true;
                         if (bGenerateVCProjectFiles)
@@ -1049,6 +1074,14 @@ namespace UnrealBuildTool
                         {
                             bGenerationSuccess &= GenerateProjectFiles(new MakefileGenerator(), Arguments);
                         }
+						if (bGenerateCMakefiles)
+						{
+							bGenerationSuccess &= GenerateProjectFiles(new CMakefileGenerator(), Arguments);
+						}
+						if (bGenerateQMakefiles)
+						{
+							bGenerationSuccess &= GenerateProjectFiles(new QMakefileGenerator(), Arguments);
+						}
 
                         if(!bGenerationSuccess)
                         {
diff --git a/Engine/Source/Programs/UnrealBuildTool/UnrealBuildTool_Mono.csproj b/Engine/Source/Programs/UnrealBuildTool/UnrealBuildTool_Mono.csproj
--- a/Engine/Source/Programs/UnrealBuildTool/UnrealBuildTool_Mono.csproj
+++ b/Engine/Source/Programs/UnrealBuildTool/UnrealBuildTool_Mono.csproj
@@ -179,6 +179,8 @@
     <Compile Include="System\ModuleProcessingException.cs" />
     <Compile Include="System\VCSolutionOptions.cs" />
     <Compile Include="System\Distcc.cs" />
+    <Compile Include="System\CMakefileGenerator.cs" />
+    <Compile Include="System\QMakefileGenerator.cs" />
   </ItemGroup>
   <ItemGroup>
     <BootstrapperPackage Include="Microsoft.Net.Client.3.5">
