Fix CIS (UBT non-mono project).

diff --git a/Engine/Source/Programs/UnrealBuildTool/UnrealBuildTool.csproj b/Engine/Source/Programs/UnrealBuildTool/UnrealBuildTool.csproj
--- a/Engine/Source/Programs/UnrealBuildTool/UnrealBuildTool.csproj
+++ b/Engine/Source/Programs/UnrealBuildTool/UnrealBuildTool.csproj
@@ -224,6 +224,8 @@
     <Compile Include="System\BuildHostPlatform.cs" />
     <Compile Include="System\FlatCPPIncludeDepencencyCache.cs" />
     <Compile Include="System\Distcc.cs" />
+    <Compile Include="System\CMakefileGenerator.cs" />
+    <Compile Include="System\QMakefileGenerator.cs" />
     <Compile Include="System\ModuleProcessingException.cs" />
     <Compile Include="System\MakefileGenerator.cs" />
     <Compile Include="System\SNDBS.cs" />
@@ -315,4 +317,4 @@
     <PreBuildEvent>
     </PreBuildEvent>
   </PropertyGroup>
-</Project>
\ No newline at end of file
+</Project>
