Making BuildConfiguration.bForceUnityBuild take precedence over bFasterWithoutUnity

diff --git a/Engine/Source/Programs/UnrealBuildTool/Configuration/UEBuildModule.cs b/Engine/Source/Programs/UnrealBuildTool/Configuration/UEBuildModule.cs
--- a/Engine/Source/Programs/UnrealBuildTool/Configuration/UEBuildModule.cs
+++ b/Engine/Source/Programs/UnrealBuildTool/Configuration/UEBuildModule.cs
@@ -1226,16 +1226,19 @@ namespace UnrealBuildTool
 
 
 			// Should we use unity build mode for this module?
-			bool bModuleUsesUnityBuild = BuildConfiguration.bUseUnityBuild;
-			if( bFasterWithoutUnity )
+			bool bModuleUsesUnityBuild = BuildConfiguration.bUseUnityBuild || BuildConfiguration.bForceUnityBuild;
+			if (!BuildConfiguration.bForceUnityBuild)
 			{
-				bModuleUsesUnityBuild = false;
-			}
-			else if( !BuildConfiguration.bForceUnityBuild && IsGameModule && SourceFilesToBuild.CPPFiles.Count < BuildConfiguration.MinGameModuleSourceFilesForUnityBuild )
-			{
-				// Game modules with only a small number of source files are usually better off having faster iteration times
-				// on single source file changes, so we forcibly disable unity build for those modules
-				bModuleUsesUnityBuild = false;
+				if (bFasterWithoutUnity)
+				{
+					bModuleUsesUnityBuild = false;
+				}
+				else if (IsGameModule && SourceFilesToBuild.CPPFiles.Count < BuildConfiguration.MinGameModuleSourceFilesForUnityBuild)
+				{
+					// Game modules with only a small number of source files are usually better off having faster iteration times
+					// on single source file changes, so we forcibly disable unity build for those modules
+					bModuleUsesUnityBuild = false;
+				}
 			}
 
 			// The environment with which to compile the CPP files
