- Fixed an issue with missing a library for distribution builds [UE-4837]
- Improved code signing/zipalign support in SDK 21
- Forcing NDK 19, since NDK 21 will generate code that cannot run on previous OSs

diff --git a/Engine/Source/Programs/UnrealBuildTool/Android/AndroidToolChain.cs b/Engine/Source/Programs/UnrealBuildTool/Android/AndroidToolChain.cs
--- a/Engine/Source/Programs/UnrealBuildTool/Android/AndroidToolChain.cs
+++ b/Engine/Source/Programs/UnrealBuildTool/Android/AndroidToolChain.cs
@@ -579,7 +579,7 @@ namespace UnrealBuildTool
 
 			if (!bHasPrintedApiLevel)
 			{
-				Console.WriteLine("Compiling with NDK API '{0}'", GetNdkApiLevel());
+				Console.WriteLine("Compiling Native code with NDK API '{0}'", GetNdkApiLevel());
 				bHasPrintedApiLevel = true;
 			}
 
diff --git a/Engine/Source/Programs/UnrealBuildTool/Android/UEBuildAndroid.cs b/Engine/Source/Programs/UnrealBuildTool/Android/UEBuildAndroid.cs
--- a/Engine/Source/Programs/UnrealBuildTool/Android/UEBuildAndroid.cs
+++ b/Engine/Source/Programs/UnrealBuildTool/Android/UEBuildAndroid.cs
@@ -15,7 +15,7 @@ namespace UnrealBuildTool
 		/// Android settings.
 		/// </summary>
 		[XmlConfig]
-		public static string AndroidNdkApiTarget = "latest";
+		public static string AndroidNdkApiTarget = "android-19";
 		[XmlConfig]
 		public static string AndroidSdkApiTarget = "latest";
 
diff --git a/Engine/Source/Programs/UnrealBuildTool/Android/UEDeployAndroid.cs b/Engine/Source/Programs/UnrealBuildTool/Android/UEDeployAndroid.cs
--- a/Engine/Source/Programs/UnrealBuildTool/Android/UEDeployAndroid.cs
+++ b/Engine/Source/Programs/UnrealBuildTool/Android/UEDeployAndroid.cs
@@ -95,7 +95,7 @@ namespace UnrealBuildTool.Android
 					}
 				}
 
-				Console.WriteLine("Building with SDK API '{0}'", Target);
+				Console.WriteLine("Building Java with SDK API '{0}'", Target);
 				CachedSDKLevel = Target;
 			}
 
@@ -714,7 +714,11 @@ namespace UnrealBuildTool.Android
 					{
 						// always delete libs up to this point so fat binaries and incremental builds work together (otherwise we might end up with multiple
 						// so files in an apk that doesn't want them)
-						DeleteDirectory(UE4BuildPath + "/libs");
+						// note that we don't want to delete all libs, just the ones we copied (
+						foreach (string Lib in Directory.EnumerateFiles(UE4BuildPath + "/libs", "libUE4*.so", SearchOption.AllDirectories))
+						{
+							File.Delete(Lib);
+						}
 
 						// if we need to run ndk-build, do it now (if making a shared .apk, we need to wait until all .libs exist)
 						if (!string.IsNullOrEmpty(FinalNdkBuildABICommand))
@@ -753,34 +757,32 @@ namespace UnrealBuildTool.Android
 
 						Log.TraceInformation("\n===={0}====PERFORMING FINAL APK PACKAGE OPERATION================================================", DateTime.Now.ToString());
 
+						string AntBuildType = "debug";
+						string AntOutputSuffix = "-debug";
+						if (bForDistribution)
+						{
+							// this will write out ant.properties with info needed to sign a distribution build
+							PrepareToSignApk(UE4BuildPath);
+							AntBuildType = "release";
+							AntOutputSuffix = "-release";
+						}
+
 						// Use ant to build the .apk file
-						RunCommandLineProgramAndThrowOnError(UE4BuildPath, "cmd.exe", "/c \"" + GetAntPath() + "\" -quiet " + (bForDistribution ? "release" : "debug"), "Making .apk with Ant... (note: it's safe to ignore javac obsolete warnings)");
+						RunCommandLineProgramAndThrowOnError(UE4BuildPath, "cmd.exe", "/c \"" + GetAntPath() + "\" -quiet " + AntBuildType, "Making .apk with Ant... (note: it's safe to ignore javac obsolete warnings)");
 
 						// make sure destination exists
 						Directory.CreateDirectory(Path.GetDirectoryName(DestApkName));
 
-						// do we need to sign for distro?
-						if (bForDistribution)
-						{
-							// use diffeent source and dest apk's for signed mode
-							string SourceApkName = UE4BuildPath + "/bin/" + ProjectName + "-release-unsigned.apk";
-							SignApk(UE4BuildPath + "/SigningConfig.xml", SourceApkName, DestApkName);
-						}
-						else
-						{
-							// now copy to the final location
-							File.Copy(UE4BuildPath + "/bin/" + ProjectName + "-debug" + ".apk", DestApkName, true);
-						}
+						// now copy to the final location
+						File.Copy(UE4BuildPath + "/bin/" + ProjectName + AntOutputSuffix + ".apk", DestApkName, true);
 					}
 				}
 			}
 		}
 
-		private void SignApk(string ConfigFilePath, string SourceApk, string DestApk)
+		private void PrepareToSignApk(string BuildPath)
 		{
-			string JarCommandPath = Environment.ExpandEnvironmentVariables("%JAVA_HOME%/bin/jarsigner.exe");
-			string ZipalignCommandPath = Environment.ExpandEnvironmentVariables("%ANDROID_HOME%/tools/zipalign.exe");
-
+			string ConfigFilePath = Path.Combine(BuildPath, "SigningConfig.xml");
 			if (!File.Exists(ConfigFilePath))
 			{
 				throw new BuildException("Unable to sign for Shipping without signing config file: '{0}", ConfigFilePath);
@@ -828,31 +830,15 @@ namespace UnrealBuildTool.Android
 				}
 			}
 
-			string CommandLine = "-sigalg SHA1withRSA -digestalg SHA1";
-			CommandLine += " -storepass " + KeystorePassword;
-
-			if (KeyPassword != "_sameaskeystore_")
-			{
-				CommandLine += " -keypass " + KeyPassword;
-			}
-
-			// put on the keystore
-			CommandLine += " -keystore \"" + Keystore + "\"";
-
-			// finish off the commandline
-			CommandLine += " \"" + SourceApk + "\" " + Alias;
-
-			// sign in-place
-			RunCommandLineProgramAndThrowOnError(Path.GetDirectoryName(ConfigFilePath), JarCommandPath, CommandLine, "Signing for distribution...");
 
-			if (File.Exists(DestApk))
-			{
-				File.Delete(DestApk);
-			}
+			string[] AntPropertiesLines = new string[4];
+			AntPropertiesLines[0] = "key.store=" + Keystore;
+			AntPropertiesLines[1] = "key.alias=" + Alias;
+			AntPropertiesLines[2] = "key.store.password=" + KeystorePassword;
+			AntPropertiesLines[3] = "key.alias.password=" + ((KeyPassword == "_sameaskeystore_") ? KeystorePassword : KeyPassword);
 
-			// now we need to zipalign the apk to the final destination (to 4 bytes, must be 4)
-			CommandLine = string.Format("4 \"{0}\" \"{1}\"", SourceApk, DestApk);
-			RunCommandLineProgramAndThrowOnError(Path.GetDirectoryName(ConfigFilePath), ZipalignCommandPath, CommandLine, "Zipaligning final .apk...");
+			// now write out the properties
+			File.WriteAllLines(Path.Combine(BuildPath, "ant.properties"), AntPropertiesLines);
 		}
 
 		public override bool PrepTargetForDeployment(UEBuildTarget InTarget)
