Fix android_native_app_glue and cpu-features not getting built & included in the build.

diff --git a/Engine/Source/Programs/UnrealBuildTool/Android/AndroidToolChain.cs b/Engine/Source/Programs/UnrealBuildTool/Android/AndroidToolChain.cs
--- a/Engine/Source/Programs/UnrealBuildTool/Android/AndroidToolChain.cs
+++ b/Engine/Source/Programs/UnrealBuildTool/Android/AndroidToolChain.cs
@@ -11,7 +11,7 @@ namespace UnrealBuildTool
 {
 	public class AndroidToolChain : UEToolChain
 	{
-		static private bool bHasNDKExtensionsCompiled = false;
+		static private string ModuleHasNDKExtensionsCompiled = null;
 
 		// the number of the clang version being used to compile 
 		static private float ClangVersionFloat = 0;
@@ -552,14 +552,18 @@ namespace UnrealBuildTool
 			return false;
 		}
 
-		static void ConditionallyAddNDKSourceFiles(List<FileItem> SourceFiles)
+		static void ConditionallyAddNDKSourceFiles(List<FileItem> SourceFiles, string ModuleName)
 		{
-			if (!bHasNDKExtensionsCompiled)
+			// We need to add the extra glue and cpu code only once to the first module. But since
+			// this is called multiple times for the same module we remember the module we
+			// add the extra code to and add it to that always. I.e. it makes this call
+			// consistent within one UBT invocation.
+			if (null == ModuleHasNDKExtensionsCompiled || ModuleHasNDKExtensionsCompiled.Equals(ModuleName))
 			{
 				SourceFiles.Add(FileItem.GetItemByPath(Environment.GetEnvironmentVariable("NDKROOT") + "/sources/android/native_app_glue/android_native_app_glue.c"));
 				SourceFiles.Add(FileItem.GetItemByPath(Environment.GetEnvironmentVariable("NDKROOT") + "/sources/android/cpufeatures/cpu-features.c"));
+				ModuleHasNDKExtensionsCompiled = ModuleName;
 			}
-			bHasNDKExtensionsCompiled = true;
 		}
 
 		// cache the location of NDK tools
@@ -594,7 +598,7 @@ namespace UnrealBuildTool
 			// Directly added NDK files for NDK extensions
 			if (!UnrealBuildTool.RunningRocket())
 			{
-				ConditionallyAddNDKSourceFiles(SourceFiles);
+				ConditionallyAddNDKSourceFiles(SourceFiles, ModuleName);
 			}
 
 			// Add preprocessor definitions to the argument list.
