fix for Android project settings not showing up in a GitHub source build.
UE-3347
#android

diff --git a/Engine/Source/Programs/UnrealBuildTool/Android/UEBuildAndroid.cs b/Engine/Source/Programs/UnrealBuildTool/Android/UEBuildAndroid.cs
--- a/Engine/Source/Programs/UnrealBuildTool/Android/UEBuildAndroid.cs
+++ b/Engine/Source/Programs/UnrealBuildTool/Android/UEBuildAndroid.cs
@@ -244,10 +244,6 @@ namespace UnrealBuildTool
 							// InModule.AddPlatformSpecificDynamicallyLoadedModule("Android_ES31TargetPlatform");
                         }
 					}
-					else if (InModule.ToString() == "UnrealEd")
-					{
-						InModule.AddPlatformSpecificDynamicallyLoadedModule("AndroidPlatformEditor");
-					}
 					else if (InModule.ToString() == "TargetPlatform")
 					{
 						bBuildShaderFormats = true;
diff --git a/Engine/Source/Programs/UnrealBuildTool/IOS/UEBuildIOS.cs b/Engine/Source/Programs/UnrealBuildTool/IOS/UEBuildIOS.cs
--- a/Engine/Source/Programs/UnrealBuildTool/IOS/UEBuildIOS.cs
+++ b/Engine/Source/Programs/UnrealBuildTool/IOS/UEBuildIOS.cs
@@ -247,10 +247,6 @@ namespace UnrealBuildTool
                             InModule.AddPlatformSpecificDynamicallyLoadedModule("IOSTargetPlatform");
                         }
                     }
-                    else if (InModule.ToString() == "UnrealEd")
-                    {
-                        InModule.AddPlatformSpecificDynamicallyLoadedModule("IOSPlatformEditor");
-                    }
                     else if (InModule.ToString() == "TargetPlatform")
                     {
                         bBuildShaderFormats = true;
