Fixed a problem with Server targets missing from Xcode project

diff --git a/Engine/Source/Programs/UnrealBuildTool/System/XcodeProjectFileGenerator.cs b/Engine/Source/Programs/UnrealBuildTool/System/XcodeProjectFileGenerator.cs
--- a/Engine/Source/Programs/UnrealBuildTool/System/XcodeProjectFileGenerator.cs
+++ b/Engine/Source/Programs/UnrealBuildTool/System/XcodeProjectFileGenerator.cs
@@ -1122,6 +1122,10 @@ namespace UnrealBuildTool
 							{
 								DisplayName += "Editor";
 							}
+							else if (TargetName.Contains("Server"))
+							{
+								DisplayName += "Server";
+							}
 						}
 					}
 
