Remove unused AndroidCommandPath string to remove clang warning

diff --git a/Engine/Source/Programs/UnrealBuildTool/Android/UEDeployAndroid.cs b/Engine/Source/Programs/UnrealBuildTool/Android/UEDeployAndroid.cs
--- a/Engine/Source/Programs/UnrealBuildTool/Android/UEDeployAndroid.cs
+++ b/Engine/Source/Programs/UnrealBuildTool/Android/UEDeployAndroid.cs
@@ -486,7 +486,6 @@ namespace UnrealBuildTool.Android
 			Log.TraceInformation("\n===={0}====PREPARING TO MAKE APK=================================================================", DateTime.Now.ToString());
 
 			// cache some tools paths
-			string AndroidCommandPath = Environment.ExpandEnvironmentVariables("%ANDROID_HOME%/tools/android.bat");
 			string NDKBuildPath = Environment.ExpandEnvironmentVariables("%NDKROOT%/ndk-build.cmd");
 
 			// set up some directory info
