fix incorrect html5 sdk discovery in tool chain.cs

diff --git a/Engine/Source/Programs/UnrealBuildTool/HTML5/HTML5ToolChain.cs b/Engine/Source/Programs/UnrealBuildTool/HTML5/HTML5ToolChain.cs
--- a/Engine/Source/Programs/UnrealBuildTool/HTML5/HTML5ToolChain.cs
+++ b/Engine/Source/Programs/UnrealBuildTool/HTML5/HTML5ToolChain.cs
@@ -194,7 +194,7 @@ namespace UnrealBuildTool
 				Result += " -s CASE_INSENSITIVE_FS=1 ";
 
 
-				string BaseSDKPath = Environment.GetEnvironmentVariable("EMSCRIPTEN");
+                string BaseSDKPath = HTML5SDKInfo.EmscriptenSDKPath();
 				Result += " --js-library \"" + BaseSDKPath + "/Src/library_openal.js\" ";
 			}
 
