one more place where the minimum version of iOS needs to be set to 6.1
#ios

diff --git a/Engine/Source/Programs/UnrealBuildTool/IOS/IOSToolChain.cs b/Engine/Source/Programs/UnrealBuildTool/IOS/IOSToolChain.cs
--- a/Engine/Source/Programs/UnrealBuildTool/IOS/IOSToolChain.cs
+++ b/Engine/Source/Programs/UnrealBuildTool/IOS/IOSToolChain.cs
@@ -98,8 +98,8 @@ namespace UnrealBuildTool
 			{
 				switch (MinVersion)
 				{
-					case "IOS_6":
-						RunTimeIOSVersion = "6.0";
+					case "IOS_61":
+						RunTimeIOSVersion = "6.1";
 						break;
 					case "IOS_7":
 						RunTimeIOSVersion = "7.0";
