correctly qoute emscripten compiler path to fix issues when the path has spaces.  fixes UE-5177

diff --git a/Engine/Source/Programs/UnrealBuildTool/HTML5/HTML5ToolChain.cs b/Engine/Source/Programs/UnrealBuildTool/HTML5/HTML5ToolChain.cs
--- a/Engine/Source/Programs/UnrealBuildTool/HTML5/HTML5ToolChain.cs
+++ b/Engine/Source/Programs/UnrealBuildTool/HTML5/HTML5ToolChain.cs
@@ -19,7 +19,7 @@ namespace UnrealBuildTool
 		{
             if (HTML5SDKInfo.IsSDKInstalled() && HTML5SDKInfo.IsPythonInstalled())
             {
-					EMCCPath = HTML5SDKInfo.EmscriptenCompiler();
+                    EMCCPath = "\"" + HTML5SDKInfo.EmscriptenCompiler() + "\"";
 					PythonPath = HTML5SDKInfo.PythonPath(); 
 
 					// set some environment variable we'll need
