- Added a UEBuildSettings.txt file to the APK generation process so that we can check for high level changes between re-packaging (to make it safer to iteratively package) [UE-4836]

diff --git a/Engine/Source/Programs/UnrealBuildTool/Android/UEDeployAndroid.cs b/Engine/Source/Programs/UnrealBuildTool/Android/UEDeployAndroid.cs
--- a/Engine/Source/Programs/UnrealBuildTool/Android/UEDeployAndroid.cs
+++ b/Engine/Source/Programs/UnrealBuildTool/Android/UEDeployAndroid.cs
@@ -367,29 +367,6 @@ namespace UnrealBuildTool.Android
 			}
 		}
 
-		private static bool IsAPIVersionCurrent(string PropertiesFileLocation)
-		{
-			string ProjectProperties = Path.Combine(PropertiesFileLocation, "project.properties");
-			if (File.Exists(ProjectProperties))
-			{
-				string[] Contents = File.ReadAllLines(ProjectProperties);
-				foreach (string Line in Contents)
-				{
-					// look for the special line
-					if (Line.StartsWith("target="))
-					{
-						// skip over target=
-						string Version = Line.Substring(7);
-						// do they match?
-						return string.Compare(Version, GetSdkApiLevel(), true) == 0;
-					}
-				}
-			}
-
-			// if any other case happens, we are not current
-			return false;
-		}
-
 		private void UpdateProjectProperties(string UE4BuildPath, string ProjectName)
 		{
 			Log.TraceInformation("\n===={0}====UPDATING BUILD CONFIGURATION FILES====================================================", DateTime.Now.ToString());
@@ -481,6 +458,113 @@ namespace UnrealBuildTool.Android
 		}
 
 
+		private string GetAllBuildSettings(string BuildPath, bool bForDistribution, bool bMakeSeparateApks, bool bOBBinApk)
+		{
+			// make the settings string - this will be char by char compared against last time
+			StringBuilder CurrentSettings = new StringBuilder();
+			CurrentSettings.AppendFormat("NDKROOT={0}{1}", Environment.GetEnvironmentVariable("NDKROOT"), Environment.NewLine);
+			CurrentSettings.AppendFormat("ANDROID_HOME={0}{1}", Environment.GetEnvironmentVariable("ANDROID_HOME"), Environment.NewLine);
+			CurrentSettings.AppendFormat("ANT_HOME={0}{1}", Environment.GetEnvironmentVariable("ANT_HOME"), Environment.NewLine);
+			CurrentSettings.AppendFormat("JAVA_HOME={0}{1}", Environment.GetEnvironmentVariable("JAVA_HOME"), Environment.NewLine);
+			CurrentSettings.AppendFormat("SDKVersion={0}{1}", GetSdkApiLevel(), Environment.NewLine);
+			CurrentSettings.AppendFormat("bForDistribution={0}{1}", bForDistribution, Environment.NewLine);
+			CurrentSettings.AppendFormat("bMakeSeparateApks={0}{1}", bMakeSeparateApks, Environment.NewLine);
+			CurrentSettings.AppendFormat("bOBBinApk={0}{1}", bOBBinApk, Environment.NewLine);
+
+			string[] Arches = AndroidToolChain.GetAllArchitectures();
+			foreach (string Arch in Arches)
+			{
+				CurrentSettings.AppendFormat("Arch={0}{1}", Arch, Environment.NewLine);
+			}
+
+			string[] GPUArchitectures = AndroidToolChain.GetAllGPUArchitectures();
+			foreach (string GPUArch in GPUArchitectures)
+			{
+				CurrentSettings.AppendFormat("GPUArch={0}{1}", GPUArch, Environment.NewLine);
+			}
+
+			return CurrentSettings.ToString();
+		}
+
+		private bool CheckDependencies(string ProjectName, string ProjectDirectory, string UE4BuildFilesPath, string GameBuildFilesPath, string EngineDirectory, string JavaSettingsFile, string CookFlavor, string OutputPath, string UE4BuildPath, bool bMakeSeparateApks)
+		{
+			string[] Arches = AndroidToolChain.GetAllArchitectures();
+			string[] GPUArchitectures = AndroidToolChain.GetAllGPUArchitectures();
+
+			// check all input files (.so, java files, .ini files, etc)
+			bool bAllInputsCurrent = true;
+			foreach (string Arch in Arches)
+			{
+				foreach (string GPUArch in GPUArchitectures)
+				{
+					string SourceSOName = AndroidToolChain.InlineArchName(OutputPath, Arch, GPUArch);
+					// if the source binary was UE4Game, replace it with the new project name, when re-packaging a binary only build
+					string ApkFilename = Path.GetFileNameWithoutExtension(OutputPath).Replace("UE4Game", ProjectName);
+					string DestApkName = Path.Combine(ProjectDirectory, "Binaries/Android/") + ApkFilename + ".apk";
+
+					// if we making multiple Apks, we need to put the architecture into the name
+					if (bMakeSeparateApks)
+					{
+						DestApkName = AndroidToolChain.InlineArchName(DestApkName, Arch, GPUArch);
+					}
+
+					// check to see if it's out of date before trying the slow make apk process (look at .so and all Engine and Project build files to be safe)
+					List<String> InputFiles = new List<string>();
+					InputFiles.Add(SourceSOName);
+					InputFiles.AddRange(Directory.EnumerateFiles(UE4BuildFilesPath, "*.*", SearchOption.AllDirectories));
+					if (Directory.Exists(GameBuildFilesPath))
+					{
+						InputFiles.AddRange(Directory.EnumerateFiles(GameBuildFilesPath, "*.*", SearchOption.AllDirectories));
+					}
+
+					// rebuild if .ini files change
+					// @todo android: programmatically determine if any .ini setting changed?
+					InputFiles.Add(Path.Combine(EngineDirectory, "Config\\BaseEngine.ini"));
+					InputFiles.Add(Path.Combine(ProjectDirectory, "Config\\DefaultEngine.ini"));
+
+					// make sure changed java settings will rebuild apk
+					InputFiles.Add(JavaSettingsFile);
+
+					// rebuild if .pak files exist for OBB in APK case
+					if (UEBuildConfiguration.bOBBinAPK)
+					{
+						string PAKFileLocation = ProjectDirectory + "/Saved/StagedBuilds/Android" + CookFlavor + "/" + ProjectName + "/Content/Paks";
+						if (Directory.Exists(PAKFileLocation))
+						{
+							var PakFiles = Directory.EnumerateFiles(PAKFileLocation, "*.pak", SearchOption.TopDirectoryOnly);
+							foreach (var Name in PakFiles)
+							{
+								InputFiles.Add(Name);
+							}
+						}
+					}
+
+					// look for any newer input file
+					DateTime ApkTime = File.GetLastWriteTimeUtc(DestApkName);
+					foreach (var InputFileName in InputFiles)
+					{
+						if (File.Exists(InputFileName))
+						{
+							// skip .log files
+							if (Path.GetExtension(InputFileName) == ".log")
+							{
+								continue;
+							}
+							DateTime InputFileTime = File.GetLastWriteTimeUtc(InputFileName);
+							if (InputFileTime.CompareTo(ApkTime) > 0)
+							{
+								bAllInputsCurrent = false;
+								Log.TraceInformation("{0} is out of date due to newer input file {1}", DestApkName, InputFileName);
+								break;
+							}
+						}
+					}
+				}
+			}
+
+			return bAllInputsCurrent;
+		}
+
 		private void MakeApk(string ProjectName, string ProjectDirectory, string OutputPath, string EngineDirectory, bool bForDistribution, string CookFlavor, bool bMakeSeparateApks, bool bIncrementalPackage)
 		{
 			Log.TraceInformation("\n===={0}====PREPARING TO MAKE APK=================================================================", DateTime.Now.ToString());
@@ -498,91 +582,40 @@ namespace UnrealBuildTool.Android
 			string[] GPUArchitectures = AndroidToolChain.GetAllGPUArchitectures();
 			int NumArches = Arches.Length * GPUArchitectures.Length;
 
+
             // See if we need to create a 'default' Java Build settings file if one doesn't exist (if it does exist we have to assume it has been setup correctly)
             string UE4JavaBuildSettingsFileName = GetUE4JavaBuildSettingsFileName(EngineDirectory);
             WriteJavaBuildSettingsFile(UE4JavaBuildSettingsFileName, UEBuildConfiguration.bOBBinAPK);
 
-			// if the last packaged version is wrong, then we can skip the checking and do it
-			if (!IsAPIVersionCurrent(UE4BuildPath))
+			// check to see if any "meta information" is newer than last time we build
+			string CurrentBuildSettings = GetAllBuildSettings(UE4BuildPath, bForDistribution, bMakeSeparateApks, UEBuildConfiguration.bOBBinAPK);
+			string BuildSettingsCacheFile = Path.Combine(UE4BuildPath, "UEBuildSettings.txt");
+
+			// do we match previous build settings?
+			bool bBuildSettingsMatch = false;
+			if (File.Exists(BuildSettingsCacheFile))
 			{
-				Log.TraceInformation("Output .apk file(s) were made with a different API version, forcing repackage.");
+				string PreviousBuildSettings = File.ReadAllText(BuildSettingsCacheFile);
+				if (PreviousBuildSettings == CurrentBuildSettings)
+				{
+					bBuildSettingsMatch = true;
+				}
+				else
+				{
+					Log.TraceInformation("Previous .apk file(s) were made with different build settings, forcing repackage.");
+				}
 			}
-			else
+
+			// only check input dependencies if the build settings already match
+			if (bBuildSettingsMatch)
 			{
 				// check if so's are up to date against various inputs
-				bool bAllInputsCurrent = true;
+				bool bAllInputsCurrent = CheckDependencies(ProjectName, ProjectDirectory, UE4BuildFilesPath, GameBuildFilesPath, 
+					EngineDirectory, UE4JavaBuildSettingsFileName, CookFlavor, OutputPath, UE4BuildPath, bMakeSeparateApks);
 
-				foreach (string Arch in Arches)
-				{
-					foreach (string GPUArch in GPUArchitectures)
-					{
-						string SourceSOName = AndroidToolChain.InlineArchName(OutputPath, Arch, GPUArch);
-						// if the source binary was UE4Game, replace it with the new project name, when re-packaging a binary only build
-						string ApkFilename = Path.GetFileNameWithoutExtension(OutputPath).Replace("UE4Game", ProjectName);
-						string DestApkName = Path.Combine(ProjectDirectory, "Binaries/Android/") + ApkFilename + ".apk";
-
-						// if we making multiple Apks, we need to put the architecture into the name
-						if (bMakeSeparateApks)
-						{
-							DestApkName = AndroidToolChain.InlineArchName(DestApkName, Arch, GPUArch);
-						}
-
-						// check to see if it's out of date before trying the slow make apk process (look at .so and all Engine and Project build files to be safe)
-						List<String> InputFiles = new List<string>();
-						InputFiles.Add(SourceSOName);
-						InputFiles.AddRange(Directory.EnumerateFiles(UE4BuildFilesPath, "*.*", SearchOption.AllDirectories));
-						if (Directory.Exists(GameBuildFilesPath))
-						{
-							InputFiles.AddRange(Directory.EnumerateFiles(GameBuildFilesPath, "*.*", SearchOption.AllDirectories));
-						}
-
-						// rebuild if .ini files change
-						// @todo android: programmatically determine if any .ini setting changed?
-						InputFiles.Add(Path.Combine(EngineDirectory, "Config\\BaseEngine.ini"));
-						InputFiles.Add(Path.Combine(ProjectDirectory, "Config\\DefaultEngine.ini"));
-
-						// make sure changed java settings will rebuild apk
-						InputFiles.Add(UE4JavaBuildSettingsFileName);
-
-						// rebuild if .pak files exist for OBB in APK case
-						if (UEBuildConfiguration.bOBBinAPK)
-						{
-							string PAKFileLocation = ProjectDirectory + "/Saved/StagedBuilds/Android" + CookFlavor + "/" + ProjectName + "/Content/Paks";
-							if (Directory.Exists(PAKFileLocation))
-							{
-								var PakFiles = Directory.EnumerateFiles(PAKFileLocation, "*.pak", SearchOption.TopDirectoryOnly);
-								foreach (var Name in PakFiles)
-								{
-									InputFiles.Add(Name);
-								}
-							}
-						}
-
-						// look for any newer input file
-						DateTime ApkTime = File.GetLastWriteTimeUtc(DestApkName);
-						foreach (var InputFileName in InputFiles)
-						{
-							if (File.Exists(InputFileName))
-							{
-								// skip .log files
-								if (Path.GetExtension(InputFileName) == ".log")
-								{
-									continue;
-								}
-								DateTime InputFileTime = File.GetLastWriteTimeUtc(InputFileName);
-								if (InputFileTime.CompareTo(ApkTime) > 0)
-								{
-									bAllInputsCurrent = false;
-									Log.TraceInformation("{0} is out of date due to newer input file {1}", DestApkName, InputFileName);
-									break;
-								}
-							}
-						}
-					}
-				}
 				if (bAllInputsCurrent)
 				{
-					Log.TraceInformation("Output .apk file(s) are up to date (compared to the .so and .java input files)");
+					Log.TraceInformation("Output .apk file(s) are up to date (dependencies and build settings are up to date)");
 					return;
 				}
 			}
@@ -656,6 +689,9 @@ namespace UnrealBuildTool.Android
 			// update metadata files (like project.properties, build.xml) if we are missing a build.xml or if we just overwrote project.properties with a bad version in it (from game/engine dir)
 			UpdateProjectProperties(UE4BuildPath, ProjectName);
 
+			// at this point, we can write out the cached build settings to compare for a next build
+			File.WriteAllText(BuildSettingsCacheFile, CurrentBuildSettings);
+
 			// now make the apk(s)
 			string FinalNdkBuildABICommand = "";
 			for (int ArchIndex = 0; ArchIndex < Arches.Length; ArchIndex++)
