make it so debug symbols are not generated for rocket redist libs
UEPLAT-249
#ios

diff --git a/Engine/Source/Programs/UnrealBuildTool/IOS/IOSToolChain.cs b/Engine/Source/Programs/UnrealBuildTool/IOS/IOSToolChain.cs
--- a/Engine/Source/Programs/UnrealBuildTool/IOS/IOSToolChain.cs
+++ b/Engine/Source/Programs/UnrealBuildTool/IOS/IOSToolChain.cs
@@ -376,7 +376,7 @@ namespace UnrealBuildTool
 			}
 
 			// Create DWARF format debug info if wanted,
-			if (CompileEnvironment.Config.bCreateDebugInfo)
+			if (CompileEnvironment.Config.bCreateDebugInfo && !UnrealBuildTool.BuildingRocket())
 			{
 				Result += " -gdwarf-2";
 			}
@@ -1099,7 +1099,7 @@ namespace UnrealBuildTool
 				string FinalRemoteExecutablePath = String.Format("{0}/Payload/{1}.app/{1}", RemoteShadowDirectoryMac, Target.GameName);
 
 				// strip the debug info from the executable if needed
-				if (BuildConfiguration.bStripSymbolsOnIOS || (Target.Configuration == UnrealTargetConfiguration.Shipping))
+				if (BuildConfiguration.bStripSymbolsOnIOS || (Target.Configuration == UnrealTargetConfiguration.Shipping) || UnrealBuildTool.BuildingRocket())
 				{
 					Process StripProcess = new Process();
 					StripProcess.StartInfo.WorkingDirectory = RemoteShadowDirectoryMac;
