Strip debug symbols from Rocket builds
#android
#ue4

diff --git a/Engine/Source/Programs/UnrealBuildTool/Android/AndroidToolChain.cs b/Engine/Source/Programs/UnrealBuildTool/Android/AndroidToolChain.cs
--- a/Engine/Source/Programs/UnrealBuildTool/Android/AndroidToolChain.cs
+++ b/Engine/Source/Programs/UnrealBuildTool/Android/AndroidToolChain.cs
@@ -294,7 +294,7 @@ namespace UnrealBuildTool
 			}
 
 			// debug info
-			if (CompileEnvironment.Config.bCreateDebugInfo)
+            if (CompileEnvironment.Config.bCreateDebugInfo && !UnrealBuildTool.BuildingRocket())
 			{
 				Result += " -g2 -gdwarf-2";
 			}
@@ -424,7 +424,12 @@ namespace UnrealBuildTool
 			Result += " -Wl,-shared,-Bsymbolic";
 			Result += " -Wl,--no-undefined";
 
-			if (Architecture == "-armv7")
+            if (UnrealBuildTool.BuildingRocket())
+            {
+                Result += " -Wl,--strip-debug";
+            }
+
+            if (Architecture == "-armv7")
 			{
 				Result += ToolchainParamsArm;
 				Result += " -march=armv7-a";
