Fix UBT locking up on Linux.

- Provide Linux stubs for GetProcesses(). Using Path.GetFullPath() on process modules leads to mono locking up sometimes.
- Fixes UE-4950.

diff --git a/Engine/Source/Programs/UnrealBuildTool/System/BuildHostPlatform.cs b/Engine/Source/Programs/UnrealBuildTool/System/BuildHostPlatform.cs
--- a/Engine/Source/Programs/UnrealBuildTool/System/BuildHostPlatform.cs
+++ b/Engine/Source/Programs/UnrealBuildTool/System/BuildHostPlatform.cs
@@ -414,5 +414,29 @@ namespace UnrealBuildTool
 			// @TODO: Implement GetDllApiVersion for Linux
 			return -1;
 		}
+
+		/// <summary>
+		/// Currently Mono returns incomplete process names in Process.GetProcesses() so we need to use /proc
+		/// (also, Mono locks up during process traversal sometimes, trying to open /dev/snd/pcm*)
+		/// </summary>
+		/// <returns></returns>
+		public override ProcessInfo[] GetProcesses()
+		{
+			// @TODO: Implement for Linux
+			return new List<ProcessInfo>().ToArray();
+		}
+		
+		/// <summary>
+		/// Currently Mono returns incomplete list of modules for Process.Modules so we need to parse /proc/PID/maps.
+		/// (also, Mono locks up during process traversal sometimes, trying to open /dev/snd/pcm*)
+		/// </summary>
+		/// <param name="PID"></param>
+		/// <param name="Filename"></param>
+		/// <returns></returns>
+		public override string[] GetProcessModules(int PID, string Filename)
+		{
+			// @TODO: Implement for Linux
+			return new List<string>().ToArray();
+		}
 	}
 }
