Make UBT look for a few more clang versions (UE-5508).

- Helps on Ubuntu 14.10 and newer Debian that switched to co-installable clang.

diff --git a/Engine/Source/Programs/UnrealBuildTool/Linux/LinuxToolChain.cs b/Engine/Source/Programs/UnrealBuildTool/Linux/LinuxToolChain.cs
--- a/Engine/Source/Programs/UnrealBuildTool/Linux/LinuxToolChain.cs
+++ b/Engine/Source/Programs/UnrealBuildTool/Linux/LinuxToolChain.cs
@@ -135,7 +135,15 @@ namespace UnrealBuildTool
             if (!CrossCompiling())
             {
                 // use native linux toolchain
-                ClangPath = Which("clang++");
+                string[] ClangNames = { "clang++", "clang++-3.5", "clang++-3.3" };
+                foreach (var ClangName in ClangNames)
+                {
+					ClangPath = Which(ClangName);
+					if (!String.IsNullOrEmpty(ClangPath))
+					{
+						break;
+					}
+				}
                 GCCPath = Which("g++");
                 ArPath = Which("ar");
                 RanlibPath = Which("ranlib");
