Skip path conversion for custom resources in MacToolChain if the path is empty

diff --git a/Engine/Source/Programs/UnrealBuildTool/Mac/MacToolChain.cs b/Engine/Source/Programs/UnrealBuildTool/Mac/MacToolChain.cs
--- a/Engine/Source/Programs/UnrealBuildTool/Mac/MacToolChain.cs
+++ b/Engine/Source/Programs/UnrealBuildTool/Mac/MacToolChain.cs
@@ -962,8 +962,11 @@ namespace UnrealBuildTool
 						CustomBuildPath = Path.GetDirectoryName(UProjectFilePath) + "/Build/Mac";
 					}
 
-					CustomResourcesPath = ConvertPath(Path.GetFullPath(CustomResourcesPath));
-					CustomBuildPath = ConvertPath(Path.GetFullPath(CustomBuildPath));
+					if (!string.IsNullOrEmpty(CustomResourcesPath) && !string.IsNullOrEmpty(CustomBuildPath))
+					{
+						CustomResourcesPath = ConvertPath(Path.GetFullPath(CustomResourcesPath));
+						CustomBuildPath = ConvertPath(Path.GetFullPath(CustomBuildPath));
+					}
 
 					bool bBuildingEditor = GameName.EndsWith("Editor");
 
