Enable one more (rather crucial) warning for native builds too.

diff --git a/Engine/Source/Programs/UnrealBuildTool/Linux/LinuxToolChain.cs b/Engine/Source/Programs/UnrealBuildTool/Linux/LinuxToolChain.cs
--- a/Engine/Source/Programs/UnrealBuildTool/Linux/LinuxToolChain.cs
+++ b/Engine/Source/Programs/UnrealBuildTool/Linux/LinuxToolChain.cs
@@ -319,11 +319,6 @@ namespace UnrealBuildTool
 				{
 					Result += " -Wno-undefined-bool-conversion";	// hides checking if 'this' pointer is null
 				}
-
-                if (!CrossCompiling())
-                {
-                    Result += " -Wno-logical-op-parentheses";   // needed for external headers we can't change
-                }
             }
 
             Result += " -Wno-unused-variable";
