Fixed major UBT bug with bEnableFastBuildIteration enabled
- Fixed generated headers not compiling due to invalid inlined PCH includes
- This could be reproed:
       - Build target (until up to date).  Regen project files.  Touch a .cpp (not a UObject header).  Build.  Touch a UObject header.  Build. (generated .cpps will now be bogus).

diff --git a/Engine/Source/Programs/UnrealBuildTool/System/ExternalExecution.cs b/Engine/Source/Programs/UnrealBuildTool/System/ExternalExecution.cs
--- a/Engine/Source/Programs/UnrealBuildTool/System/ExternalExecution.cs
+++ b/Engine/Source/Programs/UnrealBuildTool/System/ExternalExecution.cs
@@ -544,7 +544,7 @@ namespace UnrealBuildTool
 
 				// ensure the headers are up to date
 				bool bUHTNeedsToRun = (UEBuildConfiguration.bForceHeaderGeneration == true || AreGeneratedCodeFilesOutOfDate(UObjectModules));
-				if( bUHTNeedsToRun )
+				if( bUHTNeedsToRun || UnrealBuildTool.IsGatheringBuild )
 				{
 					// Since code files are definitely out of date, we'll now finish computing information about the UObject modules for UHT.  We
 					// want to save this work until we know that UHT actually needs to be run to speed up best-case iteration times.
