PS4 - PS4DevKitUtil - standalone app for controlling a PS4 Devkit
* Devkit enumeration
* Devkit details (name, power status, connection state, SDK Version, HostName, MappedDrive)
* Device snapshot (running processes and thread info)
* Launch elfs
* Capture Devkit debug output

diff --git a/Engine/Source/Programs/UnrealBuildTool/System/ProjectFileGenerator.cs b/Engine/Source/Programs/UnrealBuildTool/System/ProjectFileGenerator.cs
--- a/Engine/Source/Programs/UnrealBuildTool/System/ProjectFileGenerator.cs
+++ b/Engine/Source/Programs/UnrealBuildTool/System/ProjectFileGenerator.cs
@@ -501,7 +501,10 @@ namespace UnrealBuildTool
 
 					// Add all of the Android C# projects
 					AddAndroidProjects( ProgramsFolder );
-				}
+
+                    // Add all of the PS4 C# projects
+                    AddPS4Projects( ProgramsFolder );
+                }
 
 
 				// Eliminate all redundant master project folders.  E.g., folders which contain only one project and that project
@@ -1220,6 +1223,19 @@ namespace UnrealBuildTool
 			}
 		}
 
+        /// <summary>
+        /// Adds all of the PS4 C# projects to the master project
+        /// </summary>
+        private void AddPS4Projects(MasterProjectFolder Folder)
+        {
+            string ProjectFolderName = Path.Combine(EngineRelativePath, "Source", "Programs", "PS4");
+            DirectoryInfo ProjectFolderInfo = new DirectoryInfo(ProjectFolderName);
+            if (ProjectFolderInfo.Exists)
+            {
+                Folder.ChildProjects.Add(AddSimpleCSharpProject("PS4/PS4DevKitUtil"));
+            }
+        }
+
 		/// <summary>
 		/// Adds all of the Android C# projects to the master project
 		/// </summary>
