Log more autosdks info when running on buildmachines.

diff --git a/Engine/Source/Programs/UnrealBuildTool/Configuration/UEBuildPlatform.cs b/Engine/Source/Programs/UnrealBuildTool/Configuration/UEBuildPlatform.cs
--- a/Engine/Source/Programs/UnrealBuildTool/Configuration/UEBuildPlatform.cs
+++ b/Engine/Source/Programs/UnrealBuildTool/Configuration/UEBuildPlatform.cs
@@ -872,6 +872,8 @@ namespace UnrealBuildTool
 
 		protected static string AutoSetupEnvVar = "AutoSDKSetup";
 
+		public static bool bShouldLogInfo = false;
+
 		/** 
 		 * Whether platform supports switching SDKs during runtime
 		 * 
@@ -1132,7 +1134,7 @@ namespace UnrealBuildTool
 			if (!IsAutoSDKSafe())
 			{
 				Console.ForegroundColor = ConsoleColor.Red;
-				Console.WriteLine(GetSDKTargetPlatformName() + " attempted to run SDK hook which could have damaged manual SDK install!");
+				LogAutoSDK( GetSDKTargetPlatformName() + " attempted to run SDK hook which could have damaged manual SDK install!");				
 				Console.ResetColor();
 
 				return false;
@@ -1144,23 +1146,28 @@ namespace UnrealBuildTool
 
 				if (File.Exists(HookExe))
 				{
-					Console.WriteLine("Running {0} hook {1}", Hook, HookExe);
-
+					LogAutoSDK("Running {0} hook {1}", Hook, HookExe);
+					
 					// run it
 					Process HookProcess = new Process();
 					HookProcess.StartInfo.WorkingDirectory = SDKDirectory;
 					HookProcess.StartInfo.FileName = HookExe;
 					HookProcess.StartInfo.Arguments = "";
 					HookProcess.StartInfo.WindowStyle = ProcessWindowStyle.Hidden;
+					HookProcess.StartInfo.UseShellExecute = false;
+					HookProcess.StartInfo.RedirectStandardOutput = true;
+					HookProcess.StartInfo.RedirectStandardError = true;					
 
 					//installers may require administrator access to succeed. so run as an admmin.
 					HookProcess.StartInfo.Verb = "runas";
 					HookProcess.Start();
 					HookProcess.WaitForExit();
 
+					LogAutoSDK(HookProcess.StandardOutput.ReadToEnd());
+					LogAutoSDK(HookProcess.StandardError.ReadToEnd());
 					if (HookProcess.ExitCode != 0)
 					{
-						Console.WriteLine("Hook exited uncleanly (returned {0}), considering it failed.", HookProcess.ExitCode);
+						LogAutoSDK("Hook exited uncleanly (returned {0}), considering it failed.", HookProcess.ExitCode);
 						return false;
 					}
 
@@ -1168,12 +1175,12 @@ namespace UnrealBuildTool
 				}
 				else
 				{
-					Console.WriteLine("File {0} does not exist", HookExe);
+					LogAutoSDK("File {0} does not exist", HookExe);
 				}
 			}
 			else
 			{
-				Console.WriteLine("Version string is blank for {0}. Can't determine {1} hook.", PlatformSDKRoot, Hook.ToString());
+				LogAutoSDK("Version string is blank for {0}. Can't determine {1} hook.", PlatformSDKRoot, Hook.ToString());
 			}
 
 			return bHookCanBeNonExistent;
@@ -1214,8 +1221,8 @@ namespace UnrealBuildTool
 						string[] Parts = VariableString.Split('=');
 						if (Parts.Length != 2)
 						{
-							Console.WriteLine("Incorrect environment variable declaration:");
-							Console.WriteLine(VariableString);
+							LogAutoSDK("Incorrect environment variable declaration:");
+							LogAutoSDK(VariableString);
 							return false;
 						}
 
@@ -1248,7 +1255,7 @@ namespace UnrealBuildTool
 						string EnvVarValue = EnvVarValues[i];
 						if (BuildConfiguration.bPrintDebugInfo)
 						{
-							Console.WriteLine("Setting variable '{0}' to '{1}'", EnvVarName, EnvVarValue);
+							LogAutoSDK("Setting variable '{0}' to '{1}'", EnvVarName, EnvVarValue);
 						}
 						Environment.SetEnvironmentVariable(EnvVarName, EnvVarValue);
 					}
@@ -1269,10 +1276,7 @@ namespace UnrealBuildTool
 						{
 							if (PathVar.IndexOf(PathRemove, StringComparison.OrdinalIgnoreCase) >= 0)
 							{
-								if (BuildConfiguration.bPrintDebugInfo)
-								{
-									Console.WriteLine("Removing Path: '{0}'", PathVar);
-								}
+								LogAutoSDK("Removing Path: '{0}'", PathVar);								
 								ModifiedPathVars.Remove(PathVar);
 							}
 						}
@@ -1286,10 +1290,7 @@ namespace UnrealBuildTool
 						{
 							if (String.Compare(PathAdd, PathVar, true) == 0)
 							{
-								if (BuildConfiguration.bPrintDebugInfo)
-								{
-									Console.WriteLine("Removing Path: '{0}'", PathVar);
-								}
+								LogAutoSDK("Removing Path: '{0}'", PathVar);								
 								ModifiedPathVars.Remove(PathVar);
 							}
 						}
@@ -1300,10 +1301,7 @@ namespace UnrealBuildTool
 					{
 						if (!ModifiedPathVars.Contains(PathAdd))
 						{
-							if (BuildConfiguration.bPrintDebugInfo)
-							{
-								Console.WriteLine("Adding Path: '{0}'", PathAdd);
-							}
+							LogAutoSDK("Adding Path: '{0}'", PathAdd);							
 							ModifiedPathVars.Add(PathAdd);
 						}
 					}
@@ -1333,7 +1331,7 @@ namespace UnrealBuildTool
 			}
 			else
 			{
-				Console.WriteLine("Cannot set up environment for {1} because command file {1} does not exist.", PlatformSDKRoot, EnvVarFile);
+				LogAutoSDK("Cannot set up environment for {1} because command file {1} does not exist.", PlatformSDKRoot, EnvVarFile);
 			}
 
 			return false;
@@ -1483,7 +1481,7 @@ namespace UnrealBuildTool
 					// switch over (note that version string can be empty)
 					if (!RunAutoSDKHooks(AutoSDKRoot, CurrentSDKString, SDKHookType.Uninstall))
 					{
-						Console.WriteLine("Failed to uninstall currently installed SDK {0}", CurrentSDKString);
+						LogAutoSDK("Failed to uninstall currently installed SDK {0}", CurrentSDKString);
 						InvalidateCurrentlyInstalledAutoSDK();
 						return;
 					}
@@ -1492,7 +1490,7 @@ namespace UnrealBuildTool
 
 					if (!RunAutoSDKHooks(AutoSDKRoot, GetRequiredSDKString(), SDKHookType.Install, false))
 					{
-						Console.WriteLine("Failed to install required SDK {0}.  Attemping to uninstall", GetRequiredSDKString());
+						LogAutoSDK("Failed to install required SDK {0}.  Attemping to uninstall", GetRequiredSDKString());
 						RunAutoSDKHooks(AutoSDKRoot, GetRequiredSDKString(), SDKHookType.Uninstall, false);
 						return;
 					}
@@ -1500,7 +1498,7 @@ namespace UnrealBuildTool
 					string EnvVarFile = Path.Combine(AutoSDKRoot, SDKEnvironmentVarsFile);
 					if (!File.Exists(EnvVarFile))
 					{
-						Console.WriteLine("Installation of required SDK {0}.  Did not generate Environment file {1}", GetRequiredSDKString(), EnvVarFile);
+						LogAutoSDK("Installation of required SDK {0}.  Did not generate Environment file {1}", GetRequiredSDKString(), EnvVarFile);
 						RunAutoSDKHooks(AutoSDKRoot, GetRequiredSDKString(), SDKHookType.Uninstall, false);
 						return;
 					}
@@ -1545,7 +1543,7 @@ namespace UnrealBuildTool
 			// load environment variables from current SDK
 			if (!SetupEnvironmentFromAutoSDK(PlatformSDKRoot))
 			{
-				Console.WriteLine("Failed to load environment from required SDK {0}", GetRequiredSDKString());
+				LogAutoSDK("Failed to load environment from required SDK {0}", GetRequiredSDKString());
 				InvalidateCurrentlyInstalledAutoSDK();
 				return SDKStatus.Invalid;
 			}
@@ -1573,6 +1571,8 @@ namespace UnrealBuildTool
 		// Arbitrates between manual SDKs and setting up AutoSDK based on program options and platform preferences.
 		public void ManageAndValidateSDK()
 		{
+			bShouldLogInfo = BuildConfiguration.bPrintDebugInfo || Environment.GetEnvironmentVariable("IsBuildMachine") == "1";
+
 			// do not modify installed manifests if parent process has already set everything up.
 			// this avoids problems with determining IsAutoSDKSafe and doing an incorrect invalidate.
 			if (bAllowAutoSDKSwitching && !HasParentProcessSetupAutoSDK())
@@ -1598,7 +1598,8 @@ namespace UnrealBuildTool
 				}
 			}
 
-			if (BuildConfiguration.bPrintDebugInfo)
+
+			if (bShouldLogInfo)
 			{
 				PrintSDKInfo();
 			}
@@ -1606,28 +1607,42 @@ namespace UnrealBuildTool
 
 		public void PrintSDKInfo()
 		{
-
-
 			if (HasRequiredSDKsInstalled() == SDKStatus.Valid)
 			{
 				bool bHasRequiredManualSDK = HasRequiredManualSDK() == SDKStatus.Valid;
 				if (HasSetupAutoSDK())
 				{
 					string PlatformSDKRoot = GetPathToPlatformAutoSDKs();
-					Console.WriteLine(GetSDKTargetPlatformName() + " using SDK from: " + Path.Combine(PlatformSDKRoot, GetRequiredSDKString()));
+					LogAutoSDK(GetSDKTargetPlatformName() + " using SDK from: " + Path.Combine(PlatformSDKRoot, GetRequiredSDKString()));
 				}
 				else if (bHasRequiredManualSDK)
 				{
-					Console.WriteLine(this.ToString() + " using manually installed SDK " + GetRequiredSDKString());
+					LogAutoSDK(this.ToString() + " using manually installed SDK " + GetRequiredSDKString());
 				}
 				else
 				{
-					Console.WriteLine(this.ToString() + " setup error.  Inform platform team.");
+					LogAutoSDK(this.ToString() + " setup error.  Inform platform team.");
 				}
 			}
 			else
 			{
-				Console.WriteLine(this.ToString() + " has no valid SDK");
+				LogAutoSDK(this.ToString() + " has no valid SDK");
+			}
+		}
+
+		protected static void LogAutoSDK(string Format, params object[] Args)
+		{
+			if (bShouldLogInfo)
+			{
+				Log.WriteLine(TraceEventType.Information, Format, Args);
+			}
+		}
+
+		protected static void LogAutoSDK(String Message)
+		{
+			if (bShouldLogInfo)
+			{
+				Log.WriteLine(TraceEventType.Information, Message);
 			}
 		}
 
