Maybe fix build machine timeouts.

diff --git a/Engine/Source/Programs/UnrealBuildTool/Configuration/UEBuildPlatform.cs b/Engine/Source/Programs/UnrealBuildTool/Configuration/UEBuildPlatform.cs
--- a/Engine/Source/Programs/UnrealBuildTool/Configuration/UEBuildPlatform.cs
+++ b/Engine/Source/Programs/UnrealBuildTool/Configuration/UEBuildPlatform.cs
@@ -1154,17 +1154,19 @@ namespace UnrealBuildTool
 					HookProcess.StartInfo.FileName = HookExe;
 					HookProcess.StartInfo.Arguments = "";
 					HookProcess.StartInfo.WindowStyle = ProcessWindowStyle.Hidden;
-					HookProcess.StartInfo.UseShellExecute = false;
-					HookProcess.StartInfo.RedirectStandardOutput = true;
-					HookProcess.StartInfo.RedirectStandardError = true;					
+
+					// seems to break the build machines?
+					//HookProcess.StartInfo.UseShellExecute = false;
+					//HookProcess.StartInfo.RedirectStandardOutput = true;
+					//HookProcess.StartInfo.RedirectStandardError = true;					
 
 					//installers may require administrator access to succeed. so run as an admmin.
 					HookProcess.StartInfo.Verb = "runas";
 					HookProcess.Start();
 					HookProcess.WaitForExit();
 
-					LogAutoSDK(HookProcess.StandardOutput.ReadToEnd());
-					LogAutoSDK(HookProcess.StandardError.ReadToEnd());
+					//LogAutoSDK(HookProcess.StandardOutput.ReadToEnd());
+					//LogAutoSDK(HookProcess.StandardError.ReadToEnd());
 					if (HookProcess.ExitCode != 0)
 					{
 						LogAutoSDK("Hook exited uncleanly (returned {0}), considering it failed.", HookProcess.ExitCode);
