GitHub PR #641: Allow XGE to be used when compiling from binary-only version of UE4 (if XGE is installed)

https://github.com/EpicGames/UnrealEngine/pull/641

diff --git a/Engine/Source/Programs/UnrealBuildTool/Configuration/BuildConfiguration.cs b/Engine/Source/Programs/UnrealBuildTool/Configuration/BuildConfiguration.cs
--- a/Engine/Source/Programs/UnrealBuildTool/Configuration/BuildConfiguration.cs
+++ b/Engine/Source/Programs/UnrealBuildTool/Configuration/BuildConfiguration.cs
@@ -570,11 +570,6 @@ namespace UnrealBuildTool
                 bLogDetailedActionStats = false;
 			}
 
-			if (UnrealBuildTool.RunningRocket())
-			{
-				bAllowXGE = false;
-			}
-
 			// PDB
 			if (bUsePDBFiles)
 			{
