Getting rid of auto-startup modules feature as it was causing confusion with loading time and it is not worth it to make it flexible as users expected it.

diff --git a/Engine/Source/Developer/GameplayDebugger/GameplayDebugger.Build.cs b/Engine/Source/Developer/GameplayDebugger/GameplayDebugger.Build.cs
--- a/Engine/Source/Developer/GameplayDebugger/GameplayDebugger.Build.cs
+++ b/Engine/Source/Developer/GameplayDebugger/GameplayDebugger.Build.cs
@@ -6,8 +6,6 @@ namespace UnrealBuildTool.Rules
     {
         public GameplayDebugger(TargetInfo Target)
         {
-            bIsAutoStartupModule = false; //we have to hold on with this feature for now.
-
             PublicIncludePaths.AddRange(
                 new string[] {
 				    "Developer/GameplayDebugger/Public",
diff --git a/Engine/Source/Developer/NewLogVisualizer/NewLogVisualizer.Build.cs b/Engine/Source/Developer/NewLogVisualizer/NewLogVisualizer.Build.cs
--- a/Engine/Source/Developer/NewLogVisualizer/NewLogVisualizer.Build.cs
+++ b/Engine/Source/Developer/NewLogVisualizer/NewLogVisualizer.Build.cs
@@ -6,8 +6,6 @@ public class NewLogVisualizer : ModuleRules
 {
 	public NewLogVisualizer(TargetInfo Target)
 	{
-        bIsAutoStartupModule = true;
-        
         PrivateIncludePathModuleNames.AddRange(
 			new string[] {
 				"DesktopPlatform",
diff --git a/Engine/Source/Programs/UnrealBuildTool/Configuration/UEBuildEditor.cs b/Engine/Source/Programs/UnrealBuildTool/Configuration/UEBuildEditor.cs
--- a/Engine/Source/Programs/UnrealBuildTool/Configuration/UEBuildEditor.cs
+++ b/Engine/Source/Programs/UnrealBuildTool/Configuration/UEBuildEditor.cs
@@ -58,13 +58,6 @@ namespace UnrealBuildTool
 																					bInAllowExports:!ShouldCompileMonolithic(),
 																					InModuleNames: new List<string>() { "Launch" } );
 
-				// This "default" game module is required by the editor executable in case
-				// you want to run it in content-only mode.
-				if (!ExtraModuleNames.Contains("UE4Game"))
-				{
-					ExtraModuleNames.Add("UE4Game");
-				}
-
 				AppBinaries.Add( new UEBuildBinaryCPP( this, Config ) );
 			}
 
diff --git a/Engine/Source/Programs/UnrealBuildTool/Configuration/UEBuildModule.cs b/Engine/Source/Programs/UnrealBuildTool/Configuration/UEBuildModule.cs
--- a/Engine/Source/Programs/UnrealBuildTool/Configuration/UEBuildModule.cs
+++ b/Engine/Source/Programs/UnrealBuildTool/Configuration/UEBuildModule.cs
@@ -978,9 +978,6 @@ namespace UnrealBuildTool
 		/** If true and unity builds are enabled, this module will build without unity. */
 		public bool bFasterWithoutUnity = false;
 
-		/** If true then the engine will call it's StartupModule at engine initialization automatically. */
-		public bool bIsAutoStartupModule = false;
-
 		/** Overrides BuildConfiguration.MinFilesUsingPrecompiledHeader if non-zero. */
 		public int MinFilesUsingPrecompiledHeaderOverride = 0;
 
@@ -1072,7 +1069,6 @@ namespace UnrealBuildTool
 			bool InUseRTTI,
 			bool InEnableBufferSecurityChecks,
 			bool InFasterWithoutUnity,
-			bool InIsAutoStartupModule,
 			int InMinFilesUsingPrecompiledHeaderOverride,
 			bool InEnableExceptions,
 			bool bInBuildSourceFiles
@@ -1123,7 +1119,6 @@ namespace UnrealBuildTool
 			bUseRTTI                               = InUseRTTI;
 			bEnableBufferSecurityChecks 		   = InEnableBufferSecurityChecks;
 			bFasterWithoutUnity                    = InFasterWithoutUnity;
-			bIsAutoStartupModule                   = InIsAutoStartupModule;
 			MinFilesUsingPrecompiledHeaderOverride = InMinFilesUsingPrecompiledHeaderOverride;
 			bEnableExceptions                      = InEnableExceptions;
 		}
@@ -2156,7 +2151,6 @@ namespace UnrealBuildTool
 			bool InUseRTTI,
 			bool InEnableBufferSecurityChecks,
 			bool InFasterWithoutUnity,
-			bool InIsAutoStartupModule,
 			int InMinFilesUsingPrecompiledHeaderOverride,
 			bool InEnableExceptions,
 			bool bInBuildSourceFiles
@@ -2166,7 +2160,7 @@ namespace UnrealBuildTool
 			InPublicIncludePathModuleNames,InPublicDependencyModuleNames,InPublicDelayLoadDLLs,InPublicAdditionalLibraries,InPublicFrameworks,InPublicWeakFrameworks,InPublicAdditionalFrameworks,InPublicAdditionalShadowFiles,InPublicAdditionalBundleResources,
 			InPrivateIncludePaths,InPrivateIncludePathModuleNames,InPrivateDependencyModuleNames,
             InCircularlyReferencedDependentModules, InDynamicallyLoadedModuleNames, InPlatformSpecificDynamicallyLoadedModuleNames, InOptimizeCode,
-			InAllowSharedPCH, InSharedPCHHeaderFile, InUseRTTI, InEnableBufferSecurityChecks, InFasterWithoutUnity, InIsAutoStartupModule, InMinFilesUsingPrecompiledHeaderOverride,
+			InAllowSharedPCH, InSharedPCHHeaderFile, InUseRTTI, InEnableBufferSecurityChecks, InFasterWithoutUnity, InMinFilesUsingPrecompiledHeaderOverride,
 			InEnableExceptions, bInBuildSourceFiles)
 		{
 			PrivateAssemblyReferences = HashSetFromOptionalEnumerableStringParameter(InPrivateAssemblyReferences);
diff --git a/Engine/Source/Programs/UnrealBuildTool/Configuration/UEBuildTarget.cs b/Engine/Source/Programs/UnrealBuildTool/Configuration/UEBuildTarget.cs
--- a/Engine/Source/Programs/UnrealBuildTool/Configuration/UEBuildTarget.cs
+++ b/Engine/Source/Programs/UnrealBuildTool/Configuration/UEBuildTarget.cs
@@ -1545,8 +1545,6 @@ namespace UnrealBuildTool
 				CreateLinkerFixupsCPPFile();
 			}
 
-			CreateAutoStartupModuleListGetters();
-
 			// Build the target's binaries.
 			foreach (var Binary in AppBinaries)
 			{
@@ -1836,136 +1834,6 @@ namespace UnrealBuildTool
 			return FilteredBinaries;
 		}
 
-		/// <summary>
-		/// Gets auto-startup module list.
-		/// </summary>
-		/// <param name="Binary">Binary to gather auto-startup list.</param>
-		/// <returns>Auto-startup module name list.</returns>
-		private IEnumerable<string> GetAutoStartupModuleList(UEBuildBinary Binary)
-		{
-			var Output = new List<string>();
-
-			foreach (var DepModule in Binary.GetAllDependencyModules(false, true))
-			{
-				if (DepModule is UEBuildModuleCPP && (DepModule as UEBuildModuleCPP).bIsAutoStartupModule)
-				{
-					Output.Add(DepModule.Name);
-				}
-			}
-
-			return Output;
-		}
-
-		/// <summary>
-		/// Tests if the module provided is a game module.
-		/// </summary>
-		/// <param name="Module">Module to check.</param>
-		/// <returns>True if provided module is a game module. False otherwise.</returns>
-		private static bool IsGameModule(UEBuildModule Module)
-		{
-			return Module.Type == UEBuildModuleType.Game || (Module.Type == UEBuildModuleType.Runtime && Module.Name == "UE4Game");
-		}
-
-		/// <summary>
-		/// Creates getters that passes auto-startup module list to the engine.
-		/// </summary>
-		private void CreateAutoStartupModuleListGetters()
-		{
-			if (TargetType != TargetRules.TargetType.Program && !ShouldCompileMonolithic())
-			{
-				foreach (var Module in ExtraModuleNames.Select(Name => GetModuleByName(Name)).Where(
-						InnerModule => IsGameModule(InnerModule)).Cast<UEBuildModuleCPP>())
-				{
-					UEBuildBinary Binary = null;
-					if (Module.Name == "UE4Game")
-					{
-						Binary = AppBinaries[0];
-					}
-					else
-					{
-						Binary = Module.Binary;
-					}
-
-					Module.AddAdditionalCPPFiles(new FileItem[] { FileItem.GetItemByPath(
-							CreateAutoGeneratedModulesListGetterCPPFile(
-								GetAutoStartupModuleList(Binary),
-								Module.Binary.Config.IntermediateDirectory)
-						)
-					});
-				}
-			}
-			else
-			{
-				var SourceFilename = CreateAutoGeneratedModulesListGetterCPPFile(
-					GetAutoStartupModuleList(AppBinaries[0]),
-					AppBinaries[0].Config.IntermediateDirectory);
-
-				BindArtificialModuleToBinary(
-					CreateArtificialModule(
-						"AutoStartupModulesListGetter", AppBinaries[0].Config.IntermediateDirectory,
-						new FileItem[] { FileItem.GetItemByPath(SourceFilename) }, new string[] { }
-					), AppBinaries[0]);
-			}
-		}
-
-		/// <summary>
-		/// Creates a getter that outputs auto-startup module list.
-		/// </summary>
-		/// <param name="AutoStartupModuleEnumerable">A list of auto-startup modules.</param>
-		/// <param name="DestPath">Destination path for the file.</param>
-		/// <returns>Created source file path.</returns>
-		private string CreateAutoGeneratedModulesListGetterCPPFile(IEnumerable<string> AutoStartupModuleEnumerable, string DestPath)
-		{
-			var BaseFileName = "AutoStartupModuleListGetter";
-
-			var ModuleDir = Path.Combine(DestPath, BaseFileName);
-
-			if (!Directory.Exists(ModuleDir))
-			{
-				Directory.CreateDirectory(ModuleDir);
-			}
-
-			var HeaderName = BaseFileName + ".h";
-			var HeaderPath = Path.Combine(ModuleDir, HeaderName);
-
-			if(!File.Exists(HeaderPath))
-			{
-				File.WriteAllText(HeaderPath, "");
-			}
-
-			string SourceFilename = Path.Combine(ModuleDir, BaseFileName + ".generated.cpp");
-
-			var SourceLines = new List<string>();
-
-			SourceLines.Add(string.Format("#include \"{0}\"", HeaderName));
-			SourceLines.Add("");
-			SourceLines.Add("const char* EnumAutoStartupModuleName(int Index)");
-			SourceLines.Add("{");
-
-			SourceLines.Add("\tconst char* Names[] = {");
-
-			foreach (var AutoStartupModuleName in AutoStartupModuleEnumerable.OrderBy(ModuleName => ModuleName))
-			{
-				SourceLines.Add(string.Format("\t\t\"{0}\",", AutoStartupModuleName));
-			}
-
-			SourceLines.Add("\t\tnullptr");
-			SourceLines.Add("\t};");
-			SourceLines.Add("");
-			SourceLines.Add("\treturn Names[Index];");
-
-			SourceLines.Add("}");
-			SourceLines.Add("");
-
-			string LineEnd = Environment.OSVersion.Platform == PlatformID.Unix ? "\n" : "\r\n";
-			if (!File.Exists(SourceFilename) || File.ReadAllText(SourceFilename).Trim() != string.Join(LineEnd, SourceLines).Trim())
-			{
-				File.WriteAllLines(SourceFilename, SourceLines);
-			}
-
-			return SourceFilename;
-		}
-
 		/// <summary>
 		/// All non-program monolithic binaries implicitly depend on all static plugin libraries so they are always linked appropriately
 		/// In order to do this, we create a new module here with a cpp file we emit that invokes an empty function in each library.
@@ -2137,7 +2005,6 @@ namespace UnrealBuildTool
 				InUseRTTI: false,
 				InEnableBufferSecurityChecks: true,
 				InFasterWithoutUnity: true,
-				InIsAutoStartupModule: false,
 				InMinFilesUsingPrecompiledHeaderOverride: 0,
 				InEnableExceptions: false,
 				bInBuildSourceFiles: true);
@@ -3081,7 +2948,6 @@ namespace UnrealBuildTool
 							InUseRTTI: RulesObject.bUseRTTI,
 							InEnableBufferSecurityChecks: RulesObject.bEnableBufferSecurityChecks,
 							InFasterWithoutUnity: RulesObject.bFasterWithoutUnity,
-							InIsAutoStartupModule: RulesObject.bIsAutoStartupModule,
 							InMinFilesUsingPrecompiledHeaderOverride: RulesObject.MinFilesUsingPrecompiledHeaderOverride,
 							InEnableExceptions: RulesObject.bEnableExceptions,
 							bInBuildSourceFiles: bBuildSourceFiles
@@ -3122,7 +2988,6 @@ namespace UnrealBuildTool
 							InUseRTTI: RulesObject.bUseRTTI,
 							InEnableBufferSecurityChecks: RulesObject.bEnableBufferSecurityChecks,
 							InFasterWithoutUnity: RulesObject.bFasterWithoutUnity,
-							InIsAutoStartupModule: RulesObject.bIsAutoStartupModule,
 							InMinFilesUsingPrecompiledHeaderOverride: RulesObject.MinFilesUsingPrecompiledHeaderOverride,
 							InEnableExceptions: RulesObject.bEnableExceptions,
 							bInBuildSourceFiles : bBuildSourceFiles
diff --git a/Engine/Source/Programs/UnrealBuildTool/System/RulesCompiler.cs b/Engine/Source/Programs/UnrealBuildTool/System/RulesCompiler.cs
--- a/Engine/Source/Programs/UnrealBuildTool/System/RulesCompiler.cs
+++ b/Engine/Source/Programs/UnrealBuildTool/System/RulesCompiler.cs
@@ -155,9 +155,6 @@ namespace UnrealBuildTool
 		/** If true and unity builds are enabled, this module will build without unity. */
 		public bool bFasterWithoutUnity = false;
 
-		/** If true then the engine will call it's StartupModule at engine initialization automatically. */
-		public bool bIsAutoStartupModule = false;
-
 		/** Overrides BuildConfiguration.MinFilesUsingPrecompiledHeader if non-zero. */
 		public int MinFilesUsingPrecompiledHeaderOverride = 0;
 
diff --git a/Engine/Source/Runtime/Core/Private/Modules/ModuleManager.cpp b/Engine/Source/Runtime/Core/Private/Modules/ModuleManager.cpp
--- a/Engine/Source/Runtime/Core/Private/Modules/ModuleManager.cpp
+++ b/Engine/Source/Runtime/Core/Private/Modules/ModuleManager.cpp
@@ -986,28 +986,3 @@ bool FModuleManager::DoesLoadedModuleHaveUObjects( const FName ModuleName )
 
 	return false;
 }
-
-const ANSICHAR* (*GEnumAutoStartupModuleName)(int32 Index) = nullptr;
-
-void FModuleManager::InitializeAutoStartupModules()
-{
-	TArray<FString> Modules;
-	GetAutoStartupModuleList(Modules);
-
-	for (auto ModuleName : Modules)
-	{
-		LoadModuleChecked<IModuleInterface>(*ModuleName);
-	}
-}
-
-void FModuleManager::GetAutoStartupModuleList(TArray<FString>& OutModules) const
-{
-	checkf(GEnumAutoStartupModuleName != nullptr, TEXT("This callback should be set by the target binary."));
-
-	const ANSICHAR* NamePtr = nullptr;
-	int32 Index = 0;
-	while ((NamePtr = GEnumAutoStartupModuleName(Index++)) != nullptr)
-	{
-		OutModules.Add(FString(NamePtr));
-	}
-}
diff --git a/Engine/Source/Runtime/Core/Public/Modules/ModuleManager.h b/Engine/Source/Runtime/Core/Public/Modules/ModuleManager.h
--- a/Engine/Source/Runtime/Core/Public/Modules/ModuleManager.h
+++ b/Engine/Source/Runtime/Core/Public/Modules/ModuleManager.h
@@ -437,11 +437,6 @@ public:
 
 	virtual bool Exec( UWorld* Inworld, const TCHAR* Cmd, FOutputDevice& Ar ) override;
 
-	/**
-	 * Calls for each auto-startup module its StartupModule function.
-	 */
-	void InitializeAutoStartupModules();
-
 protected:
 
 	/**
@@ -515,13 +510,6 @@ private:
 	/** Finds modules matching a given name wildcard within a given directory. */
 	void FindModulePathsInDirectory(const FString &DirectoryName, bool bIsGameDirectory, const TCHAR *NamePattern, TMap<FName, FString> &OutModulePaths) const;
 
-	/**
-	 * Loads array of module names that needs to be auto-loaded.
-	 *
-	 * @param OutModules Output array to fill with module names.
-	 */
-	void GetAutoStartupModuleList(TArray<FString>& OutModules) const;
-
 private:
 
 	/** Map of all modules.  Maps the case-insensitive module name to information about that module, loaded or not. */
@@ -626,18 +614,6 @@ class FDefaultGameModuleImpl
 	}
 };
 
-/**
- * Callback to retrieve auto-startup modules list.
- *
- * It should return string ptr to auto-startup module for consecutive indices
- * until the list is finished. Then this function should return nullptr.
- *
- * @param Index Auto-startup module index to return.
- *
- * @returns Auto-startup module name.
- */
-extern CORE_API const ANSICHAR* (*GEnumAutoStartupModuleName)(int32 Index);
-
 /**
  * Module implementation boilerplate for regular modules.
  *
@@ -725,16 +701,6 @@ extern CORE_API const ANSICHAR* (*GEnumAutoStartupModuleName)(int32 Index);
 	#define IMPLEMENT_DEBUGGAME()
 #endif 
 
-#define REGISTER_AUTO_STARTUP_MODULE_LIST_GETTER() \
-	struct FRegisterAutoStartupModuleListEnumerator \
-	{ \
-		FRegisterAutoStartupModuleListEnumerator() \
-		{ \
-			extern const ANSICHAR* EnumAutoStartupModuleName(int Index); \
-			GEnumAutoStartupModuleName = &EnumAutoStartupModuleName; \
-		} \
-	} RegisterAutoStartupModuleListEnumerator;
-
 #if IS_PROGRAM
 
 	#if IS_MONOLITHIC
@@ -745,8 +711,7 @@ extern CORE_API const ANSICHAR* (*GEnumAutoStartupModuleName)(int32 Index);
 			IMPLEMENT_FOREIGN_ENGINE_DIR() \
 			IMPLEMENT_GAME_MODULE(FDefaultGameModuleImpl, ModuleName) \
 			PER_MODULE_BOILERPLATE \
-			FEngineLoop GEngineLoop; \
-			REGISTER_AUTO_STARTUP_MODULE_LIST_GETTER();
+			FEngineLoop GEngineLoop;
 
 	#else		
 
@@ -761,8 +726,7 @@ extern CORE_API const ANSICHAR* (*GEnumAutoStartupModuleName)(int32 Index);
 			} AutoSet##ModuleName; \
 			PER_MODULE_BOILERPLATE \
 			PER_MODULE_BOILERPLATE_ANYLINK(FDefaultGameModuleImpl, ModuleName) \
-			FEngineLoop GEngineLoop; \
-			REGISTER_AUTO_STARTUP_MODULE_LIST_GETTER();
+			FEngineLoop GEngineLoop;
 	#endif
 
 #else
@@ -786,8 +750,7 @@ extern CORE_API const ANSICHAR* (*GEnumAutoStartupModuleName)(int32 Index);
 			{ \
 				extern void UELinkerFixups(); \
 				UELinkerFixups(); \
-			} \
-			REGISTER_AUTO_STARTUP_MODULE_LIST_GETTER();
+			}
 
 	#else	//PLATFORM_DESKTOP
 
@@ -799,8 +762,7 @@ extern CORE_API const ANSICHAR* (*GEnumAutoStartupModuleName)(int32 Index);
 			IMPLEMENT_FOREIGN_ENGINE_DIR() \
 			IMPLEMENT_GAME_MODULE( ModuleImplClass, ModuleName ) \
 			/* Implement the GIsGameAgnosticExe variable (See Core.h). */ \
-			bool GIsGameAgnosticExe = false; \
-			REGISTER_AUTO_STARTUP_MODULE_LIST_GETTER();
+			bool GIsGameAgnosticExe = false;
 
 	#endif	//PLATFORM_DESKTOP
 
@@ -808,8 +770,7 @@ extern CORE_API const ANSICHAR* (*GEnumAutoStartupModuleName)(int32 Index);
 
 	#define IMPLEMENT_PRIMARY_GAME_MODULE( ModuleImplClass, ModuleName, GameName ) \
 		/* Nothing special to do for modular builds.  The game name will be set via the command-line */ \
-		IMPLEMENT_GAME_MODULE( ModuleImplClass, ModuleName ) \
-		REGISTER_AUTO_STARTUP_MODULE_LIST_GETTER();
+		IMPLEMENT_GAME_MODULE( ModuleImplClass, ModuleName )
 #endif	//IS_MONOLITHIC
 
 #endif
diff --git a/Engine/Source/Runtime/Landscape/Landscape.Build.cs b/Engine/Source/Runtime/Landscape/Landscape.Build.cs
--- a/Engine/Source/Runtime/Landscape/Landscape.Build.cs
+++ b/Engine/Source/Runtime/Landscape/Landscape.Build.cs
@@ -7,8 +7,6 @@ public class Landscape : ModuleRules
 {
 	public Landscape(TargetInfo Target)
 	{
-		bIsAutoStartupModule = true;
-
 		PrivateIncludePaths.AddRange(
 			new string[] {
                 "Runtime/Engine/Private",
diff --git a/Engine/Source/Runtime/Launch/Private/LaunchEngineLoop.cpp b/Engine/Source/Runtime/Launch/Private/LaunchEngineLoop.cpp
--- a/Engine/Source/Runtime/Launch/Private/LaunchEngineLoop.cpp
+++ b/Engine/Source/Runtime/Launch/Private/LaunchEngineLoop.cpp
@@ -1527,8 +1527,6 @@ int32 FEngineLoop::PreInit( const TCHAR* CmdLine )
 			// Execute the commandlet.
 			double CommandletExecutionStartTime = FPlatformTime::Seconds();
 
-			FModuleManager::Get().InitializeAutoStartupModules();
-
 			// Commandlets don't always handle -run= properly in the commandline so we'll provide them
 			// with a custom version that doesn't have it.
 			Commandlet->ParseParms( CommandletCommandLine );
@@ -1693,7 +1691,6 @@ int32 FEngineLoop::PreInit( const TCHAR* CmdLine )
 
 
 #else // WITH_ENGINE
-	FModuleManager::Get().InitializeAutoStartupModules();
 	EndInitTextLocalization();
 	FPlatformMisc::PlatformPostInit();
 #endif // WITH_ENGINE
@@ -1856,8 +1853,6 @@ bool FEngineLoop::LoadStartupModules()
 		return false;
 	}
 
-	FModuleManager::Get().InitializeAutoStartupModules();
-
 	return true;
 }
 
diff --git a/Engine/Source/Runtime/UE4Game/Private/UE4Game.cpp b/Engine/Source/Runtime/UE4Game/Private/UE4Game.cpp
--- a/Engine/Source/Runtime/UE4Game/Private/UE4Game.cpp
+++ b/Engine/Source/Runtime/UE4Game/Private/UE4Game.cpp
@@ -12,5 +12,3 @@ TCHAR GInternalGameName[64] = TEXT("");
 IMPLEMENT_DEBUGGAME()
 IMPLEMENT_FOREIGN_ENGINE_DIR()
 #endif
-
-REGISTER_AUTO_STARTUP_MODULE_LIST_GETTER();
