addition of build settings to ini files
part of UEPLAT-114
#ue4

diff --git a/Engine/Source/Programs/AutomationTool/Scripts/BuildProjectCommand.Automation.cs b/Engine/Source/Programs/AutomationTool/Scripts/BuildProjectCommand.Automation.cs
--- a/Engine/Source/Programs/AutomationTool/Scripts/BuildProjectCommand.Automation.cs
+++ b/Engine/Source/Programs/AutomationTool/Scripts/BuildProjectCommand.Automation.cs
@@ -73,7 +73,7 @@ public partial class Project : CommandUtils
 				foreach (var ClientPlatform in Params.ClientTargetPlatforms)
 				{
 					CrashReportPlatforms.Add(ClientPlatform);
-					Agenda.AddTargets(Params.ClientCookedTargets.ToArray(), ClientPlatform, BuildConfig, Params.CodeBasedUprojectPath);
+					Agenda.AddTargets(Params.ClientCookedTargets.ToArray(), ClientPlatform, BuildConfig, Params.RawProjectPath);
 				}
 			}
 		}
@@ -84,7 +84,7 @@ public partial class Project : CommandUtils
 				foreach (var ServerPlatform in Params.ServerTargetPlatforms)
 				{
 					CrashReportPlatforms.Add(ServerPlatform);
-					Agenda.AddTargets(Params.ServerCookedTargets.ToArray(), ServerPlatform, BuildConfig, Params.CodeBasedUprojectPath);
+					Agenda.AddTargets(Params.ServerCookedTargets.ToArray(), ServerPlatform, BuildConfig, Params.RawProjectPath);
 				}
 			}
 		}
diff --git a/Engine/Source/Programs/UnrealBuildTool/Android/AndroidToolChain.cs b/Engine/Source/Programs/UnrealBuildTool/Android/AndroidToolChain.cs
--- a/Engine/Source/Programs/UnrealBuildTool/Android/AndroidToolChain.cs
+++ b/Engine/Source/Programs/UnrealBuildTool/Android/AndroidToolChain.cs
@@ -947,5 +947,10 @@ namespace UnrealBuildTool
 				Log.TraceInformation(Line.Data);
 			}
 		}
+
+		public override UnrealTargetPlatform GetPlatform()
+		{
+			return UnrealTargetPlatform.Android;
+		}
 	};
 }
diff --git a/Engine/Source/Programs/UnrealBuildTool/HTML5/HTML5ToolChain.cs b/Engine/Source/Programs/UnrealBuildTool/HTML5/HTML5ToolChain.cs
--- a/Engine/Source/Programs/UnrealBuildTool/HTML5/HTML5ToolChain.cs
+++ b/Engine/Source/Programs/UnrealBuildTool/HTML5/HTML5ToolChain.cs
@@ -506,5 +506,10 @@ namespace UnrealBuildTool
             // make sure we don't add mem files more than once. 
             Manifest.AddBinaryNames(MemFile, null);
         }
+
+		public override UnrealTargetPlatform GetPlatform()
+		{
+			return UnrealTargetPlatform.HTML5;
+		}
 	};
 }
diff --git a/Engine/Source/Programs/UnrealBuildTool/IOS/IOSToolChain.cs b/Engine/Source/Programs/UnrealBuildTool/IOS/IOSToolChain.cs
--- a/Engine/Source/Programs/UnrealBuildTool/IOS/IOSToolChain.cs
+++ b/Engine/Source/Programs/UnrealBuildTool/IOS/IOSToolChain.cs
@@ -89,8 +89,10 @@ namespace UnrealBuildTool
 			BaseSDKDirSim = XcodeDeveloperDir + "Platforms/iPhoneSimulator.platform/Developer/SDKs";
 		}
 
-		public static void ParseProjectSettings()
+		public override void ParseProjectSettings()
 		{
+			base.ParseProjectSettings();
+
 			// look in ini settings for what platforms to compile for
 			ConfigCacheIni Ini = new ConfigCacheIni(UnrealTargetPlatform.IOS, "Engine", UnrealBuildTool.GetUProjectPath());
 			string MinVersion = "IOS_6";
@@ -190,7 +192,6 @@ namespace UnrealBuildTool
 		{
 			base.SetUpGlobalEnvironment();
 
-			ParseProjectSettings();
 			ParseArchitectures();
 
 			if (IOSSDKVersion == "latest")
@@ -1454,6 +1455,11 @@ namespace UnrealBuildTool
 			return RunTimeIOSDevices;
 		}
 
+		public override UnrealTargetPlatform GetPlatform()
+		{
+			return UnrealTargetPlatform.IOS;
+		}
+
 		public static int RunExecutableAndWait( string ExeName, string ArgumentList, out string StdOutResults )
 		{
 			// Create the process
diff --git a/Engine/Source/Programs/UnrealBuildTool/IOS/UEDeployIOS.cs b/Engine/Source/Programs/UnrealBuildTool/IOS/UEDeployIOS.cs
--- a/Engine/Source/Programs/UnrealBuildTool/IOS/UEDeployIOS.cs
+++ b/Engine/Source/Programs/UnrealBuildTool/IOS/UEDeployIOS.cs
@@ -387,7 +387,7 @@ namespace UnrealBuildTool.IOS
 		public override bool PrepTargetForDeployment(UEBuildTarget InTarget)
 		{
 			string GameName = InTarget.AppName;
-			string BuildPath = InTarget.ProjectDirectory + "/Binaries/IOS";
+			string BuildPath = (GameName == "UE4Game" ? "../../Engine" : InTarget.ProjectDirectory) + "/Binaries/IOS";
 			string ProjectDirectory = InTarget.ProjectDirectory;
 
 			if (BuildHostPlatform.Current.Platform == UnrealTargetPlatform.Mac && Environment.GetEnvironmentVariable("UBT_NO_POST_DEPLOY") != "true")
diff --git a/Engine/Source/Programs/UnrealBuildTool/Linux/LinuxToolChain.cs b/Engine/Source/Programs/UnrealBuildTool/Linux/LinuxToolChain.cs
--- a/Engine/Source/Programs/UnrealBuildTool/Linux/LinuxToolChain.cs
+++ b/Engine/Source/Programs/UnrealBuildTool/Linux/LinuxToolChain.cs
@@ -1090,5 +1090,10 @@ namespace UnrealBuildTool
 
             return OutputFiles;
         }
-    }
+
+		public override UnrealTargetPlatform GetPlatform()
+		{
+			return UnrealTargetPlatform.Linux;
+		}
+	}
 }
diff --git a/Engine/Source/Programs/UnrealBuildTool/Mac/MacToolChain.cs b/Engine/Source/Programs/UnrealBuildTool/Mac/MacToolChain.cs
--- a/Engine/Source/Programs/UnrealBuildTool/Mac/MacToolChain.cs
+++ b/Engine/Source/Programs/UnrealBuildTool/Mac/MacToolChain.cs
@@ -1349,5 +1349,10 @@ namespace UnrealBuildTool
 
 			return OutputFiles;
 		}
+
+		public override UnrealTargetPlatform GetPlatform()
+		{
+			return UnrealTargetPlatform.Mac;
+		}
 	};
 }
diff --git a/Engine/Source/Programs/UnrealBuildTool/ToolChain/UEToolChain.cs b/Engine/Source/Programs/UnrealBuildTool/ToolChain/UEToolChain.cs
--- a/Engine/Source/Programs/UnrealBuildTool/ToolChain/UEToolChain.cs
+++ b/Engine/Source/Programs/UnrealBuildTool/ToolChain/UEToolChain.cs
@@ -47,6 +47,8 @@ namespace UnrealBuildTool
 		string GetPlatformVersion();
 
 		string GetPlatformDevices();
+
+		UnrealTargetPlatform GetPlatform();
 	}
 
 	public abstract class UEToolChain : IUEToolChain
@@ -161,6 +163,89 @@ namespace UnrealBuildTool
 
 		public virtual void SetUpGlobalEnvironment()
 		{
+			ParseProjectSettings();
+		}
+
+		public virtual void ParseProjectSettings()
+		{
+			ConfigCacheIni Ini = new ConfigCacheIni(GetPlatform(), "Engine", UnrealBuildTool.GetUProjectPath());
+			bool bValue = UEBuildConfiguration.bCompileAPEX;
+			if (Ini.GetBool("/Script/BuildSettings.BuildSettings", "bCompileApex", out bValue))
+			{
+				UEBuildConfiguration.bCompileAPEX = bValue;
+			}
+
+			bValue = UEBuildConfiguration.bCompileBox2D;
+			if (Ini.GetBool("/Script/BuildSettings.BuildSettings", "bCompileBox2D", out bValue))
+			{
+				UEBuildConfiguration.bCompileBox2D = bValue;
+			}
+
+			bValue = UEBuildConfiguration.bCompileICU;
+			if (Ini.GetBool("/Script/BuildSettings.BuildSettings", "bCompileICU", out bValue))
+			{
+				UEBuildConfiguration.bCompileICU = bValue;
+			}
+
+			bValue = UEBuildConfiguration.bCompileSimplygon;
+			if (Ini.GetBool("/Script/BuildSettings.BuildSettings", "bCompileSimplygon", out bValue))
+			{
+				UEBuildConfiguration.bCompileSimplygon = bValue;
+			}
+
+			bValue = UEBuildConfiguration.bCompileLeanAndMeanUE;
+			if (Ini.GetBool("/Script/BuildSettings.BuildSettings", "bCompileLeanAndMeanUE", out bValue))
+			{
+				UEBuildConfiguration.bCompileLeanAndMeanUE = bValue;
+			}
+
+			bValue = UEBuildConfiguration.bIncludeADO;
+			if (Ini.GetBool("/Script/BuildSettings.BuildSettings", "bIncludeADO", out bValue))
+			{
+				UEBuildConfiguration.bIncludeADO = bValue;
+			}
+
+			bValue = UEBuildConfiguration.bCompileRecast;
+			if (Ini.GetBool("/Script/BuildSettings.BuildSettings", "bCompileRecast", out bValue))
+			{
+				UEBuildConfiguration.bCompileRecast = bValue;
+			}
+
+			bValue = UEBuildConfiguration.bCompileSpeedTree;
+			if (Ini.GetBool("/Script/BuildSettings.BuildSettings", "bCompileSpeedTree", out bValue))
+			{
+				UEBuildConfiguration.bCompileSpeedTree = bValue;
+			}
+
+			bValue = UEBuildConfiguration.bCompileWithPluginSupport;
+			if (Ini.GetBool("/Script/BuildSettings.BuildSettings", "bCompileWithPluginSupport", out bValue))
+			{
+				UEBuildConfiguration.bCompileWithPluginSupport = bValue;
+			}
+
+			bValue = UEBuildConfiguration.bCompilePhysXVehicle;
+			if (Ini.GetBool("/Script/BuildSettings.BuildSettings", "bCompilePhysXVehicle", out bValue))
+			{
+				UEBuildConfiguration.bCompilePhysXVehicle = bValue;
+			}
+
+			bValue = UEBuildConfiguration.bCompileFreeType;
+			if (Ini.GetBool("/Script/BuildSettings.BuildSettings", "bCompileFreeType", out bValue))
+			{
+				UEBuildConfiguration.bCompileFreeType = bValue;
+			}
+
+			bValue = UEBuildConfiguration.bCompileForSize;
+			if (Ini.GetBool("/Script/BuildSettings.BuildSettings", "bCompileForSize", out bValue))
+			{
+				UEBuildConfiguration.bCompileForSize = bValue;
+			}
+
+			bValue = UEBuildConfiguration.bCompileCEF3;
+			if (Ini.GetBool("/Script/BuildSettings.BuildSettings", "bCompileCEF3", out bValue))
+			{
+				UEBuildConfiguration.bCompileCEF3 = bValue;
+			}
 		}
 
 		protected void RunUnrealHeaderToolIfNeeded()
@@ -224,5 +309,10 @@ namespace UnrealBuildTool
 		{
 			return "";
 		}
+
+		public virtual UnrealTargetPlatform GetPlatform()
+		{
+			return UnrealTargetPlatform.Unknown;
+		}
 	};
 }
\ No newline at end of file
diff --git a/Engine/Source/Programs/UnrealBuildTool/WinRT/WinRTToolChain.cs b/Engine/Source/Programs/UnrealBuildTool/WinRT/WinRTToolChain.cs
--- a/Engine/Source/Programs/UnrealBuildTool/WinRT/WinRTToolChain.cs
+++ b/Engine/Source/Programs/UnrealBuildTool/WinRT/WinRTToolChain.cs
@@ -959,5 +959,10 @@ namespace UnrealBuildTool
 				bAreEnvironmentVariablesAlreadyInitialized = true;
 			}
 		}
+
+		public override UnrealTargetPlatform GetPlatform()
+		{
+			return UnrealTargetPlatform.WinRT;
+		}
 	};
 }
