Initial implementation of UnrealCodeAnalyzer.

diff --git a/Engine/Source/Programs/UnrealBuildTool/Configuration/BuildConfiguration.cs b/Engine/Source/Programs/UnrealBuildTool/Configuration/BuildConfiguration.cs
--- a/Engine/Source/Programs/UnrealBuildTool/Configuration/BuildConfiguration.cs
+++ b/Engine/Source/Programs/UnrealBuildTool/Configuration/BuildConfiguration.cs
@@ -402,6 +402,24 @@ namespace UnrealBuildTool
 		[XmlConfig]
 		public static string DistccExecutablesPath;
 
+		/// <summary>
+		/// Run UnrealCodeAnalyzer instead of regular compilation.
+		/// </summary>
+		[XmlConfig]
+		public static bool bRunUnrealCodeAnalyzer;
+
+		/// <summary>
+		/// Percentage of cpp files in module where header must be included to get included in PCH.
+		/// </summary>
+		[XmlConfig]
+		public static float UCAUsageThreshold;
+
+		/// <summary>
+		/// If specified, UnrealCodeAnalyzer will analyze only specific module.
+		/// </summary>
+		[XmlConfig]
+		public static string UCAModuleToAnalyze;
+
 		/// <summary>
 		/// Sets the configuration back to defaults.
 		/// </summary>
@@ -525,6 +543,12 @@ namespace UnrealBuildTool
 			// The normal Posix place to install distcc/dmucs would be /usr/local so start there
 			DistccExecutablesPath = "/usr/local/bin";
 
+			// By default compile code instead of analyzing it.
+			bRunUnrealCodeAnalyzer = false;
+
+			// If header is included in 30% or more cpp files in module it should be included in PCH.
+			UCAUsageThreshold = 0.3f;
+
 			// The default for normal Mac users should be to use DistCode which installs as an Xcode plugin and provides dynamic host management
 			if (BuildHostPlatform.Current.Platform == UnrealTargetPlatform.Mac)
 			{
diff --git a/Engine/Source/Programs/UnrealBuildTool/Configuration/UEBuildBinary.cs b/Engine/Source/Programs/UnrealBuildTool/Configuration/UEBuildBinary.cs
--- a/Engine/Source/Programs/UnrealBuildTool/Configuration/UEBuildBinary.cs
+++ b/Engine/Source/Programs/UnrealBuildTool/Configuration/UEBuildBinary.cs
@@ -7,6 +7,7 @@ using System.Diagnostics;
 using System.IO;
 using System.Linq;
 using System.Xml;
+using System.Globalization;
 
 namespace UnrealBuildTool
 {	
@@ -474,6 +475,26 @@ namespace UnrealBuildTool
 			bIncludeDependentLibrariesInLibrary = bInIncludeDependentLibrariesInLibrary;
 		}
 
+		bool IsBuildingDll(UEBuildBinaryType Type)
+		{
+			if (BuildConfiguration.bRunUnrealCodeAnalyzer)
+			{
+				return false;
+			}
+
+			return Type == UEBuildBinaryType.DynamicLinkLibrary;
+		}
+
+		bool IsBuildingLibrary(UEBuildBinaryType Type)
+		{
+			if (BuildConfiguration.bRunUnrealCodeAnalyzer)
+			{
+				return false;
+			}
+
+			return Type == UEBuildBinaryType.StaticLibrary;
+		}
+
 		/// <summary>
 		/// Builds the binary.
 		/// </summary>
@@ -482,33 +503,97 @@ namespace UnrealBuildTool
 		/// <returns></returns>
 		public override IEnumerable<FileItem> Build(IUEToolChain TargetToolChain, CPPEnvironment CompileEnvironment, LinkEnvironment LinkEnvironment)
 		{
+			// UnrealCodeAnalyzer produces output files only for a specific module.
+			if (BuildConfiguration.bRunUnrealCodeAnalyzer && !(ModuleNames.Contains(BuildConfiguration.UCAModuleToAnalyze)))
+			{
+				return new List<FileItem>();
+			}
+
+			// Setup linking environment.
+			var BinaryLinkEnvironment = SetupBinaryLinkEnvironment(LinkEnvironment, CompileEnvironment);
 
-			// Determine the type of binary we're linking.
-			switch (Config.Type)
+			// Return linked files.
+			return SetupOutputFiles(TargetToolChain, ref BinaryLinkEnvironment);
+		}
+
+		/// <summary>
+		/// Called to allow the binary to modify the link environment of a different binary containing 
+		/// a module that depends on a module in this binary.
+		/// </summary>
+		/// <param name="DependentLinkEnvironment">The link environment of the dependency</param>
+		public override void SetupDependentLinkEnvironment(ref LinkEnvironment DependentLinkEnvironment)
+		{
+			foreach (string OutputFilePath in Config.OutputFilePaths)
 			{
-				case UEBuildBinaryType.DynamicLinkLibrary:
-					CompileEnvironment.Config.bIsBuildingDLL = true;
-					CompileEnvironment.Config.bIsBuildingLibrary = false;
-					break;
-				case UEBuildBinaryType.StaticLibrary:
-					CompileEnvironment.Config.bIsBuildingDLL = false;
-					CompileEnvironment.Config.bIsBuildingLibrary = true;
-					break;
-				default:
-					CompileEnvironment.Config.bIsBuildingDLL = false;
-					CompileEnvironment.Config.bIsBuildingLibrary = false;
-					break;
-			};
+				string LibraryFileName;
+				if (Config.Type == UEBuildBinaryType.StaticLibrary
+					|| DependentLinkEnvironment.Config.Target.Platform == CPPTargetPlatform.Mac
+					|| DependentLinkEnvironment.Config.Target.Platform == CPPTargetPlatform.Linux)
+				{
+					LibraryFileName = OutputFilePath;
+				}
+				else
+				{
+					LibraryFileName = Path.Combine(Config.IntermediateDirectory, Path.GetFileNameWithoutExtension(OutputFilePath) + ".lib");
+				}
+				DependentLinkEnvironment.Config.AdditionalLibraries.Add(LibraryFileName);
+			}
+		}
 
-			var OutputFiles = new List<FileItem>();
+		/// <summary>
+		/// Called to allow the binary to to determine if it matches the Only module "short module name".
+		/// </summary>
+		/// <param name="OnlyModules"></param>
+		/// <returns>The OnlyModule if found, null if not</returns>
+		public override OnlyModule FindOnlyModule(List<OnlyModule> OnlyModules)
+		{
+			foreach (var ModuleName in ModuleNames)
+			{
+				foreach (var OnlyModule in OnlyModules)
+				{
+					if (OnlyModule.OnlyModuleName.ToLower() == ModuleName.ToLower())
+					{
+						return OnlyModule;
+					}
+				}
+			}
+			return null;
+		}
 
-			var BinaryCompileEnvironment = CompileEnvironment.DeepCopy();
-			var BinaryLinkEnvironment = LinkEnvironment.DeepCopy();
+		public override List<UEBuildModule> FindGameModules()
+		{
+			var GameModules = new List<UEBuildModule>();
+			foreach (var ModuleName in ModuleNames)
+			{
+				UEBuildModule Module = Target.GetModuleByName(ModuleName);
+				if (!Utils.IsFileUnderDirectory(Module.ModuleDirectory, BuildConfiguration.RelativeEnginePath))
+				{
+					GameModules.Add(Module);
+				}
+			}
+			return GameModules;
+		}
 
-			// Process each module that is linked into the binary.
-			var BinaryDependencies = new List<UEBuildBinary>();
+		// Object interface.
+
+		/// <summary>
+		/// ToString implementation
+		/// </summary>
+		/// <returns>Returns the OutputFilePath for this binary</returns>
+		public override string ToString()
+		{
+			return Config.OutputFilePath;
+		}
+
+		private LinkEnvironment SetupBinaryLinkEnvironment(LinkEnvironment LinkEnvironment, CPPEnvironment CompileEnvironment)
+		{
+			var BinaryLinkEnvironment = LinkEnvironment.DeepCopy();
 			var LinkEnvironmentVisitedModules = new Dictionary<UEBuildModule, bool>();
+			var BinaryDependencies = new List<UEBuildBinary>();
+			CompileEnvironment.Config.bIsBuildingDLL = IsBuildingDll(Config.Type);
+			CompileEnvironment.Config.bIsBuildingLibrary = IsBuildingLibrary(Config.Type);
 
+			var BinaryCompileEnvironment = CompileEnvironment.DeepCopy();
 			// @Hack: This to prevent UHT from listing CoreUObject.generated.cpp as its dependency.
 			// We flag the compile environment when we build UHT so that we don't need to check
 			// this for each file when generating their dependencies.
@@ -533,26 +618,30 @@ namespace UnrealBuildTool
 				// NOTE: Because of 'Shared PCHs', in monolithic builds the same PCH file may appear as a link input
 				// multiple times for a single binary.  We'll check for that here, and only add it once.  This avoids
 				// a linker warning about redundant .obj files. 
-				foreach( var LinkInputFile in LinkInputFiles )
+				foreach (var LinkInputFile in LinkInputFiles)
 				{
-					if( !BinaryLinkEnvironment.InputFiles.Contains( LinkInputFile ))
+					if (!BinaryLinkEnvironment.InputFiles.Contains(LinkInputFile))
 					{
-						BinaryLinkEnvironment.InputFiles.Add( LinkInputFile );
+						BinaryLinkEnvironment.InputFiles.Add(LinkInputFile);
 					}
 				}
-				
-				// Allow the module to modify the link environment for the binary.
-				Module.SetupPrivateLinkEnvironment(ref BinaryLinkEnvironment,ref BinaryDependencies,ref LinkEnvironmentVisitedModules);
+
+				if (!BuildConfiguration.bRunUnrealCodeAnalyzer)
+				{
+					// Allow the module to modify the link environment for the binary.
+					Module.SetupPrivateLinkEnvironment(ref BinaryLinkEnvironment, ref BinaryDependencies, ref LinkEnvironmentVisitedModules);
+				}
 			}
 
+
 			// Remove the default resource file on Windows (PCLaunch.rc) if the user has specified their own
-			if(BinaryLinkEnvironment.InputFiles.Select(Item => Path.GetFileName(Item.AbsolutePath).ToLower()).Any(Name => Name.EndsWith(".res") && !Name.EndsWith(".inl.res") && Name != "pclaunch.rc.res"))
+			if (BinaryLinkEnvironment.InputFiles.Select(Item => Path.GetFileName(Item.AbsolutePath).ToLower()).Any(Name => Name.EndsWith(".res") && !Name.EndsWith(".inl.res") && Name != "pclaunch.rc.res"))
 			{
 				BinaryLinkEnvironment.InputFiles.RemoveAll(x => Path.GetFileName(x.AbsolutePath).ToLower() == "pclaunch.rc.res");
 			}
 
 			// Allow the binary dependencies to modify the link environment.
-			foreach(var BinaryDependency in BinaryDependencies)
+			foreach (var BinaryDependency in BinaryDependencies)
 			{
 				BinaryDependency.SetupDependentLinkEnvironment(ref BinaryLinkEnvironment);
 			}
@@ -569,147 +658,110 @@ namespace UnrealBuildTool
 			// Put the non-executable output files (PDB, import library, etc) in the same directory as the production
 			BinaryLinkEnvironment.Config.OutputDirectory = Path.GetDirectoryName(Config.OutputFilePaths[0]);
 
-			// Determine the type of binary we're linking.
-			switch (Config.Type)
-			{
-				case UEBuildBinaryType.DynamicLinkLibrary:
-					BinaryLinkEnvironment.Config.bIsBuildingDLL = true;
-					BinaryLinkEnvironment.Config.bIsBuildingLibrary = false;
-					break;
-				case UEBuildBinaryType.StaticLibrary:
-					BinaryLinkEnvironment.Config.bIsBuildingDLL = false;
-					BinaryLinkEnvironment.Config.bIsBuildingLibrary = true;
-					break;
-				default:
-					BinaryLinkEnvironment.Config.bIsBuildingDLL = false;
-					BinaryLinkEnvironment.Config.bIsBuildingLibrary = false;
-					break;
-			};
-
-			if( ProjectFileGenerator.bGenerateProjectFiles )
+			// Setup link output type
+			BinaryLinkEnvironment.Config.bIsBuildingDLL = IsBuildingDll(Config.Type);
+			BinaryLinkEnvironment.Config.bIsBuildingLibrary = IsBuildingLibrary(Config.Type);
+
+			return BinaryLinkEnvironment;
+		}
+
+		private List<FileItem> SetupOutputFiles(IUEToolChain TargetToolChain, ref LinkEnvironment BinaryLinkEnvironment)
+		{
+			// Early exits first
+			if (ProjectFileGenerator.bGenerateProjectFiles)
 			{
 				// We're generating projects.  Since we only need include paths and definitions, there is no need
 				// to go ahead and run through the linking logic.
-				OutputFiles = BinaryLinkEnvironment.InputFiles;
+				return BinaryLinkEnvironment.InputFiles;
 			}
-			else if( BuildConfiguration.bEnableCodeAnalysis )
+
+			if (BuildConfiguration.bEnableCodeAnalysis)
 			{
 				// We're only analyzing code, so we won't actually link any executables.  Instead, our output
 				// files will simply be the .obj files that were compiled during static analysis.
-				OutputFiles = BinaryLinkEnvironment.InputFiles;
+				return BinaryLinkEnvironment.InputFiles;
 			}
-			else
-			{
-				if(bCreateImportLibrarySeparately)
-				{
-					// Mark the link environment as cross-referenced.
-					BinaryLinkEnvironment.Config.bIsCrossReferenced = true;
-
-					if (BinaryLinkEnvironment.Config.Target.Platform != CPPTargetPlatform.Mac && BinaryLinkEnvironment.Config.Target.Platform != CPPTargetPlatform.Linux)
-					{
-						// Create the import library.
-						OutputFiles.AddRange(BinaryLinkEnvironment.LinkExecutable(true));
-					}
-				}
 
-				BinaryLinkEnvironment.Config.bIncludeDependentLibrariesInLibrary = bIncludeDependentLibrariesInLibrary;
-
-				// Link the binary.
-				FileItem[] Executables = BinaryLinkEnvironment.LinkExecutable(false);
-				OutputFiles.AddRange(Executables);
-
-				// Produce additional console app if requested
-				if (BinaryLinkEnvironment.Config.CanProduceAdditionalConsoleApp && UEBuildConfiguration.bBuildEditor)
-				{					
-					// Produce additional binary but link it as a console app
-					var ConsoleAppLinkEvironment = BinaryLinkEnvironment.DeepCopy();
-					ConsoleAppLinkEvironment.Config.bIsBuildingConsoleApplication = true;
-					ConsoleAppLinkEvironment.Config.WindowsEntryPointOverride = "WinMainCRTStartup";		// For WinMain() instead of "main()" for Launch module
-					for (int Index = 0; Index < Config.OutputFilePaths.Length; Index++)
-					{
-						ConsoleAppLinkEvironment.Config.OutputFilePaths[Index] = GetAdditionalConsoleAppPath(ConsoleAppLinkEvironment.Config.OutputFilePaths[Index]);
-					}
+			if (BuildConfiguration.bRunUnrealCodeAnalyzer)
+			{
+				//
+				// Create actions to analyze *.includes files and provide suggestions on how to modify PCH.
+				//
+				return CreateOutputFilesForUCA(BinaryLinkEnvironment);
+			}
 
-					// Link the console app executable
-					OutputFiles.AddRange(ConsoleAppLinkEvironment.LinkExecutable(false));
-				}
+			//
+			// Regular linking action.
+			//
+			var OutputFiles = new List<FileItem>();
+			if (bCreateImportLibrarySeparately)
+			{
+				// Mark the link environment as cross-referenced.
+				BinaryLinkEnvironment.Config.bIsCrossReferenced = true;
 
-				foreach (var Executable in Executables)
+				if (BinaryLinkEnvironment.Config.Target.Platform != CPPTargetPlatform.Mac && BinaryLinkEnvironment.Config.Target.Platform != CPPTargetPlatform.Linux)
 				{
-					OutputFiles.AddRange(TargetToolChain.PostBuild(Executable, BinaryLinkEnvironment));
+					// Create the import library.
+					OutputFiles.AddRange(BinaryLinkEnvironment.LinkExecutable(true));
 				}
 			}
-			
-			return OutputFiles;
-		}
 
-		/// <summary>
-		/// Called to allow the binary to modify the link environment of a different binary containing 
-		/// a module that depends on a module in this binary.
-		/// </summary>
-		/// <param name="DependentLinkEnvironment">The link environment of the dependency</param>
-		public override void SetupDependentLinkEnvironment(ref LinkEnvironment DependentLinkEnvironment)
-		{
-			foreach (string OutputFilePath in Config.OutputFilePaths)
+			BinaryLinkEnvironment.Config.bIncludeDependentLibrariesInLibrary = bIncludeDependentLibrariesInLibrary;
+
+			// Link the binary.
+			FileItem[] Executables = BinaryLinkEnvironment.LinkExecutable(false);
+			OutputFiles.AddRange(Executables);
+
+			// Produce additional console app if requested
+			if (BinaryLinkEnvironment.Config.CanProduceAdditionalConsoleApp && UEBuildConfiguration.bBuildEditor)
 			{
-				string LibraryFileName;
-				if (Config.Type == UEBuildBinaryType.StaticLibrary
-					|| DependentLinkEnvironment.Config.Target.Platform == CPPTargetPlatform.Mac
-					|| DependentLinkEnvironment.Config.Target.Platform == CPPTargetPlatform.Linux)
+				// Produce additional binary but link it as a console app
+				var ConsoleAppLinkEvironment = BinaryLinkEnvironment.DeepCopy();
+				ConsoleAppLinkEvironment.Config.bIsBuildingConsoleApplication = true;
+				ConsoleAppLinkEvironment.Config.WindowsEntryPointOverride = "WinMainCRTStartup";		// For WinMain() instead of "main()" for Launch module
+				for (int Index = 0; Index < Config.OutputFilePaths.Length; Index++)
 				{
-					LibraryFileName = OutputFilePath;
+					ConsoleAppLinkEvironment.Config.OutputFilePaths[Index] = GetAdditionalConsoleAppPath(ConsoleAppLinkEvironment.Config.OutputFilePaths[Index]);
 				}
-				else
-				{
-					LibraryFileName = Path.Combine(Config.IntermediateDirectory, Path.GetFileNameWithoutExtension(OutputFilePath) + ".lib");
-				}
-				DependentLinkEnvironment.Config.AdditionalLibraries.Add(LibraryFileName);
+
+				// Link the console app executable
+				OutputFiles.AddRange(ConsoleAppLinkEvironment.LinkExecutable(false));
 			}
-		}
 
-		/// <summary>
-		/// Called to allow the binary to to determine if it matches the Only module "short module name".
-		/// </summary>
-		/// <param name="OnlyModules"></param>
-		/// <returns>The OnlyModule if found, null if not</returns>
-		public override OnlyModule FindOnlyModule(List<OnlyModule> OnlyModules)
-		{
-			foreach (var ModuleName in ModuleNames)
+			foreach (var Executable in Executables)
 			{
-				foreach (var OnlyModule in OnlyModules)
-				{
-					if (OnlyModule.OnlyModuleName.ToLower() == ModuleName.ToLower())
-					{
-						return OnlyModule;
-					}
-				}
+				OutputFiles.AddRange(TargetToolChain.PostBuild(Executable, BinaryLinkEnvironment));
 			}
-			return null;
+
+			return OutputFiles;
 		}
 
-		public override List<UEBuildModule> FindGameModules()
+		private List<FileItem> CreateOutputFilesForUCA(LinkEnvironment BinaryLinkEnvironment)
 		{
-			var GameModules = new List<UEBuildModule>();
-			foreach (var ModuleName in ModuleNames)
+			var OutputFiles = new List<FileItem>();
+			var ModuleName = ModuleNames.First(Name => Name.CompareTo(BuildConfiguration.UCAModuleToAnalyze) == 0);
+			var ModuleCPP = (UEBuildModuleCPP)Target.GetModuleByName(ModuleName);
+			var ModulePrivatePCH = ModuleCPP.ProcessedDependencies.UniquePCHHeaderFile;
+			var IntermediatePath = Path.Combine(Target.ProjectIntermediateDirectory, ModuleName);
+			var OutputFileName = Target.OutputPath;
+			var OutputFile = FileItem.GetItemByPath(OutputFileName);
+
+			Action LinkAction = new Action(ActionType.Compile);
+			LinkAction.WorkingDirectory = Path.GetFullPath(".");
+			LinkAction.CommandPath = System.IO.Path.Combine(LinkAction.WorkingDirectory, @"..", @"Binaries", @"Win32", @"UnrealCodeAnalyzer.exe");
+			LinkAction.bIsVCCompiler = false;
+			LinkAction.ProducedItems.Add(OutputFile);
+			LinkAction.PrerequisiteItems.AddRange(BinaryLinkEnvironment.InputFiles);
+			LinkAction.CommandArguments = @"-AnalyzePCHFile -PCHFile=""" + ModulePrivatePCH.AbsolutePath + @""" -OutputFile=""" + OutputFileName + @""" -HeaderDataPath=""" + IntermediatePath + @""" -UsageThreshold " + BuildConfiguration.UCAUsageThreshold.ToString(CultureInfo.InvariantCulture);
+
+			foreach (string IncludeSearchPath in ModuleCPP.IncludeSearchPaths)
 			{
-				UEBuildModule Module = Target.GetModuleByName(ModuleName);
-				if (!Utils.IsFileUnderDirectory(Module.ModuleDirectory, BuildConfiguration.RelativeEnginePath))
-				{
-					GameModules.Add(Module);
-				}
+				LinkAction.CommandArguments += @" /I""" + LinkAction.WorkingDirectory + @"\" + IncludeSearchPath + @"""";
 			}
-			return GameModules;
-		}
 
-		// Object interface.
+			OutputFiles.Add(OutputFile);
 
-		/// <summary>
-		/// ToString implementation
-		/// </summary>
-		/// <returns>Returns the OutputFilePath for this binary</returns>
-		public override string ToString()
-		{
-			return Config.OutputFilePath;
+			return OutputFiles;
 		}
 	};
 
diff --git a/Engine/Source/Programs/UnrealBuildTool/Configuration/UEBuildConfiguration.cs b/Engine/Source/Programs/UnrealBuildTool/Configuration/UEBuildConfiguration.cs
--- a/Engine/Source/Programs/UnrealBuildTool/Configuration/UEBuildConfiguration.cs
+++ b/Engine/Source/Programs/UnrealBuildTool/Configuration/UEBuildConfiguration.cs
@@ -253,6 +253,19 @@ namespace UnrealBuildTool
 				&& Directory.Exists(UEBuildConfiguration.UEThirdPartySourceDirectory + "NotForLicensees/Simplygon") == true
 				&& Directory.Exists("Developer/SimplygonMeshReduction") == true
 				&& !(ProjectFileGenerator.bGenerateProjectFiles && ProjectFileGenerator.bGeneratingRocketProjectFiles);
+
+			if (UnrealBuildTool.CommandLineContains(@"UnrealCodeAnalyzer") || UnrealBuildTool.CommandLineContains(@"UnrealHeaderTool"))
+			{
+				BuildConfiguration.bRunUnrealCodeAnalyzer = false;
+			}
+
+			if (BuildConfiguration.bRunUnrealCodeAnalyzer)
+			{
+				BuildConfiguration.bUsePCHFiles = false;
+				BuildConfiguration.bUseUnityBuild = false;
+				BuildConfiguration.bUseSharedPCHs = false;
+				BuildConfiguration.bAllowXGE = false;
+			}
 		}
 
 		/**
diff --git a/Engine/Source/Programs/UnrealBuildTool/Configuration/UEBuildModule.cs b/Engine/Source/Programs/UnrealBuildTool/Configuration/UEBuildModule.cs
--- a/Engine/Source/Programs/UnrealBuildTool/Configuration/UEBuildModule.cs
+++ b/Engine/Source/Programs/UnrealBuildTool/Configuration/UEBuildModule.cs
@@ -987,6 +987,8 @@ namespace UnrealBuildTool
 		/** Path to this module's redist static library */
 		public string[] RedistStaticLibraryPaths = null;
 
+		public List<string> IncludeSearchPaths = new List<string>();
+
 		/** Whether we're building the redist static library (as well as using it). */
 		public bool bBuildingRedistStaticLibrary = false;
 
@@ -1142,6 +1144,8 @@ namespace UnrealBuildTool
 			}
 
 			var ModuleCompileEnvironment = CreateModuleCompileEnvironment(CompileEnvironment);
+			IncludeSearchPaths = ModuleCompileEnvironment.Config.CPPIncludeInfo.IncludePaths.ToList();
+			IncludeSearchPaths.AddRange(ModuleCompileEnvironment.Config.CPPIncludeInfo.SystemIncludePaths.ToList());
 
 			if( IntelliSenseGatherer != null )
 			{
diff --git a/Engine/Source/Programs/UnrealBuildTool/Configuration/UEBuildTarget.cs b/Engine/Source/Programs/UnrealBuildTool/Configuration/UEBuildTarget.cs
--- a/Engine/Source/Programs/UnrealBuildTool/Configuration/UEBuildTarget.cs
+++ b/Engine/Source/Programs/UnrealBuildTool/Configuration/UEBuildTarget.cs
@@ -2256,7 +2256,16 @@ namespace UnrealBuildTool
 		{
 			// Determine the binary extension for the platform and binary type.
 			var BuildPlatform = UEBuildPlatform.GetBuildPlatform(Platform);
-			string BinaryExtension = BuildPlatform.GetBinaryExtension(BinaryType);
+			string BinaryExtension;
+
+			if (!BuildConfiguration.bRunUnrealCodeAnalyzer)
+			{
+				BinaryExtension = BuildPlatform.GetBinaryExtension(BinaryType);
+			}
+			else
+			{
+				BinaryExtension = @"-" + BuildConfiguration.UCAModuleToAnalyze + @".analysis";
+			}
 
 			UnrealTargetConfiguration LocalConfig = Configuration;
 			if(Configuration == UnrealTargetConfiguration.DebugGame && !String.IsNullOrEmpty(ModuleName) && !RulesCompiler.IsGameModule(ModuleName))
diff --git a/Engine/Source/Programs/UnrealBuildTool/System/UnrealBuildTool.cs b/Engine/Source/Programs/UnrealBuildTool/System/UnrealBuildTool.cs
--- a/Engine/Source/Programs/UnrealBuildTool/System/UnrealBuildTool.cs
+++ b/Engine/Source/Programs/UnrealBuildTool/System/UnrealBuildTool.cs
@@ -798,7 +798,8 @@ namespace UnrealBuildTool
                     {
                         string LowercaseArg = Arg.ToLowerInvariant();
                         const string OnlyPlatformSpecificForArg = "-onlyplatformspecificfor=";
-
+						const string UCAUsageThresholdString = "-ucausagethreshold=";
+						const string UCAModuleToAnalyzeString = "-ucamoduletoanalyze=";
                         if (ParseRocketCommandlineArg(Arg, ref GameName))
                         {
                             // Already handled at startup. Calling now just to properly set the game name
@@ -979,6 +980,19 @@ namespace UnrealBuildTool
 						{
 							UEBuildConfiguration.bCompileCEF3 = false;
 						}
+						else if (LowercaseArg == "-rununrealcodeanalyzer")
+						{
+							BuildConfiguration.bRunUnrealCodeAnalyzer = true;
+						}
+						else if (LowercaseArg.StartsWith(UCAUsageThresholdString))
+						{
+							string UCAUsageThresholdValue = LowercaseArg.Substring(UCAUsageThresholdString.Length);
+							BuildConfiguration.UCAUsageThreshold = float.Parse(UCAUsageThresholdValue.Replace(',', '.'), System.Globalization.CultureInfo.InvariantCulture);
+						}
+						else if (LowercaseArg.StartsWith(UCAModuleToAnalyzeString))
+						{
+							BuildConfiguration.UCAModuleToAnalyze = LowercaseArg.Substring(UCAModuleToAnalyzeString.Length).Trim();
+						}
                         else if (CheckPlatform.ToString().ToLowerInvariant() == LowercaseArg)
                         {
                             // It's the platform set...
@@ -1007,6 +1021,11 @@ namespace UnrealBuildTool
                         throw new BuildException( "UnrealBuildTool: At least one of either IsGatheringBuild or IsAssemblingBuild must be true.  Did you pass '-NoGather' with '-NoAssemble'?" );
                     }
 
+					if (BuildConfiguration.bRunUnrealCodeAnalyzer && (BuildConfiguration.UCAModuleToAnalyze == null || BuildConfiguration.UCAModuleToAnalyze.Length == 0))
+					{
+						throw new BuildException("When running UnrealCodeAnalyzer, please specify module to analyze in UCAModuleToAnalyze field in BuildConfiguration.xml");
+					}
+
                     // Send an event with basic usage dimensions
                     // [RCL] 2014-11-03 FIXME: this is incorrect since we can have more than one action
                     Telemetry.SendEvent("CommonAttributes.2",
diff --git a/Engine/Source/Programs/UnrealBuildTool/Windows/UEBuildWindows.cs b/Engine/Source/Programs/UnrealBuildTool/Windows/UEBuildWindows.cs
--- a/Engine/Source/Programs/UnrealBuildTool/Windows/UEBuildWindows.cs
+++ b/Engine/Source/Programs/UnrealBuildTool/Windows/UEBuildWindows.cs
@@ -215,7 +215,14 @@ namespace UnrealBuildTool
                 case UEBuildBinaryType.StaticLibrary:
                     return ".lib";
                 case UEBuildBinaryType.Object:
-                    return ".obj";
+					if (!BuildConfiguration.bRunUnrealCodeAnalyzer)
+					{
+						return ".obj";
+					}
+					else
+					{
+						return @".includes";
+					}
                 case UEBuildBinaryType.PrecompiledHeader:
                     return ".pch";
             }
diff --git a/Engine/Source/Programs/UnrealBuildTool/Windows/VCToolChain.cs b/Engine/Source/Programs/UnrealBuildTool/Windows/VCToolChain.cs
--- a/Engine/Source/Programs/UnrealBuildTool/Windows/VCToolChain.cs
+++ b/Engine/Source/Programs/UnrealBuildTool/Windows/VCToolChain.cs
@@ -92,38 +92,15 @@ namespace UnrealBuildTool
 				}
 			}
 
-			if (CompileEnvironment.Config.Target.Platform == CPPTargetPlatform.Win64)
-			{
-				// Pack struct members on 8-byte boundaries.
-				Arguments.Append(" /Zp8");
-			}
-			else
-			{
-				// Pack struct members on 4-byte boundaries.
-				Arguments.Append(" /Zp4");
-			}
-
 			// Separate functions for linker.
 			Arguments.Append(" /Gy");
 
-			if( !WindowsPlatform.bCompileWithClang )	// @todo clang: These options are not supported with clang-cl yet
-			{
-				// Relaxes floating point precision semantics to allow more optimization.
-				Arguments.Append(" /fp:fast");
-			}
-
 			// Compile into an .obj file, and skip linking.
 			Arguments.Append(" /c");
 
 			// Allow 800% of the default memory allocation limit.
 			Arguments.Append(" /Zm800");
 
-			if( !WindowsPlatform.bCompileWithClang )	// @todo clang: Not supported in clang-cl yet
-			{
-				// Allow large object files to avoid hitting the 2^16 section limit when running with -StressTestUnity.
-				Arguments.Append(" /bigobj");
-			}
-
 			// Disable "The file contains a character that cannot be represented in the current code page" warning for non-US windows.
 			Arguments.Append(" /wd4819");
 
@@ -205,21 +182,6 @@ namespace UnrealBuildTool
 				if( CompileEnvironment.Config.OptimizeCode >= ModuleRules.CodeOptimization.InNonDebugBuilds )
 				{
 					Arguments.Append(" /Ox");
-
-					if( !WindowsPlatform.bCompileWithClang )
-					{
-						// Allow optimized code to be debugged more easily.  This makes PDBs a bit larger, but doesn't noticeably affect
-						// compile times.  The executable code is not affected at all by this switch, only the debugging information.
-						if (EnvVars.CLExeVersion >= new Version("18.0.30723"))
-						{
-							// VC2013 Update 3 has a new flag for doing this
-							Arguments.Append(" /Zo");
-						}
-						else
-						{
-							Arguments.Append(" /d2Zi+");
-						}
-					}
 				}
 				
 				// Favor code speed.
@@ -350,6 +312,59 @@ namespace UnrealBuildTool
 					Arguments.Append(" /MD");
 				}
 			}
+
+			//
+			// Options that aren't parsed by clang-cl yet
+			//
+			if (!WindowsPlatform.bCompileWithClang && !BuildConfiguration.bRunUnrealCodeAnalyzer)	// @todo clang: Not supported in clang-cl yet
+			{
+				// Allow large object files to avoid hitting the 2^16 section limit when running with -StressTestUnity.
+				Arguments.Append(" /bigobj");
+
+				// Relaxes floating point precision semantics to allow more optimization.
+				Arguments.Append(" /fp:fast");
+
+				if (CompileEnvironment.Config.OptimizeCode >= ModuleRules.CodeOptimization.InNonDebugBuilds)
+				{
+					// Allow optimized code to be debugged more easily.  This makes PDBs a bit larger, but doesn't noticeably affect
+					// compile times.  The executable code is not affected at all by this switch, only the debugging information.
+					if (EnvVars.CLExeVersion >= new Version("18.0.30723"))
+					{
+						// VC2013 Update 3 has a new flag for doing this
+						Arguments.Append(" /Zo");
+					}
+					else
+					{
+						Arguments.Append(" /d2Zi+");
+					}
+				}
+			}
+
+			//
+			// Options skipped for UnrealCodeAnalyzer
+			//
+			if (!BuildConfiguration.bRunUnrealCodeAnalyzer)
+			{
+				if (CompileEnvironment.Config.Target.Platform == CPPTargetPlatform.Win64)
+				{
+					// Pack struct members on 8-byte boundaries.
+					Arguments.Append(" /Zp8");
+				}
+				else if (!UnrealBuildTool.CommandLineContains(@"UnrealCodeAnalyzer"))
+				{
+					// Pack struct members on 4-byte boundaries.
+					// Disabled when compiling UnrealCodeAnalyzer as it breaks compilation (some structs in clang/llvm headers require 8-byte alignment in 32-bit compilation)
+					Arguments.Append(" /Zp4");
+				}
+			}
+
+			//
+			// Options required by UnrealCodeAnalyzer
+			//
+			if (BuildConfiguration.bRunUnrealCodeAnalyzer)
+			{
+				Arguments.Append(" /DUNREAL_CODE_ANALYZER");
+			}
 		}
 
 		static void AppendCLArguments_CPP( CPPEnvironment CompileEnvironment, StringBuilder Arguments )
@@ -380,8 +395,17 @@ namespace UnrealBuildTool
 				}
 			}
 
-			// Level 4 warnings.
-			Arguments.Append(" /W4");
+			// Set warning level.
+			if (!BuildConfiguration.bRunUnrealCodeAnalyzer)
+			{
+				// Restrictive during regular compilation.
+				Arguments.Append(" /W4");
+			}
+			else
+			{
+				// If we had /W4 with clang on windows we would be flooded with warnings. This will be fixed incrementally.
+				Arguments.Append(" /W0");
+			}
 
 			if( WindowsPlatform.bCompileWithClang )
 			{
@@ -675,9 +699,19 @@ namespace UnrealBuildTool
 			AppendCLArguments_Global(CompileEnvironment, EnvVars, Arguments);
 
 			// Add include paths to the argument list.
-			foreach (string IncludePath in CompileEnvironment.Config.CPPIncludeInfo.IncludePaths)
+			if (!BuildConfiguration.bRunUnrealCodeAnalyzer)
+			{
+				foreach (string IncludePath in CompileEnvironment.Config.CPPIncludeInfo.IncludePaths)
+				{
+					Arguments.AppendFormat(" /I \"{0}\"", IncludePath);
+				}
+			}
+			else
 			{
-				Arguments.AppendFormat(" /I \"{0}\"", IncludePath);
+				foreach (string IncludePath in CompileEnvironment.Config.CPPIncludeInfo.IncludePaths)
+				{
+					Arguments.AppendFormat(" /I \"{0}\"", System.IO.Path.GetFullPath(IncludePath));
+				}
 			}
 			foreach (string IncludePath in CompileEnvironment.Config.CPPIncludeInfo.SystemIncludePaths)
 			{
@@ -850,9 +884,13 @@ namespace UnrealBuildTool
 							}
 						}
 					}
-					
-					// Add the source file path to the command-line.
-					FileArguments.AppendFormat(" \"{0}\"", SourceFile.AbsolutePath);
+
+					// UnrealCodeAnalyzer requires compiled file name to be first argument.
+					if (!BuildConfiguration.bRunUnrealCodeAnalyzer)
+					{
+						// Add the source file path to the command-line.
+						FileArguments.AppendFormat(" \"{0}\"", SourceFile.AbsolutePath);
+					}
 
 					CompileAction.StatusDescription = Path.GetFileName( SourceFile.AbsolutePath );
 				}
@@ -869,7 +907,11 @@ namespace UnrealBuildTool
 						);
 					CompileAction.ProducedItems.Add(ObjectFile);
 					Result.ObjectFiles.Add(ObjectFile);
-					FileArguments.AppendFormat(" /Fo\"{0}\"", ObjectFile.AbsolutePath);
+					// UnrealCodeAnalyzer requires specific position of output file.
+					if (!BuildConfiguration.bRunUnrealCodeAnalyzer)
+					{
+						FileArguments.AppendFormat(" /Fo\"{0}\"", ObjectFile.AbsolutePath);
+					}
 				}
 
 				// Create PDB files if we were configured to do that.
@@ -939,13 +981,30 @@ namespace UnrealBuildTool
 				}
 
 				CompileAction.WorkingDirectory = Path.GetFullPath(".");
-				CompileAction.CommandPath      = EnvVars.CompilerPath;
+				if (BuildConfiguration.bRunUnrealCodeAnalyzer)
+				{
+					CompileAction.CommandPath = System.IO.Path.Combine(CompileAction.WorkingDirectory, @"..", @"Binaries", @"Win32", @"UnrealCodeAnalyzer.exe");
+				}
+				else
+				{
+					CompileAction.CommandPath = EnvVars.CompilerPath;
+				}
 
-				if( !WindowsPlatform.bCompileWithClang )
+				if( !WindowsPlatform.bCompileWithClang || !BuildConfiguration.bRunUnrealCodeAnalyzer )
 				{
 					CompileAction.bIsVCCompiler = true;
 				}
-				CompileAction.CommandArguments = Arguments.ToString() + FileArguments.ToString() + CompileEnvironment.Config.AdditionalArguments;
+
+				string UnrealCodeAnalyzerArguments = "";
+				if (BuildConfiguration.bRunUnrealCodeAnalyzer)
+				{
+					var ObjectFileExtension = @".includes";
+					FileItem ObjectFile = FileItem.GetItemByPath(Path.Combine(CompileEnvironment.Config.OutputDirectory, Path.GetFileName(SourceFile.AbsolutePath) + ObjectFileExtension));
+					var ClangPath = System.IO.Path.Combine(CompileAction.WorkingDirectory, @"ThirdParty", @"llvm", @"3.5.0", @"bin", @"vs2013", @"x86", @"release", @"clang++.exe");
+					UnrealCodeAnalyzerArguments = @"-CreateIncludeFiles " + SourceFile.AbsolutePath + @" -OutputFile=""" + ObjectFile.AbsolutePath + @""" -- " + ClangPath + @" --driver-mode=cl ";
+				}
+
+				CompileAction.CommandArguments = UnrealCodeAnalyzerArguments + Arguments.ToString() + FileArguments.ToString() + CompileEnvironment.Config.AdditionalArguments;
 
 				if( CompileEnvironment.Config.PrecompiledHeaderAction == PrecompiledHeaderAction.Create )
 				{
@@ -958,7 +1017,7 @@ namespace UnrealBuildTool
 					Log.TraceVerbose("     Command: " + CompileAction.CommandArguments);
 				}
 
-				if( WindowsPlatform.bCompileWithClang )
+				if( WindowsPlatform.bCompileWithClang || BuildConfiguration.bRunUnrealCodeAnalyzer)
 				{
 					// Clang doesn't print the file names by default, so we'll do it ourselves
 					CompileAction.bShouldOutputStatusDescription = true;
diff --git a/Engine/Source/Runtime/Core/Public/Templates/UnrealTemplate.h b/Engine/Source/Runtime/Core/Public/Templates/UnrealTemplate.h
--- a/Engine/Source/Runtime/Core/Public/Templates/UnrealTemplate.h
+++ b/Engine/Source/Runtime/Core/Public/Templates/UnrealTemplate.h
@@ -150,7 +150,15 @@ char (&ArrayCountHelper(const T (&)[N]))[N];
 #define ARRAY_COUNT( array ) (sizeof(ArrayCountHelper(array))+0)
 
 // Offset of a struct member.
+#ifndef UNREAL_CODE_ANALYZER
+// UCA uses clang on Windows. According to C++11 standard, (which in this case clang follows and msvc doesn't)
+// forbids using reinterpret_cast in constant expressions. msvc uses reinterpret_cast in offsetof macro,
+// while clang uses compiler intrinsic. Calling static_assert(STRUCT_OFFSET(x, y) == SomeValue) causes compiler
+// error when using clang on Windows (while including windows headers).
 #define STRUCT_OFFSET( struc, member )	offsetof(struc, member)
+#else
+#define STRUCT_OFFSET( struc, member )	__builtin_offsetof(struc, member)
+#endif
 
 #if PLATFORM_VTABLE_AT_END_OF_CLASS
 	#error need implementation
diff --git a/Engine/Source/Runtime/Core/Public/Windows/WIndowsPlatform.h b/Engine/Source/Runtime/Core/Public/Windows/WIndowsPlatform.h
--- a/Engine/Source/Runtime/Core/Public/Windows/WIndowsPlatform.h
+++ b/Engine/Source/Runtime/Core/Public/Windows/WIndowsPlatform.h
@@ -57,7 +57,7 @@ typedef FWindowsPlatformTypes FPlatformTypes;
 #define PLATFORM_SUPPORTS_TBB								1
 #define PLATFORM_SUPPORTS_NAMED_PIPES						1
 #define PLATFORM_COMPILER_HAS_DEFAULTED_FUNCTIONS			0
-#if _MSC_VER >= 1800
+#if _MSC_VER >= 1800 || __clang__
 	#define PLATFORM_COMPILER_HAS_VARIADIC_TEMPLATES		1
 	#define PLATFORM_COMPILER_HAS_EXPLICIT_OPERATORS		1
 #else
