Set the OwningModule for frameworks during SetupPublicCompileEnvironment also, since it's needed for generating compiler arguments.

diff --git a/Engine/Source/Programs/UnrealBuildTool/Configuration/UEBuildModule.cs b/Engine/Source/Programs/UnrealBuildTool/Configuration/UEBuildModule.cs
--- a/Engine/Source/Programs/UnrealBuildTool/Configuration/UEBuildModule.cs
+++ b/Engine/Source/Programs/UnrealBuildTool/Configuration/UEBuildModule.cs
@@ -612,6 +612,12 @@ namespace UnrealBuildTool
 
 				// Add the additional frameworks so that the compiler can know about their #include paths
 				AdditionalFrameworks.AddRange(PublicAdditionalFrameworks);
+				
+				// Remember the module so we can refer to it when needed
+				foreach (var Framework in PublicAdditionalFrameworks)
+				{
+					Framework.OwningModule = this;
+				}
 			}
 		}
 
