Misplaced ssh command in the remote tool chain.

diff --git a/Engine/Source/Programs/UnrealBuildTool/ToolChain/RemoteToolChain.cs b/Engine/Source/Programs/UnrealBuildTool/ToolChain/RemoteToolChain.cs
--- a/Engine/Source/Programs/UnrealBuildTool/ToolChain/RemoteToolChain.cs
+++ b/Engine/Source/Programs/UnrealBuildTool/ToolChain/RemoteToolChain.cs
@@ -214,17 +214,21 @@ namespace UnrealBuildTool
 		// Gather a users root path from the remote server. Should only be called once.
 		public static void SetUserDevRootFromServer()
 		{
-			Hashtable Results = SSHCommand("/", "echo $HOME", null);
-
-			if (Results == null)
-			{
-				Log.TraceInformation("UserDevRoot Command failed to execute!");
-			}
-			else if (Results["CommandOutput"] != null)
+			if (!bUseRPCUtil)
 			{
-				// pass back the string
-				string HomeLocation = Results["CommandOutput"] as string;
-				UserDevRootMac = HomeLocation + UserDevRootMac;
+				// Only set relative to the users root when using rsync, for now
+				Hashtable Results = RPCUtilHelper.Command("/", "echo $HOME", null);
+
+				if (Results == null)
+				{
+					Log.TraceInformation("UserDevRoot Command failed to execute!");
+				}
+				else if (Results["CommandOutput"] != null)
+				{
+					// pass back the string
+					string HomeLocation = Results["CommandOutput"] as string;
+					UserDevRootMac = HomeLocation + UserDevRootMac;
+				}
 			}
 		}
 
