fix for clearing out the remote server name if the user ini doesn't have it
#ios

diff --git a/Engine/Source/Programs/UnrealBuildTool/ToolChain/RemoteToolChain.cs b/Engine/Source/Programs/UnrealBuildTool/ToolChain/RemoteToolChain.cs
--- a/Engine/Source/Programs/UnrealBuildTool/ToolChain/RemoteToolChain.cs
+++ b/Engine/Source/Programs/UnrealBuildTool/ToolChain/RemoteToolChain.cs
@@ -192,7 +192,11 @@ namespace UnrealBuildTool
 			base.ParseProjectSettings();
 
 			ConfigCacheIni Ini = new ConfigCacheIni(UnrealTargetPlatform.IOS, "Engine", UnrealBuildTool.GetUProjectPath());
-			Ini.GetString("/Script/IOSRuntimeSettings.IOSRuntimeSettings", "RemoteServerName", out RemoteServerName);
+			string ServerName = RemoteServerName;
+			if (Ini.GetString("/Script/IOSRuntimeSettings.IOSRuntimeSettings", "RemoteServerName", out ServerName))
+			{
+				RemoteServerName = ServerName;
+			}
 
 			bool bUseRSync = false;
 			if (Ini.GetBool("/Script/IOSRuntimeSettings.IOSRuntimeSettings", "bUseRSync", out bUseRSync))
