Fixed setting of rpaths on Mac when building remotely from PC

diff --git a/Engine/Source/Programs/UnrealBuildTool/Mac/MacToolChain.cs b/Engine/Source/Programs/UnrealBuildTool/Mac/MacToolChain.cs
--- a/Engine/Source/Programs/UnrealBuildTool/Mac/MacToolChain.cs
+++ b/Engine/Source/Programs/UnrealBuildTool/Mac/MacToolChain.cs
@@ -596,7 +596,7 @@ namespace UnrealBuildTool
 			string ExeDir = Path.GetDirectoryName(ExeAbsolutePath);
 			if ((Library.Contains("/Plugins/") || Library.Contains("/Binaries/ThirdParty/")) && Library.EndsWith("dylib") && LibraryDir != ExeDir)
 			{
-				string RelativePath = Utils.MakePathRelativeTo(LibraryDir, ExeDir);
+				string RelativePath = Utils.MakePathRelativeTo(LibraryDir, ExeDir).Replace("\\", "/");
 				if (!RelativePath.Contains(LibraryDir) && !RPaths.Contains(RelativePath))
 				{
 					RPaths.Add(RelativePath);
@@ -606,7 +606,7 @@ namespace UnrealBuildTool
 					{
 						string PathInBundle = Path.Combine(Path.GetDirectoryName(ExeDir), "UE4/Engine/Binaries/Mac", RelativePath.Substring(9));
 						Utils.CollapseRelativeDirectories(ref PathInBundle);
-						string RelativePathInBundle = Utils.MakePathRelativeTo(PathInBundle, ExeDir);
+						string RelativePathInBundle = Utils.MakePathRelativeTo(PathInBundle, ExeDir).Replace("\\", "/");
 						LinkCommand += " -rpath \"@loader_path/" + RelativePathInBundle + "\"";
 					}
 				}
