Enable User.ini's.
Switch HTML5 SDK setupt to prefer values from User*.ini's over environment varibles

FYI #codereivew josh.adams, ankit.khare, rob.jones

diff --git a/Engine/Source/Developer/HTML5/HTML5TargetPlatform/Private/HTML5TargetPlatform.cpp b/Engine/Source/Developer/HTML5/HTML5TargetPlatform/Private/HTML5TargetPlatform.cpp
--- a/Engine/Source/Developer/HTML5/HTML5TargetPlatform/Private/HTML5TargetPlatform.cpp
+++ b/Engine/Source/Developer/HTML5/HTML5TargetPlatform/Private/HTML5TargetPlatform.cpp
@@ -89,6 +89,10 @@ bool FHTML5TargetPlatform::IsSdkInstalled(bool bProjectHasCode, FString& OutDocu
 		const FString& Platform = It.Key.ToString();
 		const FString& Path = It.Value;
 		{
+			if (Platform == "Emscripten" && IFileManager::Get().DirectoryExists(*Path))
+			{
+				return true;
+			}
 #if PLATFORM_WINDOWS
 			if ( Platform == "Windows" && IFileManager::Get().DirectoryExists(*Path)) 
 			{
diff --git a/Engine/Source/Programs/UnrealBuildTool/Configuration/EngineConfiguration.cs b/Engine/Source/Programs/UnrealBuildTool/Configuration/EngineConfiguration.cs
--- a/Engine/Source/Programs/UnrealBuildTool/Configuration/EngineConfiguration.cs
+++ b/Engine/Source/Programs/UnrealBuildTool/Configuration/EngineConfiguration.cs
@@ -482,6 +482,32 @@ namespace UnrealBuildTool
 
 			// Game/Config/Default* ini
 			yield return Path.Combine(ProjectDirectory, "Config", "Default" + BaseIniName + ".ini");
+			
+			string UserSettingsFolder = null; // Match FPlatformProcess::UserSettingsDir()
+			string PersonalFolder = null; // Match FPlatformProcess::UserDir()
+			if (Environment.OSVersion.Platform == PlatformID.MacOSX)
+			{
+				UserSettingsFolder = Path.Combine(Environment.GetFolderPath(Environment.SpecialFolder.Personal), "Library", "Application Support", "Epic");
+				PersonalFolder = Path.Combine(Environment.GetFolderPath(Environment.SpecialFolder.Personal), "Documents");
+			}
+			else if (Environment.OSVersion.Platform == PlatformID.Unix)
+			{
+				UserSettingsFolder = Path.Combine(Environment.GetFolderPath(Environment.SpecialFolder.ApplicationData), "Epic");
+				PersonalFolder = Path.Combine(Environment.GetFolderPath(Environment.SpecialFolder.Personal), "Documents");
+			}
+			else
+			{
+				UserSettingsFolder = Environment.GetFolderPath(Environment.SpecialFolder.LocalApplicationData);
+				PersonalFolder = Environment.GetFolderPath(Environment.SpecialFolder.Personal);
+			}
+
+			// <AppData>/UE4/EngineConfig/User* ini
+			yield return Path.Combine(UserSettingsFolder, "Unreal Engine", "Engine", "Config", "User" + BaseIniName + ".ini");
+			// <Documents>/UE4/EngineConfig/User* ini
+			yield return Path.Combine(PersonalFolder, "Unreal Engine", "Engine", "Config", "User" + BaseIniName + ".ini");
+			
+			// Game/Config/User* ini
+			yield return Path.Combine(ProjectDirectory, "Config", "User" + BaseIniName + ".ini");
 		}
 
 		/// <summary>
diff --git a/Engine/Source/Programs/UnrealBuildTool/HTML5/HTML5SDKInfo.cs b/Engine/Source/Programs/UnrealBuildTool/HTML5/HTML5SDKInfo.cs
--- a/Engine/Source/Programs/UnrealBuildTool/HTML5/HTML5SDKInfo.cs
+++ b/Engine/Source/Programs/UnrealBuildTool/HTML5/HTML5SDKInfo.cs
@@ -12,6 +12,7 @@ namespace UnrealBuildTool
     public class HTML5SDKInfo
     {
         static string EmscriptenSettingsPath = Path.Combine(Environment.GetFolderPath(Environment.SpecialFolder.UserProfile), ".emscripten");
+		static ConfigCacheIni ConfigCache = null;
 
         static Dictionary<string, string> ReadEmscriptenSettings()
         {
@@ -42,13 +43,100 @@ namespace UnrealBuildTool
             return Settings;
         }
 
+		public static bool ParseProjectPath(string[] Arguments, out string ProjectPath)
+		{
+			string TargetName = null;
+			foreach (var Argument in Arguments)
+			{
+				// Treat any platform names as such
+				UnrealTargetPlatform ParsedPlatform = UEBuildPlatform.ConvertStringToPlatform(Argument);
+				if (ParsedPlatform == UnrealTargetPlatform.Unknown)
+				{
+					var ArgUpper = Argument.ToUpperInvariant();
+					switch (ArgUpper)
+					{
+						// Skip any build configurations 
+						case "DEBUG":
+						case "DEBUGGAME":
+						case "DEVELOPMENT":
+						case "SHIPPING":
+						case "TEST":
+							break;
+						default:
+							// Ignore any args except -project
+							if (!ArgUpper.StartsWith("-") || ArgUpper.StartsWith("-PROJECT="))
+							{
+								if (ArgUpper.StartsWith("-PROJECT="))
+								{
+									ArgUpper = ArgUpper.Remove(0, 9).Trim();
+								}
+								// If running Rocket, the PossibleTargetName could contain a path
+								TargetName = ArgUpper;
+
+								// If a project file was not specified see if we can find one
+								string CheckProjectFile = UProjectInfo.GetProjectForTarget(TargetName);
+								if (string.IsNullOrEmpty(CheckProjectFile))
+								{
+									CheckProjectFile = UProjectInfo.GetProjectFilePath(TargetName);
+								}
+								if (string.IsNullOrEmpty(CheckProjectFile) == false)
+								{
+									Log.TraceVerbose("Found project file for {0} - {1}", TargetName, CheckProjectFile);
+									string NewProjectFilename = CheckProjectFile;
+									if (Path.IsPathRooted(NewProjectFilename) == false)
+									{
+										NewProjectFilename = Path.GetFullPath(NewProjectFilename);
+									}
+
+									NewProjectFilename = NewProjectFilename.Replace("\\", "/");
+									ProjectPath = Path.GetDirectoryName(NewProjectFilename);
+									return true;
+								} 
+
+								if ((TargetName.Contains("/") || TargetName.Contains("\\")) && TargetName.Contains(".UPROJECT"))
+								{
+									// Parse off the path
+									ProjectPath = Path.GetDirectoryName(TargetName);
+									return true;
+								}
+							}
+						break;
+					}
+				}
+			}
+
+			ProjectPath = "";
+			return false;
+		}
+
+		static void EnsureConfigCacheIsReady()
+		{
+			if (ConfigCache == null)
+			{
+				var CmdLine = Environment.GetCommandLineArgs();
+				string ProjectPath = UnrealBuildTool.GetUProjectPath();
+				if (String.IsNullOrEmpty(ProjectPath))
+				{ 
+					ParseProjectPath(CmdLine, out ProjectPath);
+				}
+				ConfigCache = new ConfigCacheIni(UnrealTargetPlatform.HTML5, "Engine", ProjectPath);
+				// always pick from the Engine the root directory and NOT the staged engine directory. 
+				// This breaks parse order for user & game config inis so disabled but pulling from the staged directory may still be a problem 
+				//string IniFile = Path.GetFullPath(Path.GetDirectoryName(UnrealBuildTool.GetUBTPath()) + "/../../") + "Config/HTML5/HTML5Engine.ini";
+				//ConfigCache.ParseIniFile(IniFile);
+			}
+		}
+
         public static string EmscriptenSDKPath()
         {
-           var ConfigCache = new ConfigCacheIni(UnrealTargetPlatform.HTML5, "Engine", UnrealBuildTool.GetUProjectPath());
-           // always pick from the Engine the root directory and NOT the staged engine directory. 
-           string IniFile = Path.GetFullPath(Path.GetDirectoryName(UnrealBuildTool.GetUBTPath()) + "/../../") + "Config/HTML5/HTML5Engine.ini";
-           ConfigCache.ParseIniFile(IniFile);
+			EnsureConfigCacheIsReady();
+			// Search the ini file for Emscripten root first
+			string EMCCPath;
+			bool ok = ConfigCache.GetString("HTML5SDKPaths", "Emscripten", out EMCCPath);
+			if (ok && System.IO.Directory.Exists(EMCCPath))
+				return EMCCPath;
 
+			// Older method used platform name
             string PlatformName = "";
             if (!Utils.IsRunningOnMono)
             {
@@ -58,8 +146,7 @@ namespace UnrealBuildTool
             {
                 PlatformName = "Mac";
             }
-            string EMCCPath;
-            bool ok = ConfigCache.GetString("HTML5SDKPaths", PlatformName, out EMCCPath);
+            ok = ConfigCache.GetString("HTML5SDKPaths", PlatformName, out EMCCPath);
 
             if (ok && System.IO.Directory.Exists(EMCCPath))
                 return EMCCPath;
@@ -77,7 +164,15 @@ namespace UnrealBuildTool
 
         public  static string PythonPath()
         {
-            // find Python.
+			// find Python.
+			EnsureConfigCacheIsReady();
+			
+			// Check the ini first. Check for Python="path"
+			string PythonPath;
+			bool ok = ConfigCache.GetString("HTML5SDKPaths", "Python", out PythonPath);
+			if (ok && System.IO.File.Exists(PythonPath))
+				return PythonPath;
+           
             var EmscriptenSettings = ReadEmscriptenSettings();
 
 			// check emscripten generated config file. 
@@ -129,6 +224,19 @@ namespace UnrealBuildTool
             return locations[0];
         }
 
+		public static string OverriddenSDKVersion()
+		{
+			EnsureConfigCacheIsReady();
+
+			// Check the ini first. Check for ForceSDKVersion="x.xx.x"
+			string ForcedSDKVersion;
+			bool ok = ConfigCache.GetString("HTML5SDKPaths", "ForceSDKVersion", out ForcedSDKVersion);
+			if (ok)
+				return ForcedSDKVersion;
+
+			return "";
+		}
+
         public static string EmscriptenPackager()
         {
             return Path.Combine(EmscriptenSDKPath(), "tools", "file_packager.py"); 
@@ -151,6 +259,13 @@ namespace UnrealBuildTool
             return true;
         }
 
+		public static bool IsSDKVersionOverridden()
+		{
+			if (String.IsNullOrEmpty(OverriddenSDKVersion()))
+				return false;
+			return true;
+		}
+
         public static string EmscriptenVersion()
         {
             string VersionInfo = File.ReadAllText(GetVersionInfoPath());
diff --git a/Engine/Source/Programs/UnrealBuildTool/HTML5/UEBuildHTML5.cs b/Engine/Source/Programs/UnrealBuildTool/HTML5/UEBuildHTML5.cs
--- a/Engine/Source/Programs/UnrealBuildTool/HTML5/UEBuildHTML5.cs
+++ b/Engine/Source/Programs/UnrealBuildTool/HTML5/UEBuildHTML5.cs
@@ -48,6 +48,8 @@ namespace UnrealBuildTool
 		 */
 		protected override string GetRequiredSDKString()
 		{
+			if (HTML5SDKInfo.IsSDKVersionOverridden())
+				return HTML5SDKInfo.OverriddenSDKVersion();
 			return ExpectedSDKVersion;
 		}
 
@@ -86,13 +88,13 @@ namespace UnrealBuildTool
 			}
             try
             {
-                if (HTML5SDKInfo.EmscriptenVersion().Contains(ExpectedSDKVersion))
+				if (HTML5SDKInfo.EmscriptenVersion().Contains(GetRequiredSDKString()))
                 {
                     return SDKStatus.Valid;
                 }
                 else
                 {
-                    Console.WriteLine("EMSCRIPTEN sdk " + HTML5SDKInfo.EmscriptenVersion() + " found which is unsupported, Please install version " + ExpectedSDKVersion);
+					Console.WriteLine("EMSCRIPTEN sdk " + HTML5SDKInfo.EmscriptenVersion() + " found which is unsupported, Please install version " + GetRequiredSDKString());
                     return SDKStatus.Invalid;
                 }
             }
diff --git a/Engine/Source/Runtime/Core/Private/Misc/ConfigCacheIni.cpp b/Engine/Source/Runtime/Core/Private/Misc/ConfigCacheIni.cpp
--- a/Engine/Source/Runtime/Core/Private/Misc/ConfigCacheIni.cpp
+++ b/Engine/Source/Runtime/Core/Private/Misc/ConfigCacheIni.cpp
@@ -8,6 +8,7 @@
 #include "SecureHash.h"
 #include "DefaultValueHelper.h"
 #include "EngineBuildSettings.h"
+#include "Paths.h"
 
 
 DEFINE_LOG_CATEGORY(LogConfig);
@@ -2616,13 +2617,22 @@ static void GetSourceIniHierarchyFilenames(const TCHAR* InBaseIniName, const TCH
 	OutHierarchy.Add( FIniFilename(FString::Printf(TEXT("%sBase%s.ini"), EngineConfigDir, InBaseIniName), false) );
 	// Engine/Config/Base[Internal/External]* ini
 	OutHierarchy.Add( FIniFilename(FString::Printf(TEXT("%sBase%s%s.ini"), EngineConfigDir, *BuildPurposeName, InBaseIniName), false) );
+
+	// <AppData>/UE4/EngineConfig/User* ini
+    OutHierarchy.Add(FIniFilename(FPaths::Combine(FPlatformProcess::UserSettingsDir(), *FString::Printf(TEXT("Unreal Engine/Engine/Config/User%s.ini"), InBaseIniName)), false));
+	// <Documents>/UE4/EngineConfig/User* ini
+    OutHierarchy.Add(FIniFilename(FPaths::Combine(FPlatformProcess::UserDir(), *FString::Printf(TEXT("Unreal Engine/Engine/Config/User%s.ini"), InBaseIniName)), false));
+
 	// Game/Config/Default* ini
-	OutHierarchy.Add( FIniFilename(FString::Printf(TEXT("%sDefault%s.ini"), SourceConfigDir, InBaseIniName), bRequireDefaultIni) );
+	OutHierarchy.Add(FIniFilename(FString::Printf(TEXT("%sDefault%s.ini"), SourceConfigDir, InBaseIniName), bRequireDefaultIni));
 	// Game/Config/DedicatedServer* ini
-	if (IsRunningDedicatedServer()) 
-	{  
-		OutHierarchy.Add( FIniFilename(FString::Printf(TEXT("%s/DedicatedServer%s.ini"), SourceConfigDir, InBaseIniName), false) );
-	} 
+	if (IsRunningDedicatedServer())
+	{
+		OutHierarchy.Add(FIniFilename(FString::Printf(TEXT("%s/DedicatedServer%s.ini"), SourceConfigDir, InBaseIniName), false));
+	}
+	// Game/Config/User* ini
+	OutHierarchy.Add(FIniFilename(FString::Printf(TEXT("%s/User%s.ini"), SourceConfigDir, InBaseIniName), false));
+	
 	// Engine/Config/Platform/Platform* ini
 	OutHierarchy.Add( FIniFilename(FString::Printf(TEXT("%s%s/%s%s.ini"), EngineConfigDir, *PlatformName, *PlatformName, InBaseIniName), false) );
 	// Game/Config/Platform/Platform* ini
