Fix for linking issues on HTML5. Don't link libraries when building libraries.

diff --git a/Engine/Source/Programs/UnrealBuildTool/HTML5/HTML5ToolChain.cs b/Engine/Source/Programs/UnrealBuildTool/HTML5/HTML5ToolChain.cs
--- a/Engine/Source/Programs/UnrealBuildTool/HTML5/HTML5ToolChain.cs
+++ b/Engine/Source/Programs/UnrealBuildTool/HTML5/HTML5ToolChain.cs
@@ -463,22 +463,25 @@ namespace UnrealBuildTool
                 LinkAction.CommandArguments += string.Format(" \"{0}\"", InputFile.AbsolutePath);
 				LinkAction.PrerequisiteItems.Add(InputFile);
 			}
-            foreach (string InputFile in LinkEnvironment.Config.AdditionalLibraries)
-            {
-                FileItem Item = FileItem.GetItemByPath(InputFile);
-
-                if (Item.AbsolutePath.Contains(".lib"))
-                    continue; 
-
-                if (Item != null)
-                {
-                    if (Item.ToString().Contains(".js"))
-                        LinkAction.CommandArguments += string.Format(" --js-library \"{0}\"", Item.AbsolutePath);
-                    else
-                        LinkAction.CommandArguments += string.Format(" \"{0}\"", Item.AbsolutePath);
-                    LinkAction.PrerequisiteItems.Add(Item);
-                }
-            }
+			if (!LinkEnvironment.Config.bIsBuildingLibrary)
+			{
+				foreach (string InputFile in LinkEnvironment.Config.AdditionalLibraries)
+				{
+					FileItem Item = FileItem.GetItemByPath(InputFile);
+
+					if (Item.AbsolutePath.Contains(".lib"))
+						continue;
+
+					if (Item != null)
+					{
+						if (Item.ToString().Contains(".js"))
+							LinkAction.CommandArguments += string.Format(" --js-library \"{0}\"", Item.AbsolutePath);
+						else
+							LinkAction.CommandArguments += string.Format(" \"{0}\"", Item.AbsolutePath);
+						LinkAction.PrerequisiteItems.Add(Item);
+					}
+				}
+			}
 			// make the file we will create
 			OutputFile = FileItem.GetItemByPath(LinkEnvironment.Config.OutputFilePath);
 			LinkAction.ProducedItems.Add(OutputFile);
