Fix handling of paths containing spaces in iOS toolchain.

diff --git a/Engine/Source/Programs/UnrealBuildTool/IOS/IOSToolChain.cs b/Engine/Source/Programs/UnrealBuildTool/IOS/IOSToolChain.cs
--- a/Engine/Source/Programs/UnrealBuildTool/IOS/IOSToolChain.cs
+++ b/Engine/Source/Programs/UnrealBuildTool/IOS/IOSToolChain.cs
@@ -1217,7 +1217,7 @@ namespace UnrealBuildTool
 						Log.TraceInformation( "Copying bundled asset... LocalSource: {0}, LocalDest: {1}", LocalSource, LocalDest );
 
 						string ResultsText;
-						RunExecutableAndWait( "cp", String.Format( "-R -L {0} {1}", LocalSource, LocalDest ), out ResultsText );
+						RunExecutableAndWait( "cp", String.Format( "-R -L {0} {1}", LocalSource.Replace(" ", "\\ "), LocalDest.Replace(" ", "\\ ") ), out ResultsText );
 					}
 				}
 			}
@@ -1429,7 +1429,7 @@ namespace UnrealBuildTool
 
 						Log.TraceInformation( "Copying bundled asset... RemoteSource: {0}, RemoteDest: {1}, LocalDest: {2}", RemoteSource, RemoteDest, LocalDest );
 
-						Hashtable Results = RPCUtilHelper.Command( "/", String.Format( "cp -R -L {0} {1}", RemoteSource, RemoteDest ), "", null );
+						Hashtable Results = RPCUtilHelper.Command( "/", String.Format( "cp -R -L {0} {1}", RemoteSource.Replace(" ", "\\ "), RemoteDest.Replace(" ", "\\ ") ), "", null );
 
 						foreach ( DictionaryEntry Entry in Results )
 						{
