Include all plugin static libraries in Rocket builds

diff --git a/Engine/Source/Programs/UnrealBuildTool/Configuration/UEBuildTarget.cs b/Engine/Source/Programs/UnrealBuildTool/Configuration/UEBuildTarget.cs
--- a/Engine/Source/Programs/UnrealBuildTool/Configuration/UEBuildTarget.cs
+++ b/Engine/Source/Programs/UnrealBuildTool/Configuration/UEBuildTarget.cs
@@ -705,9 +705,13 @@ namespace UnrealBuildTool
 		[NonSerialized]
 		public LinkEnvironment GlobalLinkEnvironment = new LinkEnvironment();
 
-		/** All plugins enabled for this target */
+		/** All plugins which are built for this target */
 		[NonSerialized]
-		public List<PluginInfo> EnabledPlugins = new List<PluginInfo>();
+		public List<PluginInfo> BuildPlugins = new List<PluginInfo>();
+
+		/** All plugin dependencies for this target. This differs from the list of plugins that is built for Rocket, where we build everything, but link in only the enabled plugins. */
+		[NonSerialized]
+		public List<PluginInfo> DependentPlugins = new List<PluginInfo>();
 
 		/** All application binaries; may include binaries not built by this target. */
 		[NonSerialized]
@@ -1630,9 +1634,9 @@ namespace UnrealBuildTool
 			var SpecialRocketLibFilesThatAreBuildProducts = new List<string>();
 
 			// Add the enabled plugins to the build
-			foreach (PluginInfo Plugin in EnabledPlugins)
+			foreach (PluginInfo BuildPlugin in BuildPlugins)
 			{
-				SpecialRocketLibFilesThatAreBuildProducts.AddRange(AddPlugin(Plugin));
+				SpecialRocketLibFilesThatAreBuildProducts.AddRange(AddPlugin(BuildPlugin));
 			}
 
 			// Allow the platform to setup binaries/plugins/modules
@@ -1873,7 +1877,7 @@ namespace UnrealBuildTool
 					}
 				}
 			}
-			foreach (PluginInfo Plugin in EnabledPlugins)
+			foreach (PluginInfo Plugin in DependentPlugins)
 			{
 				foreach (PluginInfo.PluginModuleInfo Module in Plugin.Modules)
 				{
@@ -2422,32 +2426,41 @@ namespace UnrealBuildTool
 				}
 			}
 
-			// Build the enabled plugin list
-			if (ShouldCompileMonolithic() || TargetType == TargetRules.TargetType.Program)
+			// Build a list of enabled plugins
+			List<string> EnabledPluginNames = new List<string>(Rules.AdditionalPlugins);
+
+			// Add the list of plugins enabled by default
+			if (UEBuildConfiguration.bCompileAgainstEngine)
 			{
-				var FilterPluginNames = new List<string>(Rules.AdditionalPlugins);
+				EnabledPluginNames.AddRange(ValidPlugins.Where(x => x.bEnabledByDefault).Select(x => x.Name));
+			}
 
-				// Add the list of plugins enabled by default
-				if (UEBuildConfiguration.bCompileAgainstEngine)
-				{
-					FilterPluginNames.AddRange(ValidPlugins.Where(x => x.bEnabledByDefault).Select(x => x.Name));
-				}
+			// Update the plugin list for game targets
+			if(TargetType != TargetRules.TargetType.Program && UnrealBuildTool.HasUProjectFile())
+			{
+				// Enable all the game specific plugins by default
+				EnabledPluginNames.AddRange(ValidPlugins.Where(x => x.LoadedFrom == PluginInfo.LoadedFromType.GameProject).Select(x => x.Name));
 
-				// Update the plugin list for game targets
-				if(TargetType != TargetRules.TargetType.Program && UnrealBuildTool.HasUProjectFile())
-				{
-					// Enable all the game specific plugins by default
-					FilterPluginNames.AddRange(ValidPlugins.Where(x => x.LoadedFrom == PluginInfo.LoadedFromType.GameProject).Select(x => x.Name));
+				// Use the project settings to update the plugin list for this target
+				EnabledPluginNames = UProjectInfo.GetEnabledPlugins(UnrealBuildTool.GetUProjectFile(), EnabledPluginNames, Platform);
+			}
 
-					// Use the project settings to update the plugin list for this target
-					FilterPluginNames = UProjectInfo.GetEnabledPlugins(UnrealBuildTool.GetUProjectFile(), FilterPluginNames, Platform);
-				}
+			// Set the list of plugins we're dependent on
+			PluginInfo[] EnabledPlugins = ValidPlugins.Where(x => EnabledPluginNames.Contains(x.Name)).Distinct().ToArray();
+			DependentPlugins.AddRange(EnabledPlugins);
 
-				EnabledPlugins.AddRange(ValidPlugins.Where(x => FilterPluginNames.Contains(x.Name)).Distinct());
+			// Set the list of plugins that should be built
+			if(UnrealBuildTool.BuildingRocket() )
+			{
+				BuildPlugins.AddRange(ValidPlugins);
+			}
+			else if (ShouldCompileMonolithic() || TargetType == TargetRules.TargetType.Program)
+			{
+				BuildPlugins.AddRange(EnabledPlugins);
 			}
 			else
 			{
-				EnabledPlugins.AddRange(ValidPlugins);
+				BuildPlugins.AddRange(ValidPlugins);
 			}
 		}
 
