When generating project files, write a ProjExplorerState section to the solution options which collapses all projects by default (apart from the default game).

diff --git a/Engine/Source/Programs/UnrealBuildTool/System/VCProjectFileGenerator.cs b/Engine/Source/Programs/UnrealBuildTool/System/VCProjectFileGenerator.cs
--- a/Engine/Source/Programs/UnrealBuildTool/System/VCProjectFileGenerator.cs
+++ b/Engine/Source/Programs/UnrealBuildTool/System/VCProjectFileGenerator.cs
@@ -763,22 +763,51 @@ namespace UnrealBuildTool
 			// Save a solution config file which selects the development editor configuration by default.
 			if (bSuccess)
 			{
-				VCSolutionConfigCombination DefaultConfig = SolutionConfigCombinations.Find(x => x.Configuration == UnrealTargetConfiguration.Development && x.Platform == UnrealTargetPlatform.Win64 && x.TargetConfigurationName == "Editor");
-				if (DefaultConfig != null)
+				// Figure out the filename for the SUO file. VS2013 will automatically import the VS2012 options if necessary.
+				string SolutionOptionsExtension = (ProjectFileFormat == VCProjectFileFormat.VisualStudio2012)? "v11.suo" : "v12.suo";
+
+				// Check it doesn't exist before overwriting it. Since these files store the user's preferences, it'd be bad form to overwrite them.
+				string SolutionOptionsFileName = Path.Combine(MasterProjectRelativePath, Path.ChangeExtension(SolutionFileName, SolutionOptionsExtension));
+				if(!File.Exists(SolutionOptionsFileName))
 				{
-					// Figure out the filename for the SUO file. VS2013 will automatically import the VS2012 options if necessary.
-					string SolutionOptionsExtension = (ProjectFileFormat == VCProjectFileFormat.VisualStudio2012)? "v11.suo" : "v12.suo";
+					VCSolutionOptions Options = new VCSolutionOptions();
 
-					// Check it doesn't exist before overwriting it. Since these files store the user's preferences, it'd be bad form to overwrite them.
-					string SolutionOptionsFileName = Path.Combine(MasterProjectRelativePath, Path.ChangeExtension(SolutionFileName, SolutionOptionsExtension));
-					if(!File.Exists(SolutionOptionsFileName))
+					// Set the default configuration and startup project
+					VCSolutionConfigCombination DefaultConfig = SolutionConfigCombinations.Find(x => x.Configuration == UnrealTargetConfiguration.Development && x.Platform == UnrealTargetPlatform.Win64 && x.TargetConfigurationName == "Editor");
+					if (DefaultConfig != null)
 					{
-						VCSolutionOptions Options = new VCSolutionOptions();
-						Options.SolutionConfiguration.Add(new VCBinarySetting("ActiveCfg", DefaultConfig.VCSolutionConfigAndPlatformName));
+						List<VCBinarySetting> Settings = new List<VCBinarySetting>();
+						Settings.Add(new VCBinarySetting("ActiveCfg", DefaultConfig.VCSolutionConfigAndPlatformName));
 						if(DefaultProject != null)
 						{
-							Options.SolutionConfiguration.Add(new VCBinarySetting("StartupProject", ((MSBuildProjectFile)DefaultProject).ProjectGUID.ToString("B")));
+							Settings.Add(new VCBinarySetting("StartupProject", ((MSBuildProjectFile)DefaultProject).ProjectGUID.ToString("B")));
 						}
+						Options.SetConfiguration(Settings);
+					}
+
+					// Mark all the projects as closed by default, apart from the startup project
+					VCSolutionExplorerState ExplorerState = new VCSolutionExplorerState();
+					foreach(ProjectFile ProjectFile in AllProjectFiles)
+					{
+						string ProjectName = Path.GetFileNameWithoutExtension(ProjectFile.ProjectFilePath);
+						if(ProjectFile == DefaultProject)
+						{
+							ExplorerState.OpenProjects.Add(new Tuple<string, string[]>(ProjectName, new string[]{ ProjectName }));
+						}
+						else
+						{
+							ExplorerState.OpenProjects.Add(new Tuple<string, string[]>(ProjectName, new string[]{ }));
+						}
+					}
+					if((!bGeneratingGameProjectFiles && !bGeneratingRocketProjectFiles) || bAlwaysIncludeEngineModules)
+					{
+						ExplorerState.OpenProjects.Add(new Tuple<string, string[]>("Automation", new string[0]));
+					}
+					Options.SetExplorerState(ExplorerState);
+
+					// Write the file
+					if(Options.Sections.Count > 0)
+					{
 						Options.Write(SolutionOptionsFileName);
 					}
 				}
diff --git a/Engine/Source/Programs/UnrealBuildTool/System/VCSolutionOptions.cs b/Engine/Source/Programs/UnrealBuildTool/System/VCSolutionOptions.cs
--- a/Engine/Source/Programs/UnrealBuildTool/System/VCSolutionOptions.cs
+++ b/Engine/Source/Programs/UnrealBuildTool/System/VCSolutionOptions.cs
@@ -127,7 +127,7 @@ namespace UnrealBuildTool
 		}
 	}
 
-	class VCSolutionOptions
+	class VCOleContainer
 	{
 		[Flags]
 		enum STGC : int
@@ -214,13 +214,13 @@ namespace UnrealBuildTool
 		[DllImport("ole32.dll")]
 		static extern int StgOpenStorage([MarshalAs(UnmanagedType.LPWStr)] string pwcsName, IOleStorage pstgPriority, STGM grfMode, IntPtr snbExclude, uint reserved, out IOleStorage ppstgOpen);
 
-		public List<VCBinarySetting> SolutionConfiguration = new List<VCBinarySetting>();
+		public List<KeyValuePair<string, byte[]>> Sections = new List<KeyValuePair<string, byte[]>>();
 
-		public VCSolutionOptions()
-		{
+		public VCOleContainer()
+		{ 
 		}
 
-		public static VCSolutionOptions Read(string InputFileName)
+		public VCOleContainer(string InputFileName)
 		{
 			IOleStorage Storage = null;
 			StgOpenStorage(InputFileName, null, STGM.Direct | STGM.Read | STGM.ShareExclusive, IntPtr.Zero, 0, out Storage);
@@ -234,35 +234,22 @@ namespace UnrealBuildTool
 					STATSTG[] Stats = new STATSTG[200];
 					Enumerator.Next((uint)Stats.Length, Stats, out Fetched);
 
-					VCSolutionOptions Options = new VCSolutionOptions();
-					foreach (STATSTG Stat in Stats)
+					for(uint Idx = 0; Idx < Fetched; Idx++)
 					{
-						if (Stat.pwcsName == "SolutionConfiguration")
+						IOleStream OleStream;
+						Storage.OpenStream(Stats[Idx].pwcsName, IntPtr.Zero, STGM.Read | STGM.ShareExclusive, 0, out OleStream);
+						try
+						{
+							uint SizeRead;
+							byte[] Buffer = new byte[Stats[Idx].cbSize];
+							OleStream.Read(Buffer, (uint)Stats[Idx].cbSize, out SizeRead);
+							Sections.Add(new KeyValuePair<string, byte[]>(Stats[Idx].pwcsName, Buffer));
+						}
+						finally
 						{
-							IOleStream OleStream;
-							Storage.OpenStream(Stat.pwcsName, IntPtr.Zero, STGM.Read | STGM.ShareExclusive, 0, out OleStream);
-							try
-							{
-								uint SizeRead;
-								byte[] Buffer = new byte[Stat.cbSize];
-								OleStream.Read(Buffer, (uint)Stat.cbSize, out SizeRead);
-
-								using (MemoryStream InputStream = new MemoryStream(Buffer, false))
-								{
-									BinaryReader Reader = new BinaryReader(InputStream, Encoding.Unicode);
-									while (InputStream.Position < InputStream.Length)
-									{
-										Options.SolutionConfiguration.Add(VCBinarySetting.Read(Reader));
-									}
-								}
-							}
-							finally
-							{
-								Marshal.ReleaseComObject(OleStream);
-							}
+							Marshal.ReleaseComObject(OleStream);
 						}
 					}
-					return Options;
 				}
 				finally
 				{
@@ -281,30 +268,19 @@ namespace UnrealBuildTool
 			StgCreateDocfile(OutputFileName, STGM.Direct | STGM.Create | STGM.Write | STGM.ShareExclusive, 0, out OleStorage);
 			try
 			{
-				if (SolutionConfiguration.Count > 0)
+				foreach(KeyValuePair<string, byte[]> Section in Sections)
 				{
-					using (MemoryStream OutputStream = new MemoryStream())
+					IOleStream OleStream = null;
+					OleStorage.CreateStream(Section.Key, STGM.Write | STGM.ShareExclusive, 0, 0, out OleStream);
+					try
 					{
-						// Write all the settings to the memory stream
-						BinaryWriter Writer = new BinaryWriter(OutputStream, Encoding.Unicode);
-						foreach (VCBinarySetting Setting in SolutionConfiguration)
-						{
-							Setting.Write(Writer);
-						}
-
-						// Create a named stream in the compound document and write the binary data out
-						IOleStream OleStream = null;
-						OleStorage.CreateStream("SolutionConfiguration", STGM.Write | STGM.ShareExclusive, 0, 0, out OleStream);
-						try
-						{
-							uint Written;
-							OleStream.Write(OutputStream.GetBuffer(), (uint)OutputStream.Length, out Written);
-							OleStream.Commit(STGC.Overwrite);
-						}
-						finally
-						{
-							Marshal.ReleaseComObject(OleStream);
-						}
+						uint Written;
+						OleStream.Write(Section.Value, (uint)Section.Value.Length, out Written);
+						OleStream.Commit(STGC.Overwrite);
+					}
+					finally
+					{
+						Marshal.ReleaseComObject(OleStream);
 					}
 				}
 			}
@@ -313,5 +289,278 @@ namespace UnrealBuildTool
 				Marshal.ReleaseComObject(OleStorage);
 			}
 		}
+
+		int FindIndex(string Name)
+		{
+			return Sections.FindIndex(x => x.Key == Name);
+		}
+
+		public void SetSection(string Name, byte[] Data)
+		{
+			int Idx = FindIndex(Name);
+			if(Idx == -1)
+			{
+				Sections.Add(new KeyValuePair<string,byte[]>(Name, Data));
+			}
+			else
+			{
+				Sections[Idx] = new KeyValuePair<string,byte[]>(Name, Data);
+			}
+		}
+
+		public byte[] GetSection(string Name)
+		{
+			byte[] Data;
+			if(!TryGetSection(Name, out Data))
+			{
+				throw new KeyNotFoundException();
+			}
+			return Data;
+		}
+
+		public bool TryGetSection(string Name, out byte[] Data)
+		{
+			int Idx = FindIndex(Name);
+			if(Idx == -1)
+			{
+				Data = null;
+				return false;
+			}
+			else
+			{
+				Data = Sections[Idx].Value;
+				return true;
+			}
+		}
+	}
+
+	class VCSolutionExplorerState
+	{
+		const string EpilogueGuid = "00000000-0000-0000-0000-000000000000";
+
+		public List<Tuple<string, string[]>> OpenProjects = new List<Tuple<string,string[]>>();
+		public List<Tuple<string, string[]>> SelectedProjects = new List<Tuple<string,string[]>>();
+
+		public VCSolutionExplorerState()
+		{
+		}
+
+		public void Read(Stream InputStream)
+		{
+			BinaryReader Reader = new BinaryReader(InputStream, Encoding.Unicode);
+			ReadOpenProjects(Reader);
+			ReadSelectedProjects(Reader);
+			ReadEpilogue(Reader);
+		}
+
+		public void Write(Stream OutputStream)
+		{
+			BinaryWriter Writer = new BinaryWriter(OutputStream, Encoding.Unicode);
+			WriteOpenProjects(Writer);
+			WriteSelectedProjects(Writer);
+			WriteEpilogue(Writer);
+		}
+
+		void ReadOpenProjects(BinaryReader Reader)
+		{
+			long OpenFoldersEnd = Reader.BaseStream.Position + Reader.ReadInt32();
+			if(Reader.ReadInt32() != 11 || Reader.ReadInt16() != 1 || Reader.ReadByte() != 0)
+			{
+				throw new Exception("Unexpected data in open projects section");
+			}
+
+			int NumProjects = Reader.ReadInt16();
+			for(int ProjectIdx = 0; ProjectIdx < NumProjects; ProjectIdx++)
+			{
+				string ProjectName = ReadString(Reader);
+
+				string[] Folders = new string[Reader.ReadInt16()];
+				for(int FolderIdx = 0; FolderIdx < Folders.Length; FolderIdx++)
+				{
+					Folders[FolderIdx] = ReadString(Reader);
+				}
+
+				OpenProjects.Add(new Tuple<string, string[]>(ProjectName, Folders));
+			}
+
+			Debug.Assert(Reader.BaseStream.Position == OpenFoldersEnd);
+		}
+
+		void WriteOpenProjects(BinaryWriter Writer)
+		{
+			Writer.Write(4 + (4 + 2 + 1) + 2 + OpenProjects.Sum(x => GetStringSize(x.Item1) + 2 + x.Item2.Sum(y => GetStringSize(y))));
+			Writer.Write(11);
+			Writer.Write((short)1);
+			Writer.Write((byte)0);
+			Writer.Write((short)OpenProjects.Count);
+			foreach(Tuple<string, string[]> OpenProject in OpenProjects)
+			{
+				WriteString(Writer, OpenProject.Item1);
+				Writer.Write((short)OpenProject.Item2.Length);
+				foreach(string OpenFolder in OpenProject.Item2)
+				{
+					WriteString(Writer, OpenFolder);
+				}
+			}
+		}
+
+		void ReadSelectedProjects(BinaryReader Reader)
+		{
+			long SelectedProjectsEnd = Reader.BaseStream.Position + Reader.ReadInt32();
+			if(Reader.ReadInt32() != 1)
+			{
+				throw new Exception("Unexpected data in selected projects section");
+			}
+
+			int NumProjects = Reader.ReadInt32();
+			for(int ProjectIdx = 0; ProjectIdx < NumProjects; ProjectIdx++)
+			{
+				string ProjectName = ReadString(Reader);
+
+				string[] Items = new string[Reader.ReadInt32()];
+				for(int ItemIdx = 0; ItemIdx < Items.Length; ItemIdx++)
+				{
+					Items[ItemIdx] = ReadString(Reader);
+				}
+
+				SelectedProjects.Add(new Tuple<string,string[]>(ProjectName, Items));
+			}
+
+			Debug.Assert(Reader.BaseStream.Position == SelectedProjectsEnd);
+		}
+
+		void WriteSelectedProjects(BinaryWriter Writer)
+		{
+			Writer.Write(4 + 4 + 4 + SelectedProjects.Sum(x => GetStringSize(x.Item1) + 4 + x.Item2.Sum(y => GetStringSize(y))));
+			Writer.Write(1);
+			Writer.Write(SelectedProjects.Count);
+			foreach(Tuple<string, string[]> SelectedProject in SelectedProjects)
+			{
+				WriteString(Writer, SelectedProject.Item1);
+				Writer.Write(SelectedProject.Item2.Length);
+				foreach(string SelectedItem in SelectedProject.Item2)
+				{
+					WriteString(Writer, SelectedItem);
+				}
+			}
+		}
+
+		void ReadEpilogue(BinaryReader Reader)
+		{
+			long EpilogueEnd = Reader.BaseStream.Position + Reader.ReadInt32();
+			if(Reader.ReadInt32() != 1 || ReadString(Reader) != EpilogueGuid || Reader.ReadInt32() != 0)
+			{
+				throw new Exception("Unexpected data in epilogue");
+			}
+			Debug.Assert(Reader.BaseStream.Position == EpilogueEnd);
+		}
+
+		void WriteEpilogue(BinaryWriter Writer)
+		{
+			Writer.Write(4 + 4 + GetStringSize(EpilogueGuid) + 4);
+			Writer.Write(1);
+			WriteString(Writer, EpilogueGuid);
+			Writer.Write(0);
+		}
+
+		string ReadString(BinaryReader Reader)
+		{
+			// Read the number of bytes
+			int NumBytes = Reader.ReadByte();
+			for(int Shift = 7; (NumBytes & (1 << Shift)) != 0; Shift += 7)
+			{
+				NumBytes &= (1 << Shift) - 1;
+				NumBytes |= Reader.ReadByte() << Shift;
+			}
+
+			// Read the characters
+			return new string(Reader.ReadChars(NumBytes / 2));
+		}
+
+		void WriteString(BinaryWriter Writer, string Text)
+		{
+			// Write the number of bytes in the string, encoded as a sequence of 7-bit values with the top bit set on all but the last
+			int NumBytes = Text.Length * 2;
+			while(NumBytes >= 128)
+			{
+				Writer.Write((byte)((NumBytes & 127) | 128));
+				NumBytes >>= 7;
+			}
+			Writer.Write((byte)NumBytes);
+
+			// Write the characters
+			Writer.Write(Text.ToCharArray());
+		}
+
+		int GetStringSize(string Text)
+		{
+			int Size = 1 + (Text.Length * 2);
+			for(int Length = Text.Length; Length > 127; Length >>= 7)
+			{
+				Size++;
+			}
+			return Size;
+		}
+	}
+
+	class VCSolutionOptions : VCOleContainer
+	{
+		public VCSolutionOptions()
+		{
+		}
+
+		public VCSolutionOptions(string FileName) : base(FileName)
+		{
+		}
+
+		public IEnumerable<VCBinarySetting> GetConfiguration()
+		{
+			byte[] Data;
+			if(TryGetSection("SolutionConfiguration", out Data))
+			{
+				using (MemoryStream InputStream = new MemoryStream(Data, false))
+				{
+					BinaryReader Reader = new BinaryReader(InputStream, Encoding.Unicode);
+					while (InputStream.Position < InputStream.Length)
+					{
+						yield return VCBinarySetting.Read(Reader);
+					}
+				}
+			}
+		}
+
+		public void SetConfiguration(IEnumerable<VCBinarySetting> Settings)
+		{
+			using(MemoryStream OutputStream = new MemoryStream())
+			{
+				BinaryWriter Writer = new BinaryWriter(OutputStream, Encoding.Unicode);
+				foreach (VCBinarySetting Setting in Settings)
+				{
+					Setting.Write(Writer);
+				}
+				SetSection("SolutionConfiguration", OutputStream.ToArray());
+			}
+		}
+
+		public VCSolutionExplorerState GetExplorerState()
+		{
+			byte[] Data;
+			if(TryGetSection("ProjExplorerState", out Data))
+			{
+				VCSolutionExplorerState State = new VCSolutionExplorerState();
+				State.Read(new MemoryStream(Data, false));
+				return State;
+			}
+			return null;
+		}
+
+		public void SetExplorerState(VCSolutionExplorerState State)
+		{
+			using(MemoryStream OutputStream = new MemoryStream())
+			{
+				State.Write(OutputStream);
+				SetSection("ProjExplorerState", OutputStream.ToArray());
+			}
+		}
 	}
 }
