modified tool chain to pull back all build products from mac
#ios

diff --git a/Engine/Source/Programs/UnrealBuildTool/IOS/IOSToolChain.cs b/Engine/Source/Programs/UnrealBuildTool/IOS/IOSToolChain.cs
--- a/Engine/Source/Programs/UnrealBuildTool/IOS/IOSToolChain.cs
+++ b/Engine/Source/Programs/UnrealBuildTool/IOS/IOSToolChain.cs
@@ -1258,30 +1258,33 @@ namespace UnrealBuildTool
 						Log.TraceInformation("Dangeroulsy Fast mode was requested, but a slow mode hasn't completed. Performing slow now...");
 						bUseDangerouslyFastMode = false;
 					}
-				}	
-	
-				string RemoteExecutablePath = ConvertPath(Target.OutputPath);
-	
-				// when going super fast, just copy the executable to the final resting spot
-				if (bUseDangerouslyFastMode)
-				{
-					Log.TraceInformation("==============================================================================");
-					Log.TraceInformation("USING DANGEROUSLY FAST MODE! IF YOU HAVE ANY PROBLEMS, TRY A REBUILD/CLEAN!!!!");
-					Log.TraceInformation("==============================================================================");	
-
-					// copy the executable
-					string RemoteShadowDirectoryMac = ConvertPath(Path.GetDirectoryName(Target.OutputPath));
-					string FinalRemoteExecutablePath = String.Format("{0}/Payload/{1}.app/{1}", RemoteShadowDirectoryMac, Target.GameName);
-					RPCUtilHelper.Command("/", String.Format("cp -f {0} {1}", RemoteExecutablePath, FinalRemoteExecutablePath), "", null);
 				}
-				else if(!Utils.IsRunningOnMono && BuildHostPlatform.Current.Platform != UnrealTargetPlatform.Mac)
+
+				foreach (string FilePath in BuiltBinaries)
 				{
-					RPCUtilHelper.CopyFile(RemoteExecutablePath, Target.OutputPath, false);
+					string RemoteExecutablePath = ConvertPath(FilePath);
 
-					if (BuildConfiguration.bGeneratedSYMFile == true)
+					// when going super fast, just copy the executable to the final resting spot
+					if (bUseDangerouslyFastMode)
 					{
-						string DSYMExt = ".app.dSYM.zip";
-						RPCUtilHelper.CopyFile(RemoteExecutablePath + DSYMExt, Target.OutputPath + DSYMExt, false);
+						Log.TraceInformation("==============================================================================");
+						Log.TraceInformation("USING DANGEROUSLY FAST MODE! IF YOU HAVE ANY PROBLEMS, TRY A REBUILD/CLEAN!!!!");
+						Log.TraceInformation("==============================================================================");
+
+						// copy the executable
+						string RemoteShadowDirectoryMac = ConvertPath(Path.GetDirectoryName(Target.OutputPath));
+						string FinalRemoteExecutablePath = String.Format("{0}/Payload/{1}.app/{1}", RemoteShadowDirectoryMac, Target.GameName);
+						RPCUtilHelper.Command("/", String.Format("cp -f {0} {1}", RemoteExecutablePath, FinalRemoteExecutablePath), "", null);
+					}
+					else if (!Utils.IsRunningOnMono && BuildHostPlatform.Current.Platform != UnrealTargetPlatform.Mac)
+					{
+						RPCUtilHelper.CopyFile(RemoteExecutablePath, FilePath, false);
+
+						if (BuildConfiguration.bGeneratedSYMFile == true)
+						{
+							string DSYMExt = ".app.dSYM.zip";
+							RPCUtilHelper.CopyFile(RemoteExecutablePath + DSYMExt, FilePath + DSYMExt, false);
+						}
 					}
 				}
 
