Use a fixedfilename for OBB in APK (main.obb.png) instead of patching JavaBuildSettings.java
#android
#ue4

diff --git a/Engine/Build/Android/Java/src/com/epicgames/ue4/GameActivity.java b/Engine/Build/Android/Java/src/com/epicgames/ue4/GameActivity.java
--- a/Engine/Build/Android/Java/src/com/epicgames/ue4/GameActivity.java
+++ b/Engine/Build/Android/Java/src/com/epicgames/ue4/GameActivity.java
@@ -744,11 +744,6 @@ public class GameActivity extends NativeActivity
 		return JavaBuildSettings.PackageType.AMAZON == JavaBuildSettings.PACKAGING;
 	}
 
-	public static String GetOBBFilename()
-	{
-		return JavaBuildSettings.OBBFILENAME;
-	}
-
 	public void AndroidThunkJava_Minimize()
 	{
 		Intent startMain = new Intent(Intent.ACTION_MAIN);
diff --git a/Engine/Source/Programs/UnrealBuildTool/Android/UEDeployAndroid.cs b/Engine/Source/Programs/UnrealBuildTool/Android/UEDeployAndroid.cs
--- a/Engine/Source/Programs/UnrealBuildTool/Android/UEDeployAndroid.cs
+++ b/Engine/Source/Programs/UnrealBuildTool/Android/UEDeployAndroid.cs
@@ -261,47 +261,31 @@ namespace UnrealBuildTool.Android
             return Path.Combine(GetUE4JavaFilePath(EngineDirectory), "JavaBuildSettings.java");
         }
 
-        public void WriteJavaBuildSettingsFile(string FileName, bool OBBinAPK, string CookFlavor)
+        public void WriteJavaBuildSettingsFile(string FileName, bool OBBinAPK)
         {
              // (!UEBuildConfiguration.bOBBinAPK ? "PackageType.AMAZON" : /*bPackageForGoogle ? "PackageType.GOOGLE" :*/ "PackageType.DEVELOPMENT") + ";\n");
             string Setting = OBBinAPK ? "AMAZON" : "DEVELOPMENT";
-            string OBBFilename = "Android" + CookFlavor + ".obb" + (OBBinAPK ? ".png" : "");
-            if (!File.Exists(FileName) || ShouldWriteJavaBuildSettingsFile(FileName, Setting, OBBFilename))
+            if (!File.Exists(FileName) || ShouldWriteJavaBuildSettingsFile(FileName, Setting))
             {
 				Log.TraceInformation("\n===={0}====WRITING JAVABUILDSETTINGS.JAVA========================================================", DateTime.Now.ToString());
 				StringBuilder BuildSettings = new StringBuilder("package com.epicgames.ue4;\npublic class JavaBuildSettings\n{\n");
                 BuildSettings.Append("\tpublic enum PackageType {AMAZON, GOOGLE, DEVELOPMENT};\n");
                 BuildSettings.Append("\tpublic static final PackageType PACKAGING = PackageType." + Setting + ";\n");
-                BuildSettings.Append("\tpublic static final String OBBFILENAME = \"" + OBBFilename + "\";\n");                    
                 BuildSettings.Append("}\n");
                 File.WriteAllText(FileName, BuildSettings.ToString());
             }
         }
 
-        public bool ShouldWriteJavaBuildSettingsFile(string FileName, string setting, string OBBFilename)
+        public bool ShouldWriteJavaBuildSettingsFile(string FileName, string setting)
         {
             var fileContent = File.ReadAllLines(FileName);
-
-            // compare packaging type and OBB filenames
-            int location;
-            bool packageTypeMatch = false;
-            bool obbFilenameMatch = false;
-            foreach (var fileLine in fileContent)
-            {
-                if ((location = fileLine.IndexOf("PACKAGING = ")) != -1)
-                {
-                    location += 12 + 12; // + ("PACKAGING = ") + ("PackageType.")
-                    packageTypeMatch = String.Compare(setting, fileLine.Substring(location, Math.Min(fileLine.Length - location - 1, setting.Length))) == 0;
-                }
-                else if ((location = fileLine.IndexOf("OBBFILENAME = ")) != -1)
-                {
-                    location += 15; // + ("OBBFILENAME = \"")
-                    obbFilenameMatch = String.Compare(OBBFilename, fileLine.Substring(location, Math.Min(fileLine.Length - location - 2, OBBFilename.Length))) == 0;
-                }
-            }
-
-            // write required if both don't match
-            return !(packageTypeMatch && obbFilenameMatch);
+            if (fileContent.Length < 5)
+                return true;
+            var packageLine = fileContent[4]; // We know this to be true... because we write it below...
+            int location = packageLine.IndexOf("PACKAGING") + 12 + 12; // + ("PACKAGING = ") + ("PackageType.")
+            if (location == -1)
+                return true;
+            return String.Compare(setting, packageLine.Substring(location, Math.Min(packageLine.Length - location - 1, setting.Length))) != 0;
         }
 
 		private static string GetNDKArch(string UE4Arch)
@@ -601,7 +585,7 @@ namespace UnrealBuildTool.Android
 
             // See if we need to create a 'default' Java Build settings file if one doesn't exist (if it does exist we have to assume it has been setup correctly)
             string UE4JavaBuildSettingsFileName = GetUE4JavaBuildSettingsFileName(EngineDirectory);
-            WriteJavaBuildSettingsFile(UE4JavaBuildSettingsFileName, UEBuildConfiguration.bOBBinAPK, CookFlavor);
+            WriteJavaBuildSettingsFile(UE4JavaBuildSettingsFileName, UEBuildConfiguration.bOBBinAPK);
 
 			// check to see if any "meta information" is newer than last time we build
 			string CurrentBuildSettings = GetAllBuildSettings(UE4BuildPath, bForDistribution, bMakeSeparateApks, UEBuildConfiguration.bOBBinAPK);
@@ -670,7 +654,7 @@ namespace UnrealBuildTool.Android
 					Directory.CreateDirectory(UE4BuildPath);
 					Directory.CreateDirectory(ObbFileDestination);
 					Console.WriteLine("Obb file exists...");
-					var DestFileName = Path.Combine(ObbFileDestination, Path.GetFileName(ObbFileLocation) + ".png"); // Need a rename to turn off compression
+					var DestFileName = Path.Combine(ObbFileDestination, "main.obb.png"); // Need a rename to turn off compression
 					var SrcFileName = ObbFileLocation;
 					if (!File.Exists(DestFileName) || File.GetLastWriteTimeUtc(DestFileName) < File.GetLastWriteTimeUtc(SrcFileName))
 					{
diff --git a/Engine/Source/Runtime/Core/Private/Android/AndroidFile.cpp b/Engine/Source/Runtime/Core/Private/Android/AndroidFile.cpp
--- a/Engine/Source/Runtime/Core/Private/Android/AndroidFile.cpp
+++ b/Engine/Source/Runtime/Core/Private/Android/AndroidFile.cpp
@@ -43,9 +43,6 @@ FString GFontPathBase;
 // Is the OBB in an APK file or not
 bool GOBBinAPK;
 
-// Filename of the OBB (used for OBBinAPK case to avoid close search)
-FString GOBBFilename;
-
 extern jobject AndroidJNI_GetJavaAssetManager();
 extern AAssetManager * AndroidThunkCpp_GetAssetManager();
 
@@ -665,8 +662,8 @@ public:
 			//FMountOBBVisitor MountOBB(*this);
 			//IterateDirectory(TEXT(""), MountOBB, false, true);
 
-			// Getting the OBB to mount faster on some devices (IterateDirectory above can be slow!)
-			MountOBB(*GOBBFilename);
+			// Not searching for the OBB to mount is faster on some devices (IterateDirectory above can be slow!)
+			MountOBB(TEXT("main.obb.png"));
 		}
 		else
 		{
diff --git a/Engine/Source/Runtime/Launch/Private/Android/AndroidJNI.cpp b/Engine/Source/Runtime/Launch/Private/Android/AndroidJNI.cpp
--- a/Engine/Source/Runtime/Launch/Private/Android/AndroidJNI.cpp
+++ b/Engine/Source/Runtime/Launch/Private/Android/AndroidJNI.cpp
@@ -19,7 +19,6 @@ static IVirtualKeyboardEntry *VirtualKeyboardWidget = NULL;
 extern FString GFilePathBase;
 extern FString GFontPathBase;
 extern bool GOBBinAPK;
-extern FString GOBBFilename;
 
 //////////////////////////////////////////////////////////////////////////
 
@@ -515,13 +514,6 @@ JNIEXPORT jint JNI_OnLoad(JavaVM* InJavaVM, void* InReserved)
 	jmethodID isOBBInAPKMethod = Env->GetStaticMethodID(FJavaWrapper::GameActivityClassID, "isOBBInAPK", "()Z");
 	GOBBinAPK = (bool)Env->CallStaticBooleanMethod(FJavaWrapper::GameActivityClassID, isOBBInAPKMethod, nullptr);
 
-	// Get the OBB filename
-	jmethodID GetOBBFilenameMethod = Env->GetStaticMethodID(FJavaWrapper::GameActivityClassID, "GetOBBFilename", "()Ljava/lang/String;");
-	jstring ObbString = (jstring)Env->CallStaticObjectMethod(FJavaWrapper::GameActivityClassID, GetOBBFilenameMethod, nullptr);
-	const char * NativeObbString = Env->GetStringUTFChars(ObbString, 0);
-	GOBBFilename = FString(NativeObbString);
-	Env->ReleaseStringUTFChars(ObbString, NativeObbString);
-
 	// Get the system font directory
 	jstring fontPath = (jstring)Env->CallStaticObjectMethod(FJavaWrapper::GameActivityClassID, FJavaWrapper::AndroidThunkJava_GetFontDirectory);
 	const char * nativeFontPathString = Env->GetStringUTFChars(fontPath, 0);
