GitHub pull request #335

diff --git a/Engine/Source/Programs/AutomationTool/DeploymentContext.cs b/Engine/Source/Programs/AutomationTool/DeploymentContext.cs
--- a/Engine/Source/Programs/AutomationTool/DeploymentContext.cs
+++ b/Engine/Source/Programs/AutomationTool/DeploymentContext.cs
@@ -186,6 +186,21 @@ public class DeploymentContext //: ProjectParams
 	/// </summary>
 	public string ShortProjectName;
 
+	private PluginRules[] _CachedPluginRules = null;
+	/// <summary>
+	/// UBT plugin rules which apply to this project
+	/// </summary>
+	public PluginRules[] ProjectPluginRules
+	{
+		get
+		{
+			if(_CachedPluginRules == null)
+			{
+				_CachedPluginRules = Plugins.GetPluginRulesForProject(RawProjectPath, new UnrealTargetPlatform[] { StageTargetPlatform.PlatformType });
+			}
+			return _CachedPluginRules;
+		}
+	}
 
 	public DeploymentContext(
 		string RawProjectPathOrName,
diff --git a/Engine/Source/Programs/AutomationTool/ProjectUtils.cs b/Engine/Source/Programs/AutomationTool/ProjectUtils.cs
--- a/Engine/Source/Programs/AutomationTool/ProjectUtils.cs
+++ b/Engine/Source/Programs/AutomationTool/ProjectUtils.cs
@@ -327,9 +327,16 @@ namespace AutomationTool
 					CommandUtils.LogVerbose("  {0}", Filename);
 				}
 
+				var AssemblySourceFiles = new List<string>();
+				AssemblySourceFiles.AddRange(TargetScripts);
+
+				// Add UBT plugins
+				var PluginSourceFiles = RulesCompiler.FindAllRulesSourceFiles(RulesCompiler.RulesFileType.Plugin, ExtraSearchPaths);
+				AssemblySourceFiles.AddRange(PluginSourceFiles);
+
 				// Check if the scripts require compilation
 				bool DoNotCompile = false;
-				if (!CommandUtils.IsBuildMachine && !CheckIfScriptAssemblyIsOutOfDate(TargetsDllFilename, TargetScripts))
+				if (!CommandUtils.IsBuildMachine && !CheckIfScriptAssemblyIsOutOfDate(TargetsDllFilename, AssemblySourceFiles))
 				{
 					Log.TraceInformation("Targets DLL {0} is up to date.", TargetsDllFilename);
 					DoNotCompile = true;
@@ -343,7 +350,7 @@ namespace AutomationTool
 					}
 				}
 
-				CompileAndLoadTargetsAssembly(Properties, TargetsDllFilename, DoNotCompile, TargetScripts);
+				CompileAndLoadTargetsAssembly(Properties, TargetsDllFilename, DoNotCompile, AssemblySourceFiles);
 			}
 		}
 
@@ -370,7 +377,7 @@ namespace AutomationTool
 			foreach (Type TargetType in AllCompiledTypes)
 			{
 				// Find TargetRules but skip all "UE4Editor", "UE4Game" targets.
-				if (typeof(TargetRules).IsAssignableFrom(TargetType))
+				if (typeof(TargetRules).IsAssignableFrom(TargetType) && !TargetType.IsAbstract)
 				{
 					// Create an instance of this type
 					CommandUtils.LogVerbose("Creating target rules object: {0}", TargetType.Name);
diff --git a/Engine/Source/Programs/AutomationTool/ScriptCompiler.cs b/Engine/Source/Programs/AutomationTool/ScriptCompiler.cs
--- a/Engine/Source/Programs/AutomationTool/ScriptCompiler.cs
+++ b/Engine/Source/Programs/AutomationTool/ScriptCompiler.cs
@@ -135,11 +135,23 @@ namespace AutomationTool
 					BuildFolders.Add(GameBuildFolder);
 				}
 			}
+
 			RemappedAdditionalScriptFolders.AddRange(BuildFolders);
 
 			Log.TraceVerbose("Discovering game folders.");
 			UnrealBuildTool.RulesCompiler.SetAssemblyNameAndGameFolders("UnrealAutomationToolRules", AllGameFolders);
 
+			// Search in plugin folders
+			foreach (var PluginObject in Plugins.AllPlugins)
+			{
+				var PluginFolder = PluginObject.Directory;
+				var PluginBuildFolder = Path.Combine(PluginFolder, "Build");
+				if (Directory.Exists(PluginBuildFolder))
+				{
+					RemappedAdditionalScriptFolders.Add(PluginBuildFolder);
+				}
+			}
+
 			var DiscoveredModules = UnrealBuildTool.RulesCompiler.FindAllRulesSourceFiles(UnrealBuildTool.RulesCompiler.RulesFileType.AutomationModule, RemappedAdditionalScriptFolders);
 			var ModulesToCompile = new List<string>(DiscoveredModules.Count);
 			foreach (var ModuleFilename in DiscoveredModules)
diff --git a/Engine/Source/Programs/AutomationTool/Scripts/CopyBuildToStagingDirectory.Automation.cs b/Engine/Source/Programs/AutomationTool/Scripts/CopyBuildToStagingDirectory.Automation.cs
--- a/Engine/Source/Programs/AutomationTool/Scripts/CopyBuildToStagingDirectory.Automation.cs
+++ b/Engine/Source/Programs/AutomationTool/Scripts/CopyBuildToStagingDirectory.Automation.cs
@@ -22,6 +22,14 @@ public partial class Project : CommandUtils
 
 	#region Utilities
 
+	/// <summary>
+	/// Rules for staging additional plugin files for packaged builds
+	/// </summary>
+	public abstract class PluginStagingRules
+	{
+		public abstract void StagePluginFiles(ProjectParams Params, DeploymentContext SC);
+	}
+
 	/// <summary>
 	/// Writes a pak response file to disk
 	/// </summary>
@@ -380,6 +388,36 @@ public partial class Project : CommandUtils
 				}
 			}
 		}
+
+		// Allow plugins to stage necessary files
+		foreach(var PluginRulesObject in SC.ProjectPluginRules)
+		{
+			var StagingRulesTypeName = PluginRulesObject.StagingRulesTypeName;
+
+			if(null != StagingRulesTypeName)
+			{
+				var PluginStagingRulesType = Type.GetType(StagingRulesTypeName);
+				
+				if(null == PluginStagingRulesType)
+				{
+					throw new AutomationException("Excepted to find plugin staging rules type '{0}'", StagingRulesTypeName);
+				}
+
+				if(!typeof(PluginStagingRules).IsAssignableFrom(PluginStagingRulesType))
+				{
+					throw new AutomationException("Expected plugin staging rules type '{0}' to be derived from PluginStagingRules", StagingRulesTypeName);
+				}
+
+				var PluginStagingRulesObject = (PluginStagingRules)Activator.CreateInstance(PluginStagingRulesType);
+				if(null == PluginStagingRulesObject)
+				{
+					throw new AutomationException("Could not create plugin staging rules of type '{0}'", StagingRulesTypeName);
+				}
+
+				PluginStagingRulesObject.StagePluginFiles(Params, SC);
+			}
+
+		}
 	}
 
 	public static void DumpTargetManifest(Dictionary<string, string> Mapping, string Filename)
diff --git a/Engine/Source/Programs/UnrealBuildTool/Configuration/UEBuildTarget.cs b/Engine/Source/Programs/UnrealBuildTool/Configuration/UEBuildTarget.cs
--- a/Engine/Source/Programs/UnrealBuildTool/Configuration/UEBuildTarget.cs
+++ b/Engine/Source/Programs/UnrealBuildTool/Configuration/UEBuildTarget.cs
@@ -1569,6 +1569,23 @@ namespace UnrealBuildTool
 				OutputItems.AddRange(Binary.Build(TargetToolChain, GlobalCompileEnvironment, GlobalLinkEnvironment));
 			}
 
+			// Run post-build plugin hook for enabled plugins
+			foreach (var Plugin in BuildPlugins)
+			{
+				if(null != Plugin.Rules)
+				{
+					var PluginOutputItems = new List<FileItem>();
+					ECompilationResult PluginResult = Plugin.Rules.PostBuildTarget(this, PluginOutputItems);
+					if(ECompilationResult.Succeeded != PluginResult)
+					{
+						var TargetName = !String.IsNullOrEmpty( GameName ) ? GameName : AppName;
+						Log.TraceInformation("Plugin {0} post build failed for target '{1}' (platform: {2})", Plugin.Name, TargetName, Platform.ToString());
+						return PluginResult;
+					}
+					OutputItems.AddRange(PluginOutputItems);
+				}
+			}
+
 			return ECompilationResult.Succeeded;
 		}
 
@@ -2828,7 +2845,13 @@ namespace UnrealBuildTool
 					if (ModuleType == UEBuildModuleType.Unknown)
 					{
 						// not a plugin, see if it is a game module 
-						if (ModuleFileRelativeToEngineDirectory.StartsWith("..") || Path.IsPathRooted(ModuleFileRelativeToEngineDirectory))
+						// If it starts with ".."
+						if (ModuleFileRelativeToEngineDirectory.StartsWith("..")
+							// or an absolute path
+							|| Path.IsPathRooted(ModuleFileRelativeToEngineDirectory)
+							// or is in a engine Plugin Templates directory
+							|| (ModuleFileRelativeToEngineDirectory.StartsWith("Plugins" + Path.DirectorySeparatorChar)  
+								&& -1 != ModuleFileRelativeToEngineDirectory.IndexOf(Path.DirectorySeparatorChar + "Templates" + Path.DirectorySeparatorChar)))
 						{
 							ApplicationOutputPath = ModuleFileName;
 							int SourceIndex = ApplicationOutputPath.LastIndexOf(Path.DirectorySeparatorChar + "Source" + Path.DirectorySeparatorChar);
diff --git a/Engine/Source/Programs/UnrealBuildTool/System/ExternalExecution.cs b/Engine/Source/Programs/UnrealBuildTool/System/ExternalExecution.cs
--- a/Engine/Source/Programs/UnrealBuildTool/System/ExternalExecution.cs
+++ b/Engine/Source/Programs/UnrealBuildTool/System/ExternalExecution.cs
@@ -314,7 +314,7 @@ namespace UnrealBuildTool
 		 * 
 		 * @return					True if the code files are out of date
 		 * */
-		private static bool AreGeneratedCodeFilesOutOfDate(List<UHTModuleInfo> UObjectModules)
+		private static bool AreGeneratedCodeFilesOutOfDate(UEBuildTarget Target, List<UHTModuleInfo> UObjectModules)
 		{
 			bool bIsOutOfDate = false;
 
@@ -405,6 +405,22 @@ namespace UnrealBuildTool
 									}
 								}
 							}
+
+							// Allow plugins a crack at determining if UHT is out of date
+							if(!bIsOutOfDate)
+							{
+								foreach (var PluginObject in Target.BuildPlugins)
+								{
+									if(null != PluginObject.Rules)
+									{
+										if (PluginObject.Rules.AreUHTGeneratedFilesOutOfDate(Target, SavedTimestamp))
+										{
+											bIsOutOfDate = true;
+											break;
+										}
+									}
+								}
+							}
 						}
 						else
 						{
@@ -543,7 +559,7 @@ namespace UnrealBuildTool
 
 
 				// ensure the headers are up to date
-				bool bUHTNeedsToRun = (UEBuildConfiguration.bForceHeaderGeneration == true || AreGeneratedCodeFilesOutOfDate(UObjectModules));
+				bool bUHTNeedsToRun = (UEBuildConfiguration.bForceHeaderGeneration == true || AreGeneratedCodeFilesOutOfDate(Target, UObjectModules));
 				if( bUHTNeedsToRun || UnrealBuildTool.IsGatheringBuild )
 				{
 					// Since code files are definitely out of date, we'll now finish computing information about the UObject modules for UHT.  We
diff --git a/Engine/Source/Programs/UnrealBuildTool/System/Plugins.cs b/Engine/Source/Programs/UnrealBuildTool/System/Plugins.cs
--- a/Engine/Source/Programs/UnrealBuildTool/System/Plugins.cs
+++ b/Engine/Source/Programs/UnrealBuildTool/System/Plugins.cs
@@ -60,6 +60,21 @@ namespace UnrealBuildTool
 		// Where does this plugin live?
 		public LoadedFromType LoadedFrom;
 
+		private PluginRules _Rules = null;
+		private bool InitializedRules = false;
+		public PluginRules Rules
+		{
+			get
+			{
+				if(!InitializedRules)
+				{
+					_Rules = RulesCompiler.CreatePluginRules(Name);
+					InitializedRules = true;
+				}
+				return _Rules;
+			}
+		}
+
 		public override string ToString()
 		{
 			return Path.Combine(Directory, Name + ".uplugin");
@@ -471,6 +486,25 @@ namespace UnrealBuildTool
 			}
 		}
 
+		/// <summary>
+		/// Discovers all plugin template directories
+		/// </summary>
+		private static void DiscoverAllPluginTemplateDirectories()
+		{
+			if (AllPluginTemplateDirectoriesVar == null)
+			{
+				AllPluginTemplateDirectoriesVar = new List<string>();
+
+				foreach (var CurPluginInfo in AllPlugins)
+				{
+					string PluginTemplatePath = Path.Combine(CurPluginInfo.Directory, "Templates");
+					if(Directory.Exists(PluginTemplatePath))
+					{
+						AllPluginTemplateDirectoriesVar.Add(PluginTemplatePath);
+					}
+				}
+			}
+		}
 
 		/// <summary>
 		/// Returns true if the specified module name is part of a plugin
@@ -501,6 +535,42 @@ namespace UnrealBuildTool
 			}
 		}
 
+		/// <summary>
+		/// Gets all the applicable UBT plugin rules for the specified uproject
+		/// </summary>
+		/// <param name="ProjectFilePath">Path to uproject file</param>
+		/// <param name="SupportedPlatforms">Platforms we're interested in</param>
+		/// <returns></returns>
+		public static PluginRules[] GetPluginRulesForProject(string ProjectFilePath, IEnumerable<UnrealTargetPlatform> SupportedPlatforms)
+		{
+			var EnabledPluginNames = new HashSet<string>(StringComparer.InvariantCultureIgnoreCase);
+
+			// add in engine plugins
+			foreach (var PluginObject in Plugins.AllPlugins)
+			{
+				if (PluginObject.bEnabledByDefault)
+				{
+					EnabledPluginNames.Add(PluginObject.Name);
+				}
+			}
+
+			foreach (var Platform in SupportedPlatforms)
+			{
+				var FilterPluginNames = UProjectInfo.GetEnabledPlugins(ProjectFilePath, new List<string>(), Platform);
+
+				foreach (var PluginName in FilterPluginNames)
+				{
+					EnabledPluginNames.Add(PluginName);
+				}
+			}
+
+			var EnabledPlugins = Plugins.AllPlugins.Where(x => EnabledPluginNames.Contains(x.Name));
+
+			return (from EnabledPlugin in EnabledPlugins
+					where EnabledPlugin.Rules != null
+					select EnabledPlugin.Rules).ToArray();
+
+		}
 	
 		/// Access the list of all plugins.  We'll scan for plugins when this is called the first time.
 		public static List< PluginInfo > AllPlugins
@@ -522,12 +592,24 @@ namespace UnrealBuildTool
 			}
 		}
 
-	
+		/// Access all plugin template directories
+		public static List<string> AllPluginTemplateDirectories
+		{
+			get
+			{
+				DiscoverAllPluginTemplateDirectories();
+				return AllPluginTemplateDirectoriesVar;
+			}
+		}
 
 		/// List of all plugins we've found so far in this session
 		private static List< PluginInfo > AllPluginsVar = null;
 
 		/// Maps plugin modules to the plugin that owns them
 		private static Dictionary< string, PluginInfo > ModulesToPluginMapVar = null;
+
+		/// List of all plugin template directories
+		private static List<string> AllPluginTemplateDirectoriesVar = null;
+
 	}
 }
diff --git a/Engine/Source/Programs/UnrealBuildTool/System/ProjectFileGenerator.cs b/Engine/Source/Programs/UnrealBuildTool/System/ProjectFileGenerator.cs
--- a/Engine/Source/Programs/UnrealBuildTool/System/ProjectFileGenerator.cs
+++ b/Engine/Source/Programs/UnrealBuildTool/System/ProjectFileGenerator.cs
@@ -191,8 +191,19 @@ namespace UnrealBuildTool
 		/// directly from the IDE.
 		bool bGatherThirdPartySource = false;
 
+		private string _MasterProjectName = "UE4";
 		/// Name of the master project file (e.g. base file name for the solution file for Visual Studio, or the Xcode project file on Mac)
-		protected string MasterProjectName = "UE4";
+		public string MasterProjectName
+		{
+			get
+			{
+				return _MasterProjectName;
+			}
+			protected set
+			{
+				_MasterProjectName = value;
+			}
+		}
 
 		/// Maps all module names that were included in generated project files, to actual project file objects.
 		/// @todo projectfiles: Nasty global static list.  This is only really used for IntelliSense, and to avoid extra folder searches for projects we've already cached source files for.
@@ -228,6 +239,17 @@ namespace UnrealBuildTool
 				}
 			}
 
+			// Search in plugin folders
+			foreach(var PluginObject in Plugins.AllPlugins)
+			{
+				var PluginFolder = PluginObject.Directory;
+				var PluginBuildFolder = Path.Combine(PluginFolder, "Build");
+				if(Directory.Exists(PluginBuildFolder))
+				{
+					BuildFolders.Add(PluginBuildFolder);
+				}
+			}
+
 			// Find all the automation modules .csproj files to add
 			var ModuleFiles = RulesCompiler.FindAllRulesSourceFiles(RulesCompiler.RulesFileType.AutomationModule, BuildFolders);
 			foreach (var ProjectFile in ModuleFiles)
@@ -380,7 +402,7 @@ namespace UnrealBuildTool
 			HashSet<ProjectFile> TemplateGameProjects = null;
 			{
 				// Setup buildable projects for all targets
-				AddProjectsForAllTargets( AllGameProjects, out EngineProject, out GameProjects, out ProgramProjects, out TemplateGameProjects );
+				AddProjectsForAllTargets(AllGameProjects, out EngineProject, out GameProjects, out ProgramProjects, out TemplateGameProjects);
 
 				// Add all game projects and game config files
 				AddAllGameProjects(GameProjects, SupportedPlatformNames, RootFolder);
@@ -506,6 +528,9 @@ namespace UnrealBuildTool
                     AddPS4Projects( ProgramsFolder );
                 }
 
+				// Allow plugins to add any custom projects
+				var PerProjectPluginRules = GatherPluginRulesForProjects(AllGameProjects);
+				AddCustomPluginProjects(PerProjectPluginRules);
 
 				// Eliminate all redundant master project folders.  E.g., folders which contain only one project and that project
 				// has the same name as the folder itself.  To the user, projects "feel like" folders already in the IDE, so we
@@ -573,6 +598,8 @@ namespace UnrealBuildTool
 						WriteProjectFiles();
 
 						Log.TraceVerbose( "Project generation complete ({0} generated, {1} imported)", GeneratedProjectFiles.Count, OtherProjectFiles.Count );
+
+						PostWriteProjects(PerProjectPluginRules);
 					}
 					else
 					{
@@ -946,7 +973,79 @@ namespace UnrealBuildTool
 			return AllTargetFiles;
 		}
 
+		Dictionary<UProjectInfo, PluginRules[]> GatherPluginRulesForProjects(List<UProjectInfo> AllGameProjects)
+		{
+			var PerProjectPluginRules = new Dictionary<UProjectInfo, PluginRules[]>();
+
+			foreach(var GameProject in AllGameProjects)
+			{
+				var PluginRulesArray = Plugins.GetPluginRulesForProject(GameProject.FilePath, SupportedPlatforms);
+
+				PerProjectPluginRules[GameProject] = PluginRulesArray;
+			}
 
+			return PerProjectPluginRules;
+		}
+
+		static IEnumerable<PluginRules> GatherUniquePluginRules(Dictionary<UProjectInfo, PluginRules[]> PerProjectPluginRules)
+		{
+			var AllPlugins = new HashSet<PluginRules>();
+
+			foreach (var ProjectPluginRules in PerProjectPluginRules)
+			{
+				foreach (var PluginRulesObject in ProjectPluginRules.Value)
+				{
+					AllPlugins.Add(PluginRulesObject);
+				}
+			}
+
+			return AllPlugins;
+		}
+
+		/// <summary>
+		/// Allow plugins to generate any custom projects
+		/// </summary>
+		/// <param name="PerProjectPluginRules">Plugin rules file per game project</param>
+		void AddCustomPluginProjects(Dictionary<UProjectInfo, PluginRules[]> PerProjectPluginRules)
+		{
+			var GameProjects = PerProjectPluginRules.Keys;
+
+			// first allow any plugin to clear out any state if we're doing multiple project file generation passes
+			foreach(var PluginRulesObject in GatherUniquePluginRules(PerProjectPluginRules))
+			{
+				PluginRulesObject.PreGenerateCustomProjects(GameProjects);
+			}
+
+			foreach(var ProjectPluginRules in PerProjectPluginRules)
+			{
+				foreach(var PluginRulesObject in ProjectPluginRules.Value)
+				{
+					PluginRulesObject.AddCustomGeneratedProjects(this, ProjectPluginRules.Key);
+				}
+			}
+		}
+
+		/// <summary>
+		/// Issue plugin post-write projects callbacks
+		/// </summary>
+		/// <param name="PerProjectPluginRules">Plugin rules file per game project</param>
+		void PostWriteProjects(Dictionary<UProjectInfo, PluginRules[]> PerProjectPluginRules)
+		{
+			// Called once per game project
+			foreach (var ProjectPluginRules in PerProjectPluginRules)
+			{
+				foreach (var PluginRulesObject in ProjectPluginRules.Value)
+				{
+					PluginRulesObject.PostWriteProject(this, ProjectPluginRules.Key);
+				}
+			}
+
+			// now call a final callback for the master project itself
+			foreach(var PluginRulesObject in GatherUniquePluginRules(PerProjectPluginRules))
+			{
+				PluginRulesObject.PostWriteMasterProject(this);
+			}
+		}
 	
 		/// <summary>
 		/// Recursively collapses all sub-folders that are redundant.  Should only be called after we're done adding
@@ -1410,6 +1509,15 @@ namespace UnrealBuildTool
 		/// <param name="bGatherThirdPartySource">True to gather source code from third party projects too</param>
 		protected void AddProjectsForAllModules( List<UProjectInfo> AllGames, Dictionary<string, ProjectFile> ProgramProjects, List<string> AllModuleFiles, bool bGatherThirdPartySource )
 		{
+			// Get all template root directories
+			var TemplateRootPaths = new List<string>();
+
+			// Add engine templates directory
+			TemplateRootPaths.Add(Path.Combine(RootRelativePath, "Templates"));
+
+			// Add plugin templates directories
+			TemplateRootPaths.AddRange(Plugins.AllPluginTemplateDirectories);
+
 			foreach( var CurModuleFile in AllModuleFiles )
 			{
 				Log.TraceVerbose("AddProjectsForAllModules " + CurModuleFile);
@@ -1423,20 +1531,34 @@ namespace UnrealBuildTool
 				// We'll keep track of whether this is an "engine" or "external" module.  This is determined below while loading module rules.
 				bool IsEngineModule = false;
 				bool IsThirdPartyModule = false;
+				bool IsTemplateModule = false;
 
-				// Check to see if this is an Engine module.  That is, the module is located under the "Engine" folder
-				string ModuleFileRelativeToEngineDirectory = Utils.MakePathRelativeTo(CurModuleFile, Path.Combine(EngineRelativePath));
-				if (!ModuleFileRelativeToEngineDirectory.StartsWith( ".." ) && !Path.IsPathRooted( ModuleFileRelativeToEngineDirectory ))
+				// Check to see if this is a template module.
+				foreach (var TemplateRootPath in TemplateRootPaths)
 				{
-					// This is an engine module
-					IsEngineModule = true;
+					if (Utils.IsFileUnderDirectory(CurModuleFile, TemplateRootPath))
+					{
+						IsTemplateModule = true;
+						break;
+					}
 				}
 
-				string ModuleFileRelativeToThirdPartyDirectory = Utils.MakePathRelativeTo(CurModuleFile, Path.Combine(EngineRelativePath, "Source", "ThirdParty" ));
-				if (!ModuleFileRelativeToThirdPartyDirectory.StartsWith( ".." ) && !Path.IsPathRooted( ModuleFileRelativeToThirdPartyDirectory ))
+				if (!IsTemplateModule)
 				{
-					// This is a third partymodule
-					IsThirdPartyModule = true;
+					// Check to see if this is an Engine module.  That is, the module is located under the "Engine" folder
+					string ModuleFileRelativeToEngineDirectory = Utils.MakePathRelativeTo(CurModuleFile, Path.Combine(EngineRelativePath));
+					if (!ModuleFileRelativeToEngineDirectory.StartsWith("..") && !Path.IsPathRooted(ModuleFileRelativeToEngineDirectory))
+					{
+						// This is an engine module
+						IsEngineModule = true;
+					}
+
+					string ModuleFileRelativeToThirdPartyDirectory = Utils.MakePathRelativeTo(CurModuleFile, Path.Combine(EngineRelativePath, "Source", "ThirdParty"));
+					if (!ModuleFileRelativeToThirdPartyDirectory.StartsWith("..") && !Path.IsPathRooted(ModuleFileRelativeToThirdPartyDirectory))
+					{
+						// This is a third partymodule
+						IsThirdPartyModule = true;
+					}
 				}
 
 				bool IncludePrivateSourceCode = true;
@@ -1556,7 +1678,7 @@ namespace UnrealBuildTool
 		/// <param name="GameProjects">Map of game folder name to all of the game projects we created</param>
 		/// <param name="ProgramProjects">Map of program names to all of the program projects we created</param>
 		/// <param name="TemplateGameProjects">Set of template game projects we found.  These will also be in the GameProjects map</param>
-		private void AddProjectsForAllTargets( List<UProjectInfo> AllGames, out ProjectFile EngineProject, out Dictionary<string, ProjectFile> GameProjects, out Dictionary<string, ProjectFile> ProgramProjects, out HashSet<ProjectFile> TemplateGameProjects )
+		private void AddProjectsForAllTargets( List<UProjectInfo> AllGames, out ProjectFile EngineProject, out Dictionary<string, ProjectFile> GameProjects, out Dictionary<string, ProjectFile> ProgramProjects, out HashSet<ProjectFile> TemplateGameProjects)
 		{
 			// As we're creating project files, we'll also keep track of whether we created an "engine" project and return that if we have one
 			EngineProject = null;
@@ -1564,22 +1686,46 @@ namespace UnrealBuildTool
 			ProgramProjects = new Dictionary<string,ProjectFile>( StringComparer.InvariantCultureIgnoreCase );
 			TemplateGameProjects = new HashSet<ProjectFile>();
 
-
 			// Find all of the target files.  This will filter out any modules or targets that don't
 			// belong to platforms we're generating project files for.
 			var AllTargetFiles = DiscoverTargets();
 
+			// Get all template root directories
+			var TemplateRootPaths = new List<string>();
+
+			// Add engine templates directory
+			TemplateRootPaths.Add(Path.Combine(RootRelativePath, "Templates"));
+
+			// Add plugin templates directories
+			TemplateRootPaths.AddRange(Plugins.AllPluginTemplateDirectories);
+
 			foreach( var TargetFilePath in AllTargetFiles )
 			{
 				var TargetName = Utils.GetFilenameWithoutAnyExtensions(TargetFilePath);		// Remove both ".cs" and ".Target"
 
+				// Check to see if this is a template target.  That is, the target is located under the "Templates" folder
+				bool IsTemplateTarget = false;
+				foreach (var TemplateRootPath in TemplateRootPaths)
+				{
+					string TargetFileRelativeToTemplatesDirectory = Utils.MakePathRelativeTo(TargetFilePath, TemplateRootPath);
+					if (!TargetFileRelativeToTemplatesDirectory.StartsWith("..") && !Path.IsPathRooted(TargetFileRelativeToTemplatesDirectory))
+					{
+						IsTemplateTarget = true;
+						break;
+					}
+
+				}
+
 				// Check to see if this is an Engine target.  That is, the target is located under the "Engine" folder
 				bool IsEngineTarget = false;
-				string TargetFileRelativeToEngineDirectory = Utils.MakePathRelativeTo(TargetFilePath, Path.Combine(EngineRelativePath), AlwaysTreatSourceAsDirectory: false);
-				if (!TargetFileRelativeToEngineDirectory.StartsWith( ".." ) && !Path.IsPathRooted( TargetFileRelativeToEngineDirectory ))
+				if (!IsTemplateTarget)
 				{
-					// This is an engine target
-					IsEngineTarget = true;
+					string TargetFileRelativeToEngineDirectory = Utils.MakePathRelativeTo(TargetFilePath, Path.Combine(EngineRelativePath), AlwaysTreatSourceAsDirectory: false);
+					if (!TargetFileRelativeToEngineDirectory.StartsWith("..") && !Path.IsPathRooted(TargetFileRelativeToEngineDirectory))
+					{
+						// This is an engine target
+						IsEngineTarget = true;
+					}
 				}
 
 				bool WantProjectFileForTarget = true;
@@ -1655,15 +1801,6 @@ namespace UnrealBuildTool
 						ProjectFile.IsGeneratedProject = true;
 						ProjectFile.IsStubProject = false;
 
-						bool IsTemplateTarget = false;
-						{
-							// Check to see if this is a template target.  That is, the target is located under the "Templates" folder
-							string TargetFileRelativeToTemplatesDirectory = Utils.MakePathRelativeTo(TargetFilePath, Path.Combine(RootRelativePath, "Templates"));
-							if (!TargetFileRelativeToTemplatesDirectory.StartsWith("..") && !Path.IsPathRooted(TargetFileRelativeToTemplatesDirectory))
-							{
-								IsTemplateTarget = true;
-							}
-						}
 						string BaseFolder = null;
 						if (IsProgramTarget)
 						{
diff --git a/Engine/Source/Programs/UnrealBuildTool/System/RulesCompiler.cs b/Engine/Source/Programs/UnrealBuildTool/System/RulesCompiler.cs
--- a/Engine/Source/Programs/UnrealBuildTool/System/RulesCompiler.cs
+++ b/Engine/Source/Programs/UnrealBuildTool/System/RulesCompiler.cs
@@ -870,8 +870,78 @@ namespace UnrealBuildTool
 		}
 	}
 
+	/// <summary>
+	/// PluginRules contains the rules allowing plugins to hook UBT functionality
+	/// </summary>
+	public abstract class PluginRules
+	{
+		/// <summary>
+		/// Fully qualified type name of the staging rules for this plugin
+		/// </summary>
+		public virtual string StagingRulesTypeName { get { return null; } }
+
+		/// <summary>
+		/// Hook called when determining if generated headers are out of date, allows plugins (mostly script integration ones)
+		/// to customize UHT dependency checking
+		/// </summary>
+		/// <param name="Target">Target being built</param>
+		/// <param name="LastUHTRunTime">Timestamp of last UHT run</param>
+		/// <returns>true UHT products are out of date, false they are not out of date</returns>
+		public virtual bool AreUHTGeneratedFilesOutOfDate(UEBuildTarget Target, DateTime LastUHTRunTime)
+		{
+			return false;
+		}
+
+		/// <summary>
+		/// Hook called after a target has set up its build items
+		/// </summary>
+		/// <param name="Target">Target being built</param>
+		/// <param name="OutputItems">Build items produced by this action</param>
+		/// <returns>Did the plugin succeed in performing its post-build processing?</returns>
+		public virtual ECompilationResult PostBuildTarget(UEBuildTarget Target, List<FileItem> OutputItems)
+		{
+			return ECompilationResult.Succeeded;
+		}
+
+		/// <summary>
+		/// Callback before custom project generation begins. Allows plugin rules to clear any stored state
+		/// </summary>
+		/// <param name="Projects">Game projects that are being generated</param>
+		public virtual void PreGenerateCustomProjects(IEnumerable<UProjectInfo> Projects)
+		{
+
+		}
+		/// <summary>
+		/// Allow plugin to add custom generated projects
+		/// </summary>
+		/// <param name="Generator">Project file generator</param>
+		/// <param name="Project">Game project</param>
+		public virtual void AddCustomGeneratedProjects(ProjectFileGenerator Generator, UProjectInfo Project)
+		{
+
+		}
+
+		/// <summary>
+		/// Callback after project file(s) have been written for a game project
+		/// </summary>
+		/// <param name="Generator">Project file generator</param>
+		/// <param name="Project">Game project</param>
+		public virtual void PostWriteProject(ProjectFileGenerator Generator, UProjectInfo Project)
+		{
+
+		}
+
+		/// <summary>
+		/// Callback after master project file has been written
+		/// </summary>
+		/// <param name="Generator">Project file generator</param>
+		public virtual void PostWriteMasterProject(ProjectFileGenerator Generator)
+		{
+
+		}
+	}
 
-	public class RulesCompiler
+    public class RulesCompiler
 	{
 		/// <summary>
 		/// Helper class to avoid adding extra conditions when getting file extensions and suffixes for
@@ -904,7 +974,11 @@ namespace UnrealBuildTool
 
 			/// *.Automation.csproj files
 			[RulesTypeProperties(Suffix: "Automation", Extension: ".csproj")]
-			AutomationModule
+			AutomationModule,
+
+			/// *.Plugin.cs files
+			[RulesTypeProperties(Suffix: "Plugin", Extension: ".cs")]
+			Plugin
 		}
 
 		class RulesFileCache
@@ -1113,10 +1187,26 @@ namespace UnrealBuildTool
 		/// Maps target names to their actual xxx.Target.cs file on disk
 		private static Dictionary<string, string> TargetNameToTargetFileMap = new Dictionary<string, string>( StringComparer.InvariantCultureIgnoreCase );
 
+		///  Maps plugin names to their actual xxx.Plugin.cs file on disk
+		private static Dictionary<string, string> PluginNameToPluginFileMap = new Dictionary<string, string>(StringComparer.InvariantCultureIgnoreCase);
+
 		/// Map of assembly names we've already compiled and loaded to their list of game folders.  This is used to prevent
 		/// trying to recompile the same assembly when ping-ponging between different types of targets
 		private static Dictionary<Assembly, List<string>> LoadedAssemblyMap = new Dictionary<Assembly, List<string>>();
 
+		///  Helper to initialize a name to file map
+		private static void InitializeNameToRulesFileMap(Dictionary<string, string> FileMap, IEnumerable<string> FileNames)
+		{
+			foreach (var CurFileName in FileNames)
+			{
+				var CleanFileName = Utils.CleanDirectorySeparators(CurFileName);
+				var Name = Path.GetFileNameWithoutExtension(Path.GetFileNameWithoutExtension(CleanFileName));	// Strip both extensions
+				if (!FileMap.ContainsKey(Name))
+				{
+					FileMap.Add(Name, CurFileName);
+				}
+			}
+		}
 
 		private static void ConditionallyCompileAndLoadRulesAssembly()
 		{
@@ -1205,36 +1295,32 @@ namespace UnrealBuildTool
 				var TargetFileNames = FindAllRulesSourceFiles(RulesFileType.Target, AdditionalSearchPaths);
 				AssemblySourceFiles.AddRange( TargetFileNames );
 
+				var PluginFileNames = FindAllRulesSourceFiles(RulesFileType.Plugin, AdditionalSearchPaths);
+				AssemblySourceFiles.AddRange(PluginFileNames);
+
 				// Create a path to the assembly that we'll either load or compile
 				string BaseIntermediatePath = UnrealBuildTool.HasUProjectFile() ?
 					Path.Combine(UnrealBuildTool.GetUProjectPath(), BuildConfiguration.BaseIntermediateFolder) : BuildConfiguration.BaseIntermediatePath;
 
 				string OutputAssemblyPath = Path.GetFullPath(Path.Combine(BaseIntermediatePath, "BuildRules", AssemblyName + ".dll"));
 
-				RulesAssembly = DynamicCompilation.CompileAndLoadAssembly( OutputAssemblyPath, AssemblySourceFiles );
+				var ReferencedAssemblies = new List<string>() 
+					{ 
+					"System.dll", 
+					"System.Core.dll", 
+					"System.Xml.dll"
+					};
+				RulesAssembly = DynamicCompilation.CompileAndLoadAssembly(OutputAssemblyPath, AssemblySourceFiles, ReferencedAssemblies);
 
 				{
 					// Setup the module map
-					foreach( var CurModuleFileName in ModuleFileNames )
-					{
-						var CleanFileName = Utils.CleanDirectorySeparators( CurModuleFileName );
-						var ModuleName = Path.GetFileNameWithoutExtension( Path.GetFileNameWithoutExtension( CleanFileName ) );	// Strip both extensions
-						if( !ModuleNameToModuleFileMap.ContainsKey( ModuleName ) )
-						{
-							ModuleNameToModuleFileMap.Add( ModuleName, CurModuleFileName );
-						}
-					}
+					InitializeNameToRulesFileMap(ModuleNameToModuleFileMap, ModuleFileNames);
 
 					// Setup the target map
-					foreach( var CurTargetFileName in TargetFileNames )
-					{
-						var CleanFileName = Utils.CleanDirectorySeparators( CurTargetFileName );
-						var TargetName = Path.GetFileNameWithoutExtension( Path.GetFileNameWithoutExtension( CleanFileName ) );	// Strip both extensions
-						if( !TargetNameToTargetFileMap.ContainsKey( TargetName ) )
-						{
-							TargetNameToTargetFileMap.Add( TargetName, CurTargetFileName );
-						}
-					}
+					InitializeNameToRulesFileMap(TargetNameToTargetFileMap, TargetFileNames);
+
+					// Setup the plugin map
+					InitializeNameToRulesFileMap(PluginNameToPluginFileMap, PluginFileNames);
 				}
 
 				// Remember that we loaded this assembly
@@ -1836,5 +1922,47 @@ namespace UnrealBuildTool
 			return BuildTarget;
 		}
 
+		/// <summary>
+		/// Creates a plugin rules object for the specified plugin name.
+		/// </summary>
+		/// <param name="PluginName">Name of the plugin</param>
+		/// <returns>The build plugin rules for the plugin</returns>
+		public static PluginRules CreatePluginRules(string PluginName)
+		{
+			ConditionallyCompileAndLoadRulesAssembly();
+			var AssemblyFileName = Path.GetFileNameWithoutExtension(RulesAssembly.Location);
+
+			// Make sure the plugin is known
+			if(!PluginNameToPluginFileMap.ContainsKey(PluginName))
+			{
+				// this is not an error because plugin rules are optional
+				return null;
+			}
+
+			var PluginFileName = PluginNameToPluginFileMap[PluginName];
+
+			var PluginTypeName = "UnrealBuildTool.Plugin.Rules." + PluginName;
+			// Because plugins are the same name as modules, we always look in "UnrealBuildTool.Plugin.Rules" namespace. 
+			var RulesObjectType = RulesAssembly.GetType(PluginTypeName);
+
+			if(RulesObjectType == null)
+			{
+				// This is an error, because it means someone defined a plugin.cs file but did not create the proper type
+				throw new BuildException("Expecting to find a type to be declared in a module rules named '{0}' in {1}. This type must derive from 'PluginRules' type defined by Unreal Build Tool and be in the UnrealBuildTool.Plugin.Rules namespace", PluginTypeName, PluginFileName);
+			}
+
+			// Create an instance of the plugin's rules object
+			PluginRules RulesObject;
+			try
+			{
+				RulesObject = (PluginRules)Activator.CreateInstance(RulesObjectType);
+			}
+			catch(Exception Ex)
+			{
+				throw new BuildException(Ex, "Unable to instantiate instance of '{0}' object type from compiled assembly '{1}'. Unreal Build Tool creates an instance of your plugin's 'Rules' object in order to call your plugin hooks. The CLR exception details may provide more information: {2} ", PluginTypeName, AssemblyFileName, Ex.ToString());
+			}
+
+			return RulesObject;
+		}
 	}
 }
