Fix trying to build all plugins for programs when -BuildRocket is set.

diff --git a/Engine/Source/Programs/UnrealBuildTool/Configuration/UEBuildTarget.cs b/Engine/Source/Programs/UnrealBuildTool/Configuration/UEBuildTarget.cs
--- a/Engine/Source/Programs/UnrealBuildTool/Configuration/UEBuildTarget.cs
+++ b/Engine/Source/Programs/UnrealBuildTool/Configuration/UEBuildTarget.cs
@@ -2450,7 +2450,7 @@ namespace UnrealBuildTool
 			DependentPlugins.AddRange(EnabledPlugins);
 
 			// Set the list of plugins that should be built
-			if(UnrealBuildTool.BuildingRocket() )
+			if( UnrealBuildTool.BuildingRocket() && TargetType != TargetRules.TargetType.Program )
 			{
 				BuildPlugins.AddRange(ValidPlugins);
 			}
