UBT: Removing UEBuildConfiguration.bBuildDedicatedServer

diff --git a/Engine/Source/Programs/UnrealBuildTool/Configuration/UEBuildConfiguration.cs b/Engine/Source/Programs/UnrealBuildTool/Configuration/UEBuildConfiguration.cs
--- a/Engine/Source/Programs/UnrealBuildTool/Configuration/UEBuildConfiguration.cs
+++ b/Engine/Source/Programs/UnrealBuildTool/Configuration/UEBuildConfiguration.cs
@@ -28,7 +28,7 @@ namespace UnrealBuildTool
         public static bool bCompileICU;
 
 		/** Whether to build a stripped down version of the game specifically for dedicated server. */
-		[XmlConfig]
+		[Obsolete("bBuildDedicatedServer has been deprecated and will be removed in future release. Update your code to use TargetInfo.Type instead or your code will not compile.")]
 		public static bool bBuildDedicatedServer;
 
 		/** Whether to compile the editor or not. Only desktop platforms (Windows or Mac) will use this, other platforms force this to false */
@@ -202,7 +202,6 @@ namespace UnrealBuildTool
             bRuntimePhysicsCooking = true;
 			bCompileBox2D = true;
 			bCompileICU = true;
-			bBuildDedicatedServer = false;
 			bBuildEditor = true;
 			bBuildRequiresCookedData = false;
 			bBuildWithEditorOnlyData = true;
diff --git a/Engine/Source/Programs/UnrealBuildTool/Configuration/UEBuildServer.cs b/Engine/Source/Programs/UnrealBuildTool/Configuration/UEBuildServer.cs
--- a/Engine/Source/Programs/UnrealBuildTool/Configuration/UEBuildServer.cs
+++ b/Engine/Source/Programs/UnrealBuildTool/Configuration/UEBuildServer.cs
@@ -76,9 +76,6 @@ namespace UnrealBuildTool
 			ref CPPEnvironmentConfiguration OutCPPEnvironmentConfiguration
 			)
 		{
-			// Dedicated server
-			UEBuildConfiguration.bBuildDedicatedServer = true;
-
 			UEBuildConfiguration.bCompileLeanAndMeanUE = true;
 
 			// Do not include the editor
diff --git a/Engine/Source/Programs/UnrealBuildTool/Configuration/UEBuildTarget.cs b/Engine/Source/Programs/UnrealBuildTool/Configuration/UEBuildTarget.cs
--- a/Engine/Source/Programs/UnrealBuildTool/Configuration/UEBuildTarget.cs
+++ b/Engine/Source/Programs/UnrealBuildTool/Configuration/UEBuildTarget.cs
@@ -2263,9 +2263,10 @@ namespace UnrealBuildTool
 		/** 
 		 * @return true if we are building for dedicated server, false otherwise.
 		 */
+		[Obsolete("IsBuildingDedicatedServer() has been deprecated and will be removed in future release. Update your code to use TargetInfo.Type instead or your code will not compile.")]
 		public bool IsBuildingDedicatedServer()
 		{
-			return UEBuildConfiguration.bBuildDedicatedServer;
+			return false;
 		}
 
 		/** Given a UBT-built binary name (e.g. "Core"), returns a relative path to the binary for the current build configuration (e.g. "../Binaries/Win64/Core-Win64-Debug.lib") */
diff --git a/Engine/Source/Runtime/Engine/Engine.Build.cs b/Engine/Source/Runtime/Engine/Engine.Build.cs
--- a/Engine/Source/Runtime/Engine/Engine.Build.cs
+++ b/Engine/Source/Runtime/Engine/Engine.Build.cs
@@ -171,7 +171,7 @@ public class Engine : ModuleRules
 			}
 		);
 
-		if (!UEBuildConfiguration.bBuildDedicatedServer)
+		if (Target.Type.Value != TargetRules.TargetType.Server)
         {
 		    DynamicallyLoadedModuleNames.AddRange(
 			    new string[] {
diff --git a/Engine/Source/Runtime/Slate/Slate.Build.cs b/Engine/Source/Runtime/Slate/Slate.Build.cs
--- a/Engine/Source/Runtime/Slate/Slate.Build.cs
+++ b/Engine/Source/Runtime/Slate/Slate.Build.cs
@@ -18,7 +18,7 @@ public class Slate : ModuleRules
 			}
 		);
 
-        if (!UEBuildConfiguration.bBuildDedicatedServer && UEBuildConfiguration.bCompileFreeType)
+		if (Target.Type.Value != TargetRules.TargetType.Server && UEBuildConfiguration.bCompileFreeType)
         {
 			AddThirdPartyPrivateStaticDependencies(Target, "FreeType2");
 			Definitions.Add("WITH_FREETYPE=1");
diff --git a/Engine/Source/Runtime/SlateCore/SlateCore.Build.cs b/Engine/Source/Runtime/SlateCore/SlateCore.Build.cs
--- a/Engine/Source/Runtime/SlateCore/SlateCore.Build.cs
+++ b/Engine/Source/Runtime/SlateCore/SlateCore.Build.cs
@@ -35,7 +35,7 @@ public class SlateCore : ModuleRules
 			}
 		);
 
-		if (!UEBuildConfiguration.bBuildDedicatedServer && UEBuildConfiguration.bCompileFreeType)
+		if (Target.Type.Value != TargetRules.TargetType.Server && UEBuildConfiguration.bCompileFreeType)
 		{
 			AddThirdPartyPrivateStaticDependencies(Target, "FreeType2");
 			Definitions.Add("WITH_FREETYPE=1");
