[Important BUNs in here]
- Enabled ES31 for GitHub/P4 users (currently only works on a Tegra K1 with recent system software)
- Always use separate APKs - this means that all APKs will have -armv7-es2 (or whatever) in their name. This may affect scripts, etc people have written
- Fixed a bug when using multiple java libraries (for users who add their own libs)
- Properly collapse Android devices into a single device with multiple variants. LaunchOn will now choose the best format based on the connected device. If you want to override the format, you will need to use the ProjectLauncher
- Added Android serial number to the device type, in case you have multiple devices of the same type (also fixes [UE-3770])

diff --git a/Engine/Build/JunkManifest.txt b/Engine/Build/JunkManifest.txt
--- a/Engine/Build/JunkManifest.txt
+++ b/Engine/Build/JunkManifest.txt
@@ -49,5 +49,3 @@ Engine\Plugins\Script\ScriptPlugin\Binaries\Win32\*
 Engine\Plugins\Script\ScriptPlugin\Binaries\Win64\*
 Engine\Plugins\Experimental\Paper2D\Binaries
 Engine\Plugins\Experimental\Paper2D\Intermediate
-REM @todo es31: Take this out in for ES31 support
-Engine\Binaries\Win64\Android\*ES31*
diff --git a/Engine/Source/Developer/Android/AndroidPlatformEditor/Private/AndroidTargetSettingsCustomization.cpp b/Engine/Source/Developer/Android/AndroidPlatformEditor/Private/AndroidTargetSettingsCustomization.cpp
--- a/Engine/Source/Developer/Android/AndroidPlatformEditor/Private/AndroidTargetSettingsCustomization.cpp
+++ b/Engine/Source/Developer/Android/AndroidPlatformEditor/Private/AndroidTargetSettingsCustomization.cpp
@@ -20,6 +20,10 @@
 
 //////////////////////////////////////////////////////////////////////////
 // FAndroidTargetSettingsCustomization
+namespace FAndroidTargetSettingsCustomizationConstants
+{
+	const FText DisabledTip = LOCTEXT("GitHubSourceRequiredToolTip", "This requires GitHub source.");
+}
 
 TSharedRef<IDetailCustomization> FAndroidTargetSettingsCustomization::MakeInstance()
 {
@@ -58,8 +62,9 @@ void FAndroidTargetSettingsCustomization::CustomizeDetails(IDetailLayoutBuilder&
 
 void FAndroidTargetSettingsCustomization::BuildAppManifestSection(IDetailLayoutBuilder& DetailLayout)
 {
-	// App manifest category
+	// Cache some categories
 	IDetailCategoryBuilder& AppManifestCategory = DetailLayout.EditCategory(TEXT("AppManifest"));
+	IDetailCategoryBuilder& BuildCategory = DetailLayout.EditCategory(TEXT("Build"));
 
 	TSharedRef<SPlatformSetupMessage> PlatformSetupMessage = SNew(SPlatformSetupMessage, GameManifestPath)
 		.PlatformName(LOCTEXT("AndroidPlatformName", "Android"))
@@ -166,6 +171,22 @@ void FAndroidTargetSettingsCustomization::BuildAppManifestSection(IDetailLayoutB
 	TSharedRef<IPropertyHandle> GooglePlayLicenseKeyProperty = DetailLayout.GetProperty(GET_MEMBER_NAME_CHECKED(UAndroidRuntimeSettings, GooglePlayLicenseKey));
 	GooglePlayCategory.AddProperty(GooglePlayLicenseKeyProperty)
 		.EditCondition(SetupForGooglePlayAttribute, NULL);
+
+
+#define SETUP_NONROCKET_PROP(PropName, Category, Tip) \
+	{ \
+		TSharedRef<IPropertyHandle> PropertyHandle = DetailLayout.GetProperty(GET_MEMBER_NAME_CHECKED(UAndroidRuntimeSettings, PropName)); \
+		Category.AddProperty(PropertyHandle) \
+			.EditCondition(SetupForPlatformAttribute, NULL) \
+			.IsEnabled(!FRocketSupport::IsRocket()) \
+			.ToolTip(!FRocketSupport::IsRocket() ? Tip : FAndroidTargetSettingsCustomizationConstants::DisabledTip); \
+	}
+	SETUP_NONROCKET_PROP(bBuildForArmV7, BuildCategory, LOCTEXT("BuildForArmV7ToolTip", "Enable ArmV7 CPU architecture support? (this will be used if all CPU architecture types are unchecked)"));
+	SETUP_NONROCKET_PROP(bBuildForX86, BuildCategory, LOCTEXT("BuildForX86ToolTip", "Enable X86 CPU architecture support?"));
+	SETUP_NONROCKET_PROP(bBuildForES2, BuildCategory, LOCTEXT("BuildForES2ToolTip", "Enable OpenGL ES2 rendering support? (this will be used if rendering types are unchecked)"));
+	SETUP_NONROCKET_PROP(bBuildForES31, BuildCategory, LOCTEXT("BuildForES31ToolTip", "Enable OpenGL ES31 + AEP (Android Extension Pack) rendering support?"));
+	// @todo android fat binary: Put back in when we expose those
+//	SETUP_NONROCKET_PROP(bSplitIntoSeparateApks, BuildCategory, LOCTEXT("SplitIntoSeparateAPKsToolTip", "If checked, CPU architectures and rendering types will be split into separate .apk files"));
 }
 
 void FAndroidTargetSettingsCustomization::BuildIconSection(IDetailLayoutBuilder& DetailLayout)
diff --git a/Engine/Source/Developer/Android/AndroidTargetPlatform/Private/AndroidTargetDevice.h b/Engine/Source/Developer/Android/AndroidTargetPlatform/Private/AndroidTargetDevice.h
--- a/Engine/Source/Developer/Android/AndroidTargetPlatform/Private/AndroidTargetDevice.h
+++ b/Engine/Source/Developer/Android/AndroidTargetPlatform/Private/AndroidTargetDevice.h
@@ -124,14 +124,8 @@ public:
 
 	virtual FString GetName() const override
 	{
-		if (AndroidVariant.IsEmpty())
-		{
-			//return DeviceName;
-			return Model;			
-		}
-
-		//return FString::Printf(TEXT("%s (%s)"), *DeviceName, *AndroidVariant);
-		return FString::Printf(TEXT("%s (%s)"), *Model, *AndroidVariant);
+		// we need a unique name for all devices, so use human usable model name and the unique id
+		return FString::Printf(TEXT("%s (%s)"), *Model, *SerialNumber);
 	}
 
 	virtual FString GetOperatingSystemName() override;
diff --git a/Engine/Source/Developer/Android/Android_ES31TargetPlatform/Android_ES31TargetPlatform.Build.cs b/Engine/Source/Developer/Android/Android_ES31TargetPlatform/Android_ES31TargetPlatform.Build.cs
--- a/Engine/Source/Developer/Android/Android_ES31TargetPlatform/Android_ES31TargetPlatform.Build.cs
+++ b/Engine/Source/Developer/Android/Android_ES31TargetPlatform/Android_ES31TargetPlatform.Build.cs
@@ -31,6 +31,8 @@ public class Android_ES31TargetPlatform : ModuleRules
 			PrivateIncludePathModuleNames.Add("TextureCompressor");		//@todo android: AndroidTargetPlatform.Build
 		}
 
+		Definitions.Add("WITH_OGGVORBIS=1");
+
 		PrivateIncludePaths.AddRange(
 			new string[]
 			{
diff --git a/Engine/Source/Developer/Android/Android_ES31TargetPlatform/Private/Android_ES31TargetPlatformModule.cpp b/Engine/Source/Developer/Android/Android_ES31TargetPlatform/Private/Android_ES31TargetPlatformModule.cpp
--- a/Engine/Source/Developer/Android/Android_ES31TargetPlatform/Private/Android_ES31TargetPlatformModule.cpp
+++ b/Engine/Source/Developer/Android/Android_ES31TargetPlatform/Private/Android_ES31TargetPlatformModule.cpp
@@ -22,7 +22,7 @@ class FAndroid_ES31TargetPlatform
 
 	virtual FText DisplayName( ) const override
 	{
-		return LOCTEXT("Android_ES31", "Android (ES 3.1)");
+		return LOCTEXT("Android_ES31", "Android (ES 3.1 + AEP)");
 	}
 
 	virtual FString PlatformName() const override
@@ -43,8 +43,18 @@ class FAndroid_ES31TargetPlatform
 
 	virtual bool SupportedByExtensionsString( const FString& ExtensionsString, const int GLESVersion ) const override
 	{
-		// GL_EXT_tessellation_shader implies ES3.1
-		return (ExtensionsString.Contains(TEXT("GL_EXT_tessellation_shader")));
+		// look for the AEP extension
+		return ExtensionsString.Contains(TEXT("GL_ANDROID_extension_pack_es31a")) && ExtensionsString.Contains(TEXT("GL_EXT_color_buffer_half_float"));
+	}
+
+	virtual FText GetVariantDisplayName() const override
+	{
+		return LOCTEXT("Android_31_ShortName", "ES 3.1 + AEP");
+	}
+
+	virtual float GetVariantPriority() const override
+	{
+		return 0.9f;
 	}
 
 	virtual bool SupportsFeature( ETargetPlatformFeatures Feature ) const override
diff --git a/Engine/Source/Developer/Android/Android_PVRTCTargetPlatform/Private/Android_PVRTCTargetPlatformModule.cpp b/Engine/Source/Developer/Android/Android_PVRTCTargetPlatform/Private/Android_PVRTCTargetPlatformModule.cpp
--- a/Engine/Source/Developer/Android/Android_PVRTCTargetPlatform/Private/Android_PVRTCTargetPlatformModule.cpp
+++ b/Engine/Source/Developer/Android/Android_PVRTCTargetPlatform/Private/Android_PVRTCTargetPlatformModule.cpp
@@ -52,7 +52,7 @@ class FAndroid_PVRTCTargetPlatform
 
 	virtual float GetVariantPriority() const override
 	{
-		return 0.9f;
+		return 0.8f;
 	}
 };
 
diff --git a/Engine/Source/Developer/DesktopPlatform/Private/PlatformInfo.cpp b/Engine/Source/Developer/DesktopPlatform/Private/PlatformInfo.cpp
--- a/Engine/Source/Developer/DesktopPlatform/Private/PlatformInfo.cpp
+++ b/Engine/Source/Developer/DesktopPlatform/Private/PlatformInfo.cpp
@@ -101,8 +101,8 @@ static const FPlatformInfo PlatformInfoArray[] = {
 	BuildPlatformInfo(TEXT("Android_ETC1"),				TEXT("Android_ETC1"),		LOCTEXT("Android_ETC1", "Android (ETC1)"),					EPlatformType::Game,		EPlatformFlags::CookFlavor,		FPlatformIconPaths(TEXT("Launcher/AndroidTarget/Platform_Android_ETC1_24x"), TEXT("Launcher/AndroidTarget/Platform_Android_128x")),					TEXT("-targetplatform=Android -cookflavor=ETC1"),		TEXT("Android"),	EPlatformSDKStatus::Unknown,	TEXT("/Engine/Tutorial/Mobile/SettingUpAndroidTutorial.SettingUpAndroidTutorial"),						IsAvailableOnWindows,						true),
 	BuildPlatformInfo(TEXT("Android_ETC2"),				TEXT("Android_ETC2"),		LOCTEXT("Android_ETC2", "Android (ETC2)"),					EPlatformType::Game,		EPlatformFlags::CookFlavor,		FPlatformIconPaths(TEXT("Launcher/AndroidTarget/Platform_Android_ETC2_24x"), TEXT("Launcher/AndroidTarget/Platform_Android_128x")),					TEXT("-targetplatform=Android -cookflavor=ETC2"),		TEXT("Android"),	EPlatformSDKStatus::Unknown,	TEXT("/Engine/Tutorial/Mobile/SettingUpAndroidTutorial.SettingUpAndroidTutorial"),						IsAvailableOnWindows,						true),
 	BuildPlatformInfo(TEXT("Android_PVRTC"),			TEXT("Android_PVRTC"),		LOCTEXT("Android_PVRTC", "Android (PVRTC)"),				EPlatformType::Game,		EPlatformFlags::CookFlavor,		FPlatformIconPaths(TEXT("Launcher/AndroidTarget/Platform_Android_PVRTC_24x"), TEXT("Launcher/AndroidTarget/Platform_Android_128x")),				TEXT("-targetplatform=Android -cookflavor=PVRTC"),		TEXT("Android"),	EPlatformSDKStatus::Unknown,	TEXT("/Engine/Tutorial/Mobile/SettingUpAndroidTutorial.SettingUpAndroidTutorial"),						IsAvailableOnWindows,						true),
-	// @todo es31: Put this back in for ES31 support
-//	BuildPlatformInfo(TEXT("Android_ES31"),				TEXT("Android_ES31"),		LOCTEXT("Android_ES31", "Android (ES31)"),					EPlatformType::Game,		EPlatformFlags::CookFlavor,		FPlatformIconPaths(TEXT("Launcher/AndroidTarget/Platform_Android_ES31_24x"), TEXT("Launcher/AndroidTarget/Platform_Android_128x")),					TEXT("-targetplatform=Android -cookflavor=ES31"),		TEXT("Android"),	EPlatformSDKStatus::Unknown,	TEXT("Shared/Tutorials/SettingUpAndroidTutorial"),														IsAvailableOnWindows,						true),
+	// @todo es31: Set the final false to true when ES31 is in Binary releases
+	BuildPlatformInfo(TEXT("Android_ES31"),				TEXT("Android_ES31"),		LOCTEXT("Android_ES31", "Android (ES31+AEP)"),				EPlatformType::Game,		EPlatformFlags::CookFlavor,		FPlatformIconPaths(TEXT("Launcher/AndroidTarget/Platform_Android_ES31_24x"), TEXT("Launcher/AndroidTarget/Platform_Android_128x")),					TEXT("-targetplatform=Android -cookflavor=ES31"),		TEXT("Android"),	EPlatformSDKStatus::Unknown,	TEXT("Shared/Tutorials/SettingUpAndroidTutorial"),														IsAvailableOnWindows,						false),
 //	BuildPlatformInfo(TEXT("Android_GL4"),				TEXT("Android_GL4"),		LOCTEXT("Android_GL4", "Android (GL4)"),					EPlatformType::Game,		EPlatformFlags::CookFlavor,		FPlatformIconPaths(TEXT("Launcher/AndroidTarget/Platform_Android_GL4_24x"), TEXT("Launcher/AndroidTarget/Platform_Android_128x")),					TEXT("-targetplatform=Android -cookflavor=GL4"),		TEXT("Android"),	EPlatformSDKStatus::Unknown,	TEXT("Shared/Tutorials/SettingUpAndroidTutorial"),														IsAvailableOnWindows,						true),
 
 	BuildPlatformInfo(TEXT("HTML5"),					TEXT("HTML5"),				LOCTEXT("HTML5", "HTML5"),									EPlatformType::Game,		EPlatformFlags::None,			FPlatformIconPaths(TEXT("Launcher/HTML5Target/Platform_HTML5_24x"), TEXT("Launcher/HTML5Target/Platform_HTML5_128x")),								TEXT(""),												TEXT("HTML5"),		EPlatformSDKStatus::Unknown,	TEXT("Platforms/HTML5/GettingStarted"),																	IsAvailableOnWindows || IsAvailableOnMac,   false),
diff --git a/Engine/Source/Editor/MainFrame/Private/Menus/CookContentMenu.h b/Engine/Source/Editor/MainFrame/Private/Menus/CookContentMenu.h
--- a/Engine/Source/Editor/MainFrame/Private/Menus/CookContentMenu.h
+++ b/Engine/Source/Editor/MainFrame/Private/Menus/CookContentMenu.h
@@ -78,6 +78,12 @@ protected:
 	 */
 	static void AddPlatformToMenu(FMenuBuilder& MenuBuilder, const PlatformInfo::FPlatformInfo& PlatformInfo)
 	{
+		// don't add sub-platforms that are disabled in rocket
+		if (!PlatformInfo.bEnabledInBinary && FRocketSupport::IsRocket())
+		{
+			return;
+		}
+
 		IProjectTargetPlatformEditorModule& ProjectTargetPlatformEditorModule = FModuleManager::LoadModuleChecked<IProjectTargetPlatformEditorModule>("ProjectTargetPlatformEditor");
 
 		FUIAction Action(
diff --git a/Engine/Source/Editor/MainFrame/Private/Menus/PackageProjectMenu.h b/Engine/Source/Editor/MainFrame/Private/Menus/PackageProjectMenu.h
--- a/Engine/Source/Editor/MainFrame/Private/Menus/PackageProjectMenu.h
+++ b/Engine/Source/Editor/MainFrame/Private/Menus/PackageProjectMenu.h
@@ -89,6 +89,12 @@ protected:
 	 */
 	static void AddPlatformToMenu(FMenuBuilder& MenuBuilder, const PlatformInfo::FPlatformInfo& PlatformInfo)
 	{
+		// don't add sub-platforms that are disabled in rocket
+		if (!PlatformInfo.bEnabledInBinary && FRocketSupport::IsRocket())
+		{
+			return;
+		}
+
 		IProjectTargetPlatformEditorModule& ProjectTargetPlatformEditorModule = FModuleManager::LoadModuleChecked<IProjectTargetPlatformEditorModule>("ProjectTargetPlatformEditor");
 
 		FUIAction Action(
diff --git a/Engine/Source/Programs/AutomationTool/Android/AndroidPlatform.Automation.cs b/Engine/Source/Programs/AutomationTool/Android/AndroidPlatform.Automation.cs
--- a/Engine/Source/Programs/AutomationTool/Android/AndroidPlatform.Automation.cs
+++ b/Engine/Source/Programs/AutomationTool/Android/AndroidPlatform.Automation.cs
@@ -364,7 +364,7 @@ public class AndroidPlatform : Platform
 	public override void Deploy(ProjectParams Params, DeploymentContext SC)
 	{
 		string DeviceArchitecture = GetBestDeviceArchitecture(Params);
-		string GPUArchitecture = "";
+		string GPUArchitecture = GetBestGPUArchitecture(Params);
 
 		string ApkName = GetFinalApkName(Params, SC.StageExecutables[0], true, DeviceArchitecture, GPUArchitecture);
 
@@ -650,10 +650,10 @@ public class AndroidPlatform : Platform
 		string[] AppArchitectures = AndroidToolChain.GetAllArchitectures();
 
 		// ask the device
-		ProcessResult Result = RunAdbCommand(Params, " shell getprop ro.product.cpu.abi", null, ERunOptions.AppMustExist);
+		ProcessResult ABIResult = RunAdbCommand(Params, " shell getprop ro.product.cpu.abi", null, ERunOptions.AppMustExist);
 
 		// the output is just the architecture
-		string DeviceArch = UnrealBuildTool.Android.UEDeployAndroid.GetUE4Arch(Result.Output.Trim());
+		string DeviceArch = UnrealBuildTool.Android.UEDeployAndroid.GetUE4Arch(ABIResult.Output.Trim());
 
 		// if the architecture wasn't built, look for a backup
 		if (Array.IndexOf(AppArchitectures, DeviceArch) == -1)
@@ -700,11 +700,35 @@ public class AndroidPlatform : Platform
 		return DeviceArch;
 	}
 
+	private string GetBestGPUArchitecture(ProjectParams Params)
+	{
+		bool bMakeSeparateApks = UnrealBuildTool.Android.UEDeployAndroid.ShouldMakeSeparateApks();
+		// if we are joining all .so's into a single .apk, there's no need to find the best one - there is no other one
+		if (!bMakeSeparateApks)
+		{
+			return "";
+		}
+
+		string[] AppGPUArchitectures = AndroidToolChain.GetAllArchitectures();
+
+		// get the device extensions
+		ProcessResult ExtensionsResult = RunAdbCommand(Params, "shell dumpsys SurfaceFlinger");
+		string Extensions = ExtensionsResult.Output.Trim();
+
+		Console.WriteLine("Extensions: {0}", Extensions);
+		// look for AEP support
+		if (Extensions.Contains("GL_ANDROID_extension_pack_es31a") && Extensions.Contains("GL_EXT_color_buffer_half_float") && AppGPUArchitectures.Contains("-es31"))
+		{
+			return "-es31";
+		}
+
+		return "-es2";
+	}
 
 	public override ProcessResult RunClient(ERunOptions ClientRunFlags, string ClientApp, string ClientCmdLine, ProjectParams Params)
 	{
 		string DeviceArchitecture = GetBestDeviceArchitecture(Params);
-		string GPUArchitecture = "";
+		string GPUArchitecture = GetBestGPUArchitecture(Params); ;
 
 		string ApkName = ClientApp + DeviceArchitecture + ".apk";
 		if (!File.Exists(ApkName))
diff --git a/Engine/Source/Programs/UnrealBuildTool/Android/AndroidToolChain.cs b/Engine/Source/Programs/UnrealBuildTool/Android/AndroidToolChain.cs
--- a/Engine/Source/Programs/UnrealBuildTool/Android/AndroidToolChain.cs
+++ b/Engine/Source/Programs/UnrealBuildTool/Android/AndroidToolChain.cs
@@ -70,6 +70,10 @@ namespace UnrealBuildTool
 
 			// Parse selected GPU architectures
 			List<string> ProjectGPUArches = new List<string>();
+			if (Ini.GetBool("/Script/AndroidRuntimeSettings.AndroidRuntimeSettings", "bBuildForES2", out bBuild) && bBuild)
+			{
+				ProjectGPUArches.Add("-es2");
+			}
 			if (Ini.GetBool("/Script/AndroidRuntimeSettings.AndroidRuntimeSettings", "bBuildForES31", out bBuild) && bBuild)
 			{
 				ProjectGPUArches.Add("-es31");
@@ -78,7 +82,7 @@ namespace UnrealBuildTool
 			{
 				ProjectGPUArches.Add("-gl4");
 			}
-			if (ProjectGPUArches.Count == 0 || (Ini.GetBool("/Script/AndroidRuntimeSettings.AndroidRuntimeSettings", "bBuildForES2", out bBuild) && bBuild))
+			if (ProjectGPUArches.Count == 0)
 			{
 				ProjectGPUArches.Add("-es2");
 			}
@@ -931,8 +935,11 @@ namespace UnrealBuildTool
 		public override void AddFilesToManifest(ref FileManifest Manifest, UEBuildBinary Binary)
 		{
 			// the binary will have all of the .so's in the output files, we need to trim down to the shared apk (which is what needs to go into the manifest)
-			string ApkFile = Path.ChangeExtension(RemoveArchName(Binary.Config.OutputFilePaths[0]), ".apk");
-			Manifest.AddFileName(ApkFile);
+			foreach (string BinaryPath in Binary.Config.OutputFilePaths)
+			{
+				string ApkFile = Path.ChangeExtension(BinaryPath, ".apk");
+				Manifest.AddFileName(ApkFile);
+			}
 		}
 
 		public override void CompileCSharpProject(CSharpEnvironment CompileEnvironment, string ProjectFileName, string DestinationFile)
diff --git a/Engine/Source/Programs/UnrealBuildTool/Android/UEBuildAndroid.cs b/Engine/Source/Programs/UnrealBuildTool/Android/UEBuildAndroid.cs
--- a/Engine/Source/Programs/UnrealBuildTool/Android/UEBuildAndroid.cs
+++ b/Engine/Source/Programs/UnrealBuildTool/Android/UEBuildAndroid.cs
@@ -232,26 +232,25 @@ namespace UnrealBuildTool
 					{
 						if (UEBuildConfiguration.bBuildDeveloperTools)
 						{
-                            InModule.AddPlatformSpecificDynamicallyLoadedModule("AndroidTargetPlatform");
+							InModule.AddPlatformSpecificDynamicallyLoadedModule("AndroidTargetPlatform");
 							InModule.AddPlatformSpecificDynamicallyLoadedModule("Android_PVRTCTargetPlatform");
 							InModule.AddPlatformSpecificDynamicallyLoadedModule("Android_ATCTargetPlatform");
 							InModule.AddPlatformSpecificDynamicallyLoadedModule("Android_DXTTargetPlatform");
-                            InModule.AddPlatformSpecificDynamicallyLoadedModule("Android_ETC1TargetPlatform");
-                            InModule.AddPlatformSpecificDynamicallyLoadedModule("Android_ETC2TargetPlatform");
-//                            InModule.AddPlatformSpecificDynamicallyLoadedModule("Android_GL4TargetPlatform");
-							// @todo es31: Put this back in for ES31 support
-							// ES31 code is in, but it's not fully supported in UE4 4.5, so for now we need to disable the targetplatform as it will confuse people greatly
-							// InModule.AddPlatformSpecificDynamicallyLoadedModule("Android_ES31TargetPlatform");
+							InModule.AddPlatformSpecificDynamicallyLoadedModule("Android_ETC1TargetPlatform");
+							InModule.AddPlatformSpecificDynamicallyLoadedModule("Android_ETC2TargetPlatform");
+							InModule.AddPlatformSpecificDynamicallyLoadedModule("Android_ES31TargetPlatform");
+							// @todo gl4android				
+							// InModule.AddPlatformSpecificDynamicallyLoadedModule("Android_GL4TargetPlatform");
                         }
 					}
 					else if (InModule.ToString() == "TargetPlatform")
 					{
 						bBuildShaderFormats = true;
-                        InModule.AddDynamicallyLoadedModule("TextureFormatPVR");
+						InModule.AddDynamicallyLoadedModule("TextureFormatPVR");
 						InModule.AddDynamicallyLoadedModule("TextureFormatDXT");
 						InModule.AddDynamicallyLoadedModule("TextureFormatASTC");
-                        InModule.AddPlatformSpecificDynamicallyLoadedModule("TextureFormatAndroid");    // ATITC, ETC1 and ETC2
-                        if (UEBuildConfiguration.bBuildDeveloperTools)
+						InModule.AddPlatformSpecificDynamicallyLoadedModule("TextureFormatAndroid");    // ATITC, ETC1 and ETC2
+						if (UEBuildConfiguration.bBuildDeveloperTools)
 						{
 							//InModule.AddDynamicallyLoadedModule("AudioFormatADPCM");	//@todo android: android audio
 						}
@@ -261,16 +260,15 @@ namespace UnrealBuildTool
 				// allow standalone tools to use targetplatform modules, without needing Engine
 				if (UEBuildConfiguration.bForceBuildTargetPlatforms)
 				{
-                    InModule.AddPlatformSpecificDynamicallyLoadedModule("AndroidTargetPlatform");
+					InModule.AddPlatformSpecificDynamicallyLoadedModule("AndroidTargetPlatform");
 					InModule.AddPlatformSpecificDynamicallyLoadedModule("Android_PVRTCTargetPlatform");
 					InModule.AddPlatformSpecificDynamicallyLoadedModule("Android_ATCTargetPlatform");
 					InModule.AddPlatformSpecificDynamicallyLoadedModule("Android_DXTTargetPlatform");
-                    InModule.AddPlatformSpecificDynamicallyLoadedModule("Android_ETC1TargetPlatform");
-                    InModule.AddPlatformSpecificDynamicallyLoadedModule("Android_ETC2TargetPlatform");
-//                    InModule.AddPlatformSpecificDynamicallyLoadedModule("Android_GL4TargetPlatform");
-					// @todo es31: Put this back in for ES31 support
-					// ES31 code is in, but it's not fully supported in UE4 4.5, so for now we need to disable the targetplatform as it will confuse people greatly
-					// InModule.AddPlatformSpecificDynamicallyLoadedModule("Android_ES31TargetPlatform");
+					InModule.AddPlatformSpecificDynamicallyLoadedModule("Android_ETC1TargetPlatform");
+					InModule.AddPlatformSpecificDynamicallyLoadedModule("Android_ETC2TargetPlatform");
+					InModule.AddPlatformSpecificDynamicallyLoadedModule("Android_ES31TargetPlatform");
+					// @todo gl4android				
+					// InModule.AddPlatformSpecificDynamicallyLoadedModule("Android_GL4TargetPlatform");
                 }
 
 				if (bBuildShaderFormats)
diff --git a/Engine/Source/Programs/UnrealBuildTool/Android/UEDeployAndroid.cs b/Engine/Source/Programs/UnrealBuildTool/Android/UEDeployAndroid.cs
--- a/Engine/Source/Programs/UnrealBuildTool/Android/UEDeployAndroid.cs
+++ b/Engine/Source/Programs/UnrealBuildTool/Android/UEDeployAndroid.cs
@@ -444,7 +444,7 @@ namespace UnrealBuildTool.Android
 			string UpdateCommandLine = "--silent update project --subprojects --name " + ProjectName + " --path . --target " + GetSdkApiLevel();
 			foreach (string Lib in LibsToBeAdded)
 			{
-				UpdateCommandLine += " --library " + Lib;
+				string LocalUpdateCommandLine = UpdateCommandLine + " --library " + Lib;
 
 				// make sure each library has a build.xml - --subprojects doesn't create build.xml files, but it will create project.properties
 				// and later code needs each lib to have a build.xml
@@ -452,9 +452,9 @@ namespace UnrealBuildTool.Android
 				{
 					RunCommandLineProgramAndThrowOnError(UE4BuildPath, AndroidCommandPath, "--silent update lib-project --path " + Lib + " --target " + GetSdkApiLevel(), "");
 				}
+				RunCommandLineProgramAndThrowOnError(UE4BuildPath, AndroidCommandPath, LocalUpdateCommandLine, "Updating project.properties, local.properties, and build.xml...");
 			}
 
-			RunCommandLineProgramAndThrowOnError(UE4BuildPath, AndroidCommandPath, UpdateCommandLine, "Updating project.properties, local.properties, and build.xml...");
 		}
 
 
@@ -735,7 +735,7 @@ namespace UnrealBuildTool.Android
 					}
 
 					// now do final stuff per apk (or after all .so's for a shared .apk)
-					if (bMakeSeparateApks || ArchIndex == NumArches - 1)
+					if (bMakeSeparateApks || (ArchIndex == Arches.Length - 1 && GPUArchIndex == GPUArchitectures.Length - 1))
 					{
 						// always delete libs up to this point so fat binaries and incremental builds work together (otherwise we might end up with multiple
 						// so files in an apk that doesn't want them)
@@ -870,19 +870,34 @@ namespace UnrealBuildTool.Android
 		{
 			// we need to strip architecture from any of the output paths
 			string BaseSoName = AndroidToolChain.RemoveArchName(InTarget.OutputPaths[0]);
-			// this always makes a merged .apk since for debugging, there's no way to know which one to run
-			MakeApk(InTarget.AppName, InTarget.ProjectDirectory, BaseSoName, BuildConfiguration.RelativeEnginePath, bForDistribution:false, CookFlavor:"", bMakeSeparateApks:false, bIncrementalPackage:true);
+
+			// make an apk at the end of compiling, so that we can run without packaging (debugger, cook on the fly, etc)
+			MakeApk(InTarget.AppName, InTarget.ProjectDirectory, BaseSoName, BuildConfiguration.RelativeEnginePath, bForDistribution: false, CookFlavor: "", bMakeSeparateApks:ShouldMakeSeparateApks(), bIncrementalPackage:true);
+
+			// if we made any non-standard .apk files, the generated debugger settings may be wrong
+			if (ShouldMakeSeparateApks() && (InTarget.OutputPaths.Length > 1 || !InTarget.OutputPaths[0].Contains("-armv7-es2")))
+			{
+				Console.WriteLine("================================================================================================================================");
+				Console.WriteLine("Non-default apk(s) have been made: If you are debugging, you will need to manually select one to run in the debugger properties!");
+				Console.WriteLine("================================================================================================================================");
+			}
 			return true;
 		}
 
 		public static bool ShouldMakeSeparateApks()
 		{
+			// @todo android fat binary: Currently, there isn't much utility in merging multiple .so's into a single .apk except for debugging,
+			// but we can't properly handle multiple GPU architectures in a single .apk, so we are disabling the feature for now
+			// The user will need to manually select the apk to run in their Visual Studio debugger settings (see Override APK in TADP, for instance)
+			// If we change this, pay attention to <OverrideAPKPath> in AndroidProjectGenerator
+			return true;
+			
 			// check to see if the project wants separate apks
-			ConfigCacheIni Ini = new ConfigCacheIni(UnrealTargetPlatform.Android, "Engine", UnrealBuildTool.GetUProjectPath());
-			bool bSeparateApks = false;
-			Ini.GetBool("/Script/AndroidRuntimeSettings.AndroidRuntimeSettings", "bSplitIntoSeparateApks", out bSeparateApks);
-
-			return bSeparateApks;
+// 			ConfigCacheIni Ini = new ConfigCacheIni(UnrealTargetPlatform.Android, "Engine", UnrealBuildTool.GetUProjectPath());
+// 			bool bSeparateApks = false;
+// 			Ini.GetBool("/Script/AndroidRuntimeSettings.AndroidRuntimeSettings", "bSplitIntoSeparateApks", out bSeparateApks);
+// 
+// 			return bSeparateApks;
 		}
 
 		public override bool PrepForUATPackageOrDeploy(string ProjectName, string ProjectDirectory, string ExecutablePath, string EngineDirectory, bool bForDistribution, string CookFlavor)
diff --git a/Engine/Source/Runtime/Android/AndroidRuntimeSettings/Classes/AndroidRuntimeSettings.h b/Engine/Source/Runtime/Android/AndroidRuntimeSettings/Classes/AndroidRuntimeSettings.h
--- a/Engine/Source/Runtime/Android/AndroidRuntimeSettings/Classes/AndroidRuntimeSettings.h
+++ b/Engine/Source/Runtime/Android/AndroidRuntimeSettings/Classes/AndroidRuntimeSettings.h
@@ -126,9 +126,20 @@ public:
 //	UPROPERTY(GlobalConfig, EditAnywhere, Category = Build, meta = (DisplayName = "Support x86_64 [aka x64]"))
 //	bool bBuildForX8664;
 
+	// Enable ES2 support? [CURRENTLY FOR FULL SOURCE GAMES ONLY]
+	UPROPERTY(GlobalConfig, EditAnywhere, Category = Build, meta = (DisplayName = "Support OpenGL ES2"))
+	bool bBuildForES2;
+
+	// Enable ES31 support? [CURRENTLY FOR FULL SOURCE GAMES ONLY]
+	UPROPERTY(GlobalConfig, EditAnywhere, Category = Build, meta = (DisplayName = "Support OpenGL ES31 + AEP"))
+	bool bBuildForES31;
+
 	// If selected, the checked architectures will be split into separate .apk files [CURRENTLY FOR FULL SOURCE GAMES ONLY]
-	UPROPERTY(GlobalConfig, EditAnywhere, Category = Build)
-	bool bSplitIntoSeparateApks;
+	// @todo android fat binary: Currently, there isn't much utility in merging multiple .so's into a single .apk except for debugging,
+	// but we can't properly handle multiple GPU architectures in a single .apk, so we are disabling the feature for now
+	// The user will need to manually select the apk to run in their Visual Studio debugger settings (see Override APK in TADP, for instance)
+// 	UPROPERTY(GlobalConfig, EditAnywhere, Category = Build)
+// 	bool bSplitIntoSeparateApks;
 
 	// Should Google Play support be enabled?
 	UPROPERTY(GlobalConfig, EditAnywhere, Category = GooglePlayServices)
diff --git a/Engine/Source/Runtime/Android/AndroidRuntimeSettings/Private/AndroidRuntimeSettings.cpp b/Engine/Source/Runtime/Android/AndroidRuntimeSettings/Private/AndroidRuntimeSettings.cpp
--- a/Engine/Source/Runtime/Android/AndroidRuntimeSettings/Private/AndroidRuntimeSettings.cpp
+++ b/Engine/Source/Runtime/Android/AndroidRuntimeSettings/Private/AndroidRuntimeSettings.cpp
@@ -22,5 +22,11 @@ void UAndroidRuntimeSettings::PostEditChangeProperty(struct FPropertyChangedEven
 	{
 		bBuildForArmV7 = true;
 	}
+
+	// Ensure that at least one GPU architecture is supported
+	if (!bBuildForES2 && !bBuildForES31)
+	{
+		bBuildForES2 = true;
+	}
 }
 #endif
diff --git a/Engine/Source/Runtime/OpenGLDrv/Private/Android/AndroidOpenGLPrivate.h b/Engine/Source/Runtime/OpenGLDrv/Private/Android/AndroidOpenGLPrivate.h
--- a/Engine/Source/Runtime/OpenGLDrv/Private/Android/AndroidOpenGLPrivate.h
+++ b/Engine/Source/Runtime/OpenGLDrv/Private/Android/AndroidOpenGLPrivate.h
@@ -40,8 +40,13 @@ private:
 		if (!EGL->IsInitialized())
 		{
 			FAndroidAppEntry::PlatformInit();
+#if PLATFORM_ANDROIDES31
+			EGL->InitSurface(true);
+#endif
 		}
+#if !PLATFORM_ANDROIDES31
 		EGL->InitSurface(true);
+#endif
 		EGL->SetCurrentSharedContext();
 
 		// get extensions
diff --git a/Engine/Source/Runtime/OpenGLDrv/Private/OpenGLES31.cpp b/Engine/Source/Runtime/OpenGLDrv/Private/OpenGLES31.cpp
--- a/Engine/Source/Runtime/OpenGLDrv/Private/OpenGLES31.cpp
+++ b/Engine/Source/Runtime/OpenGLDrv/Private/OpenGLES31.cpp
@@ -149,13 +149,8 @@ bool FOpenGLES31::SupportsAdvancedFeatures()
 		// Check for minimum ES 3.1 + extensions support to avoid the ES2 fallback
 		bResult = (MajorVersion == 3 && MinorVersion >= 1);
 
-		bResult &= ExtensionsString.Contains(TEXT("GL_EXT_draw_buffers_indexed"));
-		bResult &= ExtensionsString.Contains(TEXT("GL_EXT_gpu_shader5"));
-		bResult &= ExtensionsString.Contains(TEXT("GL_EXT_geometry_shader"));
-		bResult &= ExtensionsString.Contains(TEXT("GL_EXT_texture_border_clamp"));
-		bResult &= ExtensionsString.Contains(TEXT("GL_EXT_texture_buffer"));
-		bResult &= ExtensionsString.Contains(TEXT("GL_EXT_texture_cube_map_array"));
-		bResult &= ExtensionsString.Contains(TEXT("GL_OES_shader_image_atomic"));
+		bResult &= ExtensionsString.Contains(TEXT("GL_ANDROID_extension_pack_es31a"));
+		bResult &= ExtensionsString.Contains(TEXT("GL_EXT_color_buffer_half_float"));
 	}
 	else
 	{
