- Missed this file in the ES31 checkin

diff --git a/Engine/Source/Programs/UnrealBuildTool/Android/AndroidProjectGenerator.cs b/Engine/Source/Programs/UnrealBuildTool/Android/AndroidProjectGenerator.cs
--- a/Engine/Source/Programs/UnrealBuildTool/Android/AndroidProjectGenerator.cs
+++ b/Engine/Source/Programs/UnrealBuildTool/Android/AndroidProjectGenerator.cs
@@ -189,7 +189,7 @@ namespace UnrealBuildTool
 			IntermediateDirectoryPath = Path.GetDirectoryName(NMakeOutputPath) + "/../../Intermediate/Android/APK";
 
 			// string for <OverrideAPKPath>
-			string APKPath = Path.ChangeExtension(NMakeOutputPath, ".apk");
+			string APKPath = Path.ChangeExtension(NMakeOutputPath, "-armv7-es2.apk");
 
 			// string for <BuildXmlPath> and <AndroidManifestPath>
 			string BuildXmlPath = IntermediateDirectoryPath;
