Fix output of .user file when generating projects on Windows. Seems that it was relying on HasUProjectFile() being set, which is not a global (and not set per-project).

diff --git a/Engine/Source/Programs/UnrealBuildTool/System/VCProject.cs b/Engine/Source/Programs/UnrealBuildTool/System/VCProject.cs
--- a/Engine/Source/Programs/UnrealBuildTool/System/VCProject.cs
+++ b/Engine/Source/Programs/UnrealBuildTool/System/VCProject.cs
@@ -379,8 +379,7 @@ namespace UnrealBuildTool
 				ProjectFileGenerator.NewLine +
 				"<Project DefaultTargets=\"Build\" ToolsVersion=\"" + ToolsVersion + "\" xmlns=\"http://schemas.microsoft.com/developer/msbuild/2003\">" + ProjectFileGenerator.NewLine);
 
-			bool bGenerateUserFileContent = UnrealBuildTool.HasUProjectFile() ||
-				UEPlatformProjectGenerator.PlatformRequiresVSUserFileGeneration(InPlatforms, InConfigurations);
+			bool bGenerateUserFileContent = UEPlatformProjectGenerator.PlatformRequiresVSUserFileGeneration(InPlatforms, InConfigurations);
 			if (bGenerateUserFileContent)
 			{
 				VCUserFileContent.Append(
@@ -1054,41 +1053,42 @@ namespace UnrealBuildTool
 
 				if (VCUserFileContent != null)
 				{
-					if (UnrealBuildTool.HasUProjectFile())
+					if ((Platform == UnrealTargetPlatform.Win32) || (Platform == UnrealTargetPlatform.Win64))
 					{
-						if ((Platform == UnrealTargetPlatform.Win32) || (Platform == UnrealTargetPlatform.Win64))
+						VCUserFileContent.Append(
+							"	<PropertyGroup " + ConditionString + ">" + ProjectFileGenerator.NewLine);
+						if (TargetRulesObject.Type != TargetRules.TargetType.Game)
 						{
-							VCUserFileContent.Append(
-								"	<PropertyGroup " + ConditionString + ">" + ProjectFileGenerator.NewLine);
-							if (TargetRulesObject.Type != TargetRules.TargetType.Game)
+							string DebugOptions = "";
+							
+							if(IsForeignProject)
 							{
-								string DebugOptions = UProjectPath;
-
-								if(DebugOptions.Length == 0 && TargetRulesObject.Type == TargetRules.TargetType.Editor && ProjectName != "UE4")
-								{
-									DebugOptions += ProjectName;
-								}
-
-								if (Configuration == UnrealTargetConfiguration.Debug || Configuration == UnrealTargetConfiguration.DebugGame)
-								{
-									DebugOptions += " -debug";
-								}
-								else if (Configuration == UnrealTargetConfiguration.Shipping)
-								{
-									DebugOptions += " -shipping";
-								}
+								DebugOptions += UProjectPath;
+							}
+							else if(TargetRulesObject.Type == TargetRules.TargetType.Editor && ProjectName != "UE4")
+							{
+								DebugOptions += ProjectName;
+							}
 
-								VCUserFileContent.Append(
-									"		<LocalDebuggerCommandArguments>" + DebugOptions + "</LocalDebuggerCommandArguments>" + ProjectFileGenerator.NewLine
-									);
+							if (Configuration == UnrealTargetConfiguration.Debug || Configuration == UnrealTargetConfiguration.DebugGame)
+							{
+								DebugOptions += " -debug";
 							}
+							else if (Configuration == UnrealTargetConfiguration.Shipping)
+							{
+								DebugOptions += " -shipping";
+							}
+
 							VCUserFileContent.Append(
-								"		<DebuggerFlavor>WindowsLocalDebugger</DebuggerFlavor>" + ProjectFileGenerator.NewLine
-								);
-							VCUserFileContent.Append(
-								"	</PropertyGroup>" + ProjectFileGenerator.NewLine
+								"		<LocalDebuggerCommandArguments>" + DebugOptions + "</LocalDebuggerCommandArguments>" + ProjectFileGenerator.NewLine
 								);
 						}
+						VCUserFileContent.Append(
+							"		<DebuggerFlavor>WindowsLocalDebugger</DebuggerFlavor>" + ProjectFileGenerator.NewLine
+							);
+						VCUserFileContent.Append(
+							"	</PropertyGroup>" + ProjectFileGenerator.NewLine
+							);
 					}
 
 					string PlatformUserFileStrings = (ProjGenerator != null) ? ProjGenerator.GetVisualStudioUserFileStrings(Platform, Configuration, ConditionString, TargetRulesObject, TargetFilePath, ProjectFilePath) : "";
diff --git a/Engine/Source/Programs/UnrealBuildTool/Windows/WindowsProjectGenerator.cs b/Engine/Source/Programs/UnrealBuildTool/Windows/WindowsProjectGenerator.cs
--- a/Engine/Source/Programs/UnrealBuildTool/Windows/WindowsProjectGenerator.cs
+++ b/Engine/Source/Programs/UnrealBuildTool/Windows/WindowsProjectGenerator.cs
@@ -44,5 +44,10 @@ namespace UnrealBuildTool
 			}
 			return InPlatform.ToString();
 		}
+
+		public override bool RequiresVSUserFileGeneration()
+		{
+			return true;
+		}
 	}
 }
