Fix for relying on game project matching game directory when generating projects for Xcode
#ios
#mac

diff --git a/Engine/Source/Programs/UnrealBuildTool/System/XcodeProjectFileGenerator.cs b/Engine/Source/Programs/UnrealBuildTool/System/XcodeProjectFileGenerator.cs
--- a/Engine/Source/Programs/UnrealBuildTool/System/XcodeProjectFileGenerator.cs
+++ b/Engine/Source/Programs/UnrealBuildTool/System/XcodeProjectFileGenerator.cs
@@ -725,7 +725,7 @@ namespace UnrealBuildTool
 			IUEToolChain Toolchain = UEToolChain.GetPlatformToolChain(CPPTargetPlatform.IOS);
 			foreach (string GameFolder in GameFolders)
 			{
-				if (GameFolder.EndsWith(Target.TargetName))
+				if (File.Exists(Path.Combine(GameFolder, Target.TargetName+".uproject")))
 				{
 					IsAGame = true;
 					GamePath = Toolchain.ConvertPath(Path.GetFullPath(GameFolder));
