Attempted fix for Mac build failure. Check if bash_profile exists before reading it.

diff --git a/Engine/Source/Programs/UnrealBuildTool/Android/UEBuildAndroid.cs b/Engine/Source/Programs/UnrealBuildTool/Android/UEBuildAndroid.cs
--- a/Engine/Source/Programs/UnrealBuildTool/Android/UEBuildAndroid.cs
+++ b/Engine/Source/Programs/UnrealBuildTool/Android/UEBuildAndroid.cs
@@ -72,24 +72,27 @@ namespace UnrealBuildTool
 			{
 				// Try reading env variables we need from .bash_profile
 				string BashProfilePath = Path.Combine(Environment.GetFolderPath(Environment.SpecialFolder.Personal), ".bash_profile");
-				string[] BashProfileContents = File.ReadAllLines(BashProfilePath);
-				foreach (string Line in BashProfileContents)
+				if(File.Exists(BashProfilePath))
 				{
-					if (bNeedsAndroidHome && Line.StartsWith("export ANDROID_HOME="))
+					string[] BashProfileContents = File.ReadAllLines(BashProfilePath);
+					foreach (string Line in BashProfileContents)
 					{
-						string PathVar = Line.Split('=')[1].Replace("\"", "");
-						Environment.SetEnvironmentVariable("ANDROID_HOME", PathVar);
-					}
-					else if (bNeedsNDKPath && Line.StartsWith("export NDKROOT="))
-					{
-						string PathVar = Line.Split('=')[1].Replace("\"", "");
-						Environment.SetEnvironmentVariable("NDKROOT", PathVar);
-						NDKPath = PathVar;
-					}
-					else if (bNeedsAntHome && Line.StartsWith("export ANT_HOME="))
-					{
-						string PathVar = Line.Split('=')[1].Replace("\"", "");
-						Environment.SetEnvironmentVariable("ANT_HOME", PathVar);
+						if (bNeedsAndroidHome && Line.StartsWith("export ANDROID_HOME="))
+						{
+							string PathVar = Line.Split('=')[1].Replace("\"", "");
+							Environment.SetEnvironmentVariable("ANDROID_HOME", PathVar);
+						}
+						else if (bNeedsNDKPath && Line.StartsWith("export NDKROOT="))
+						{
+							string PathVar = Line.Split('=')[1].Replace("\"", "");
+							Environment.SetEnvironmentVariable("NDKROOT", PathVar);
+							NDKPath = PathVar;
+						}
+						else if (bNeedsAntHome && Line.StartsWith("export ANT_HOME="))
+						{
+							string PathVar = Line.Split('=')[1].Replace("\"", "");
+							Environment.SetEnvironmentVariable("ANT_HOME", PathVar);
+						}
 					}
 				}
 			}
