Fix unused variable warning on Mono.

diff --git a/Engine/Source/Programs/UnrealBuildTool/Android/UEDeployAndroid.cs b/Engine/Source/Programs/UnrealBuildTool/Android/UEDeployAndroid.cs
--- a/Engine/Source/Programs/UnrealBuildTool/Android/UEDeployAndroid.cs
+++ b/Engine/Source/Programs/UnrealBuildTool/Android/UEDeployAndroid.cs
@@ -580,7 +580,7 @@ namespace UnrealBuildTool.Android
 	
 			string[] Arches = AndroidToolChain.GetAllArchitectures();
 			string[] GPUArchitectures = AndroidToolChain.GetAllGPUArchitectures();
-			int NumArches = Arches.Length * GPUArchitectures.Length;
+//			int NumArches = Arches.Length * GPUArchitectures.Length;
 
 
             // See if we need to create a 'default' Java Build settings file if one doesn't exist (if it does exist we have to assume it has been setup correctly)
