Only build enabled plugins as part of Rocket builds again. UHT currently tries to link against them in monollithic builds.

diff --git a/Engine/Source/Programs/UnrealBuildTool/Configuration/UEBuildTarget.cs b/Engine/Source/Programs/UnrealBuildTool/Configuration/UEBuildTarget.cs
--- a/Engine/Source/Programs/UnrealBuildTool/Configuration/UEBuildTarget.cs
+++ b/Engine/Source/Programs/UnrealBuildTool/Configuration/UEBuildTarget.cs
@@ -2451,9 +2451,9 @@ namespace UnrealBuildTool
 			DependentPlugins.AddRange(EnabledPlugins);
 
 			// Set the list of plugins that should be built
-			if( UnrealBuildTool.BuildingRocket() && TargetType != TargetRules.TargetType.Program )
+			if (UnrealBuildTool.BuildingRocket() && TargetType != TargetRules.TargetType.Program)
 			{
-				BuildPlugins.AddRange(ValidPlugins);
+				BuildPlugins.AddRange(EnabledPlugins);//ValidPlugins);
 			}
 			else if (ShouldCompileMonolithic() || TargetType == TargetRules.TargetType.Program)
 			{
