UnrealBuildTool: Don't mix RTTI modes when compiling modules using shared PCHs

diff --git a/Engine/Source/Programs/UnrealBuildTool/Configuration/UEBuildModule.cs b/Engine/Source/Programs/UnrealBuildTool/Configuration/UEBuildModule.cs
--- a/Engine/Source/Programs/UnrealBuildTool/Configuration/UEBuildModule.cs
+++ b/Engine/Source/Programs/UnrealBuildTool/Configuration/UEBuildModule.cs
@@ -1560,6 +1560,13 @@ namespace UnrealBuildTool
 				SharedPCHHeaderFile = null;
 				return null;
 			}
+			// Don't mix RTTI modes
+			if (bUseRTTI)
+			{
+				Log.TraceVerbose("Module {0} cannot use existing Shared PCH '{1}' (from module '{2}') because RTTI modes don't match", Name, SharedPCHEnvironment.PrecompiledHeaderIncludeFilename.AbsolutePath, SharedPCHEnvironment.ModuleName);
+				SharedPCHHeaderFile = null;
+				return null;
+			}
 
 			// Don't mix non-optimized code with optimized code (PCHs won't be compatible)
 			var SharedPCHCodeOptimization = SharedPCHEnvironment.OptimizeCode;
