Removed conditional inclusion of most platform files in UBT project so the same project can be used in Visual Studio and Mono Develop/Xamarin Studio, removed _Mono UBT and UAT projects, updated automation scripts to be compatible with Mono 3.10

diff --git a/Engine/Build/BatchFiles/Linux/Build.sh b/Engine/Build/BatchFiles/Linux/Build.sh
--- a/Engine/Build/BatchFiles/Linux/Build.sh
+++ b/Engine/Build/BatchFiles/Linux/Build.sh
@@ -7,7 +7,7 @@ set -e
 cd "`dirname "$0"`/../../../.." 
 
 # First make sure that the UnrealBuildTool is up-to-date
-if ! xbuild /property:Configuration=Development /property:TargetFrameworkVersion=v4.0 /verbosity:quiet /nologo Engine/Source/Programs/UnrealBuildTool/UnrealBuildTool_Mono.csproj; then
+if ! xbuild /property:Configuration=Development /property:TargetFrameworkVersion=v4.0 /verbosity:quiet /nologo Engine/Source/Programs/UnrealBuildTool/UnrealBuildTool.csproj; then
   echo "Failed to build to build tool (UnrealBuildTool)"
   exit 1
 fi
diff --git a/Engine/Build/BatchFiles/Linux/GenerateProjectFiles.sh b/Engine/Build/BatchFiles/Linux/GenerateProjectFiles.sh
--- a/Engine/Build/BatchFiles/Linux/GenerateProjectFiles.sh
+++ b/Engine/Build/BatchFiles/Linux/GenerateProjectFiles.sh
@@ -19,12 +19,12 @@ echo Setting up Unreal Engine 4 project files...
 echo
 
 set -x
-xbuild Source/Programs/UnrealBuildTool/UnrealBuildTool_Mono.csproj \
+xbuild Source/Programs/UnrealBuildTool/UnrealBuildTool.csproj \
   /verbosity:quiet /nologo \
   /p:TargetFrameworkVersion=v4.0 \
   /p:Configuration="Development"
 
-xbuild Source/Programs/AutomationTool/AutomationTool_Mono.csproj \
+xbuild Source/Programs/AutomationTool/AutomationTool.csproj \
   /verbosity:quiet /nologo \
   /p:TargetFrameworkVersion=v4.0 \
   /p:Platform="AnyCPU" \
diff --git a/Engine/Build/BatchFiles/Mac/Build.sh b/Engine/Build/BatchFiles/Mac/Build.sh
--- a/Engine/Build/BatchFiles/Mac/Build.sh
+++ b/Engine/Build/BatchFiles/Mac/Build.sh
@@ -15,7 +15,7 @@ fi
 
 echo Building UBT... 
 
-xbuild /property:Configuration=Development  /nologo Engine/Source/Programs/UnrealBuildTool/UnrealBuildTool_Mono.csproj| grep -i error  
+xbuild /property:Configuration=Development  /nologo Engine/Source/Programs/UnrealBuildTool/UnrealBuildTool.csproj| grep -i error  
 
 case $ACTION in
 	"")
diff --git a/Engine/Build/BatchFiles/Mac/GenerateProjectFiles.sh b/Engine/Build/BatchFiles/Mac/GenerateProjectFiles.sh
--- a/Engine/Build/BatchFiles/Mac/GenerateProjectFiles.sh
+++ b/Engine/Build/BatchFiles/Mac/GenerateProjectFiles.sh
@@ -18,7 +18,7 @@ fi
 
 source SetupMono.sh "`dirname "$0"`"
 
-xbuild ../../../Source/Programs/UnrealBuildTool/UnrealBuildTool_Mono.csproj /property:Configuration="Development" /verbosity:quiet /nologo |grep -i error
+xbuild ../../../Source/Programs/UnrealBuildTool/UnrealBuildTool.csproj /property:Configuration="Development" /verbosity:quiet /nologo |grep -i error
 
 # pass all parameters to UBT
 mono ../../../Binaries/DotNET/UnrealBuildTool.exe -XcodeProjectFile "$@"
diff --git a/Engine/Build/BatchFiles/Mac/SetupMono.sh b/Engine/Build/BatchFiles/Mac/SetupMono.sh
--- a/Engine/Build/BatchFiles/Mac/SetupMono.sh
+++ b/Engine/Build/BatchFiles/Mac/SetupMono.sh
@@ -14,8 +14,8 @@ if [ -f "/usr/bin/mono" ]; then
 	if [ ${MONO_VERSION[0]} -ge 3 ]; then
 		if [ ${MONO_VERSION[1]} -eq 2 ] && [ ${MONO_VERSION[2]} -ge 6 ]; then
 			IS_MONO_INSTALLED=1
-#		elif [ ${MONO_VERSION[1]} -gt 2 ]; then # @todo: enable this when problems with running on Mono 3.4 are solved
-#			IS_MONO_INSTALLED=1
+		elif [ ${MONO_VERSION[1]} -gt 2 ]; then
+			IS_MONO_INSTALLED=1
 		fi
 	fi
 fi
diff --git a/Engine/Build/BatchFiles/RunUAT.sh b/Engine/Build/BatchFiles/RunUAT.sh
--- a/Engine/Build/BatchFiles/RunUAT.sh
+++ b/Engine/Build/BatchFiles/RunUAT.sh
@@ -53,7 +53,7 @@ fi
 
 if [ "$UATNoCompileArg" != "-NoCompile" ]; then
   # see if the .csproj exists to be compiled
-	if [ ! -f Source/Programs/AutomationTool/AutomationTool_Mono.csproj ]; then
+	if [ ! -f Source/Programs/AutomationTool/AutomationTool.csproj ]; then
 		echo No project to compile, attempting to use precompiled AutomationTool
 		UATNoCompileArg=-NoCompile
 	else
@@ -68,13 +68,10 @@ if [ "$UATNoCompileArg" != "-NoCompile" ]; then
 		echo Compiling AutomationTool with xbuild
 
 		ARGS="/p:Configuration=Development /p:Platform=AnyCPU /verbosity:quiet /nologo"
-    ARGS="${ARGS} /p:TargetFrameworkProfile="
-		if [ "$(uname)" = "Darwin" ]; then
-		  ARGS="${ARGS} /p:TargetFrameworkVersion=v4.5"
-		fi
+		ARGS="${ARGS} /p:TargetFrameworkProfile="
 
-		echo "xbuild Source/Programs/AutomationTool/AutomationTool_Mono.csproj $ARGS"
-		xbuild Source/Programs/AutomationTool/AutomationTool_Mono.csproj $ARGS
+		echo "xbuild Source/Programs/AutomationTool/AutomationTool.csproj $ARGS"
+		xbuild Source/Programs/AutomationTool/AutomationTool.csproj $ARGS
 		# make sure it succeeded
 		if [ $? -ne 0 ]; then
 			echo RunUAT ERROR: AutomationTool failed to compile.
diff --git a/Engine/Source/Developer/DesktopPlatform/Private/DesktopPlatformBase.cpp b/Engine/Source/Developer/DesktopPlatform/Private/DesktopPlatformBase.cpp
--- a/Engine/Source/Developer/DesktopPlatform/Private/DesktopPlatformBase.cpp
+++ b/Engine/Source/Developer/DesktopPlatform/Private/DesktopPlatformBase.cpp
@@ -968,12 +968,12 @@ bool FDesktopPlatformBase::BuildUnrealBuildTool(FOutputDevice &Ar)
 		CmdLineParams += BuildBatchFile;
 	}
 #elif PLATFORM_MAC
-	const FString CsProjLocation = FPaths::ConvertRelativePathToFull(GetUnrealBuildToolSourceCodePath()) / TEXT("UnrealBuildTool_Mono.csproj");
+	const FString CsProjLocation = FPaths::ConvertRelativePathToFull(GetUnrealBuildToolSourceCodePath()) / TEXT("UnrealBuildTool.csproj");
 	FString ScriptPath = FPaths::ConvertRelativePathToFull(FPaths::EngineDir() / TEXT("Build/BatchFiles/Mac/RunXBuild.sh"));
 	CompilerExecutableFilename = TEXT("/bin/sh");
 	CmdLineParams = FString::Printf(TEXT("\"%s\" /property:Configuration=Development %s"), *ScriptPath, *CsProjLocation);
 #elif PLATFORM_LINUX
-	const FString CsProjLocation = FPaths::ConvertRelativePathToFull(GetUnrealBuildToolSourceCodePath()) / TEXT("UnrealBuildTool_Mono.csproj");
+	const FString CsProjLocation = FPaths::ConvertRelativePathToFull(GetUnrealBuildToolSourceCodePath()) / TEXT("UnrealBuildTool.csproj");
 	FString ScriptPath = FPaths::ConvertRelativePathToFull(FPaths::EngineDir() / TEXT("Build/BatchFiles/Linux/RunXBuild.sh"));
 	CompilerExecutableFilename = TEXT("/bin/bash");
 	CmdLineParams = FString::Printf(TEXT("\"%s\" /property:Configuration=Development /property:TargetFrameworkVersion=v4.0 %s"), *ScriptPath, *CsProjLocation);
diff --git a/Engine/Source/Programs/AutomationTool/Android/Android.Automation.csproj b/Engine/Source/Programs/AutomationTool/Android/Android.Automation.csproj
--- a/Engine/Source/Programs/AutomationTool/Android/Android.Automation.csproj
+++ b/Engine/Source/Programs/AutomationTool/Android/Android.Automation.csproj
@@ -1,5 +1,5 @@
 ﻿<?xml version="1.0" encoding="utf-8"?>
-<Project ToolsVersion="11.0" DefaultTargets="Build" xmlns="http://schemas.microsoft.com/developer/msbuild/2003">
+<Project ToolsVersion="4.0" DefaultTargets="Build" xmlns="http://schemas.microsoft.com/developer/msbuild/2003">
   <Import Project="$(MSBuildExtensionsPath)\$(MSBuildToolsVersion)\Microsoft.Common.props" Condition="Exists('$(MSBuildExtensionsPath)\$(MSBuildToolsVersion)\Microsoft.Common.props')" />
   <PropertyGroup>
     <Configuration Condition=" '$(Configuration)' == '' ">Debug</Configuration>
diff --git a/Engine/Source/Programs/AutomationTool/AutomationTool_Mono.csproj b/Engine/Source/Programs/AutomationTool/AutomationTool_Mono.csproj
deleted file mode 100644
diff --git a/Engine/Source/Programs/AutomationTool/HTML5/HTML5.Automation.csproj b/Engine/Source/Programs/AutomationTool/HTML5/HTML5.Automation.csproj
--- a/Engine/Source/Programs/AutomationTool/HTML5/HTML5.Automation.csproj
+++ b/Engine/Source/Programs/AutomationTool/HTML5/HTML5.Automation.csproj
@@ -1,5 +1,5 @@
 ﻿<?xml version="1.0" encoding="utf-8"?>
-<Project ToolsVersion="11.0" DefaultTargets="Build" xmlns="http://schemas.microsoft.com/developer/msbuild/2003">
+<Project ToolsVersion="4.0" DefaultTargets="Build" xmlns="http://schemas.microsoft.com/developer/msbuild/2003">
   <Import Project="$(MSBuildExtensionsPath)\$(MSBuildToolsVersion)\Microsoft.Common.props" Condition="Exists('$(MSBuildExtensionsPath)\$(MSBuildToolsVersion)\Microsoft.Common.props')" />
   <PropertyGroup>
     <Configuration Condition=" '$(Configuration)' == '' ">Debug</Configuration>
diff --git a/Engine/Source/Programs/AutomationTool/IOS/IOS.Automation.csproj b/Engine/Source/Programs/AutomationTool/IOS/IOS.Automation.csproj
--- a/Engine/Source/Programs/AutomationTool/IOS/IOS.Automation.csproj
+++ b/Engine/Source/Programs/AutomationTool/IOS/IOS.Automation.csproj
@@ -1,5 +1,5 @@
 ﻿<?xml version="1.0" encoding="utf-8"?>
-<Project ToolsVersion="11.0" DefaultTargets="Build" xmlns="http://schemas.microsoft.com/developer/msbuild/2003">
+<Project ToolsVersion="4.0" DefaultTargets="Build" xmlns="http://schemas.microsoft.com/developer/msbuild/2003">
   <Import Project="$(MSBuildExtensionsPath)\$(MSBuildToolsVersion)\Microsoft.Common.props" Condition="Exists('$(MSBuildExtensionsPath)\$(MSBuildToolsVersion)\Microsoft.Common.props')" />
   <PropertyGroup>
     <Configuration Condition=" '$(Configuration)' == '' ">Debug</Configuration>
diff --git a/Engine/Source/Programs/AutomationTool/Linux/Linux.Automation.csproj b/Engine/Source/Programs/AutomationTool/Linux/Linux.Automation.csproj
--- a/Engine/Source/Programs/AutomationTool/Linux/Linux.Automation.csproj
+++ b/Engine/Source/Programs/AutomationTool/Linux/Linux.Automation.csproj
@@ -1,5 +1,5 @@
 ﻿<?xml version="1.0" encoding="utf-8"?>
-<Project ToolsVersion="11.0" DefaultTargets="Build" xmlns="http://schemas.microsoft.com/developer/msbuild/2003">
+<Project ToolsVersion="4.0" DefaultTargets="Build" xmlns="http://schemas.microsoft.com/developer/msbuild/2003">
   <Import Project="$(MSBuildExtensionsPath)\$(MSBuildToolsVersion)\Microsoft.Common.props" Condition="Exists('$(MSBuildExtensionsPath)\$(MSBuildToolsVersion)\Microsoft.Common.props')" />
   <PropertyGroup>
     <Configuration Condition=" '$(Configuration)' == '' ">Debug</Configuration>
diff --git a/Engine/Source/Programs/AutomationTool/LinuxHostPlatform.cs b/Engine/Source/Programs/AutomationTool/LinuxHostPlatform.cs
--- a/Engine/Source/Programs/AutomationTool/LinuxHostPlatform.cs
+++ b/Engine/Source/Programs/AutomationTool/LinuxHostPlatform.cs
@@ -109,7 +109,7 @@ namespace AutomationTool
 
 		public override string UBTProjectName
 		{
-			get { return "UnrealBuildTool_Mono"; }
+			get { return "UnrealBuildTool"; }
 		}
 
 		public override UnrealTargetPlatform HostEditorPlatform
diff --git a/Engine/Source/Programs/AutomationTool/Mac/Mac.Automation.csproj b/Engine/Source/Programs/AutomationTool/Mac/Mac.Automation.csproj
--- a/Engine/Source/Programs/AutomationTool/Mac/Mac.Automation.csproj
+++ b/Engine/Source/Programs/AutomationTool/Mac/Mac.Automation.csproj
@@ -1,5 +1,5 @@
 ﻿<?xml version="1.0" encoding="utf-8"?>
-<Project ToolsVersion="11.0" DefaultTargets="Build" xmlns="http://schemas.microsoft.com/developer/msbuild/2003">
+<Project ToolsVersion="4.0" DefaultTargets="Build" xmlns="http://schemas.microsoft.com/developer/msbuild/2003">
   <Import Project="$(MSBuildExtensionsPath)\$(MSBuildToolsVersion)\Microsoft.Common.props" Condition="Exists('$(MSBuildExtensionsPath)\$(MSBuildToolsVersion)\Microsoft.Common.props')" />
   <PropertyGroup>
     <Configuration Condition=" '$(Configuration)' == '' ">Debug</Configuration>
diff --git a/Engine/Source/Programs/AutomationTool/MacHostPlatform.cs b/Engine/Source/Programs/AutomationTool/MacHostPlatform.cs
--- a/Engine/Source/Programs/AutomationTool/MacHostPlatform.cs
+++ b/Engine/Source/Programs/AutomationTool/MacHostPlatform.cs
@@ -93,7 +93,7 @@ namespace AutomationTool
 			if (AppName == "xbuild")
 			{
 				AppName = "sh";
-				CommandLine = "-c 'xbuild " + (String.IsNullOrEmpty(CommandLine) ? "" : CommandLine) + " /property:TargetFrameworkVersion=v4.5 /tv:4.0 /p:DefineConstants=MONO /verbosity:quiet /nologo |grep -i error; if [ $? -ne 1 ]; then exit 1; else exit 0; fi'";
+				CommandLine = "-c 'xbuild " + (String.IsNullOrEmpty(CommandLine) ? "" : CommandLine) + " /p:DefineConstants=MONO /verbosity:quiet /nologo |grep -i error; if [ $? -ne 1 ]; then exit 1; else exit 0; fi'";
 			}
 			if (AppName.EndsWith(".exe") || ((AppName.Contains("/Binaries/Win64/") || AppName.Contains("/Binaries/Mac/")) && string.IsNullOrEmpty(Path.GetExtension(AppName))))
 			{
@@ -132,7 +132,7 @@ namespace AutomationTool
 
 		public override string UBTProjectName
 		{
-			get { return "UnrealBuildTool_Mono"; }
+			get { return "UnrealBuildTool"; }
 		}
 
 		public override UnrealTargetPlatform HostEditorPlatform
diff --git a/Engine/Source/Programs/AutomationTool/Win/Win.Automation.csproj b/Engine/Source/Programs/AutomationTool/Win/Win.Automation.csproj
--- a/Engine/Source/Programs/AutomationTool/Win/Win.Automation.csproj
+++ b/Engine/Source/Programs/AutomationTool/Win/Win.Automation.csproj
@@ -1,5 +1,5 @@
 ﻿<?xml version="1.0" encoding="utf-8"?>
-<Project ToolsVersion="11.0" DefaultTargets="Build" xmlns="http://schemas.microsoft.com/developer/msbuild/2003">
+<Project ToolsVersion="4.0" DefaultTargets="Build" xmlns="http://schemas.microsoft.com/developer/msbuild/2003">
   <Import Project="$(MSBuildExtensionsPath)\$(MSBuildToolsVersion)\Microsoft.Common.props" Condition="Exists('$(MSBuildExtensionsPath)\$(MSBuildToolsVersion)\Microsoft.Common.props')" />
   <PropertyGroup>
     <Configuration Condition=" '$(Configuration)' == '' ">Debug</Configuration>
diff --git a/Engine/Source/Programs/UnrealBuildTool/Mac/MacToolChain.cs b/Engine/Source/Programs/UnrealBuildTool/Mac/MacToolChain.cs
--- a/Engine/Source/Programs/UnrealBuildTool/Mac/MacToolChain.cs
+++ b/Engine/Source/Programs/UnrealBuildTool/Mac/MacToolChain.cs
@@ -1285,7 +1285,7 @@ namespace UnrealBuildTool
 				}
 			}
 
-			string XBuildArgs = "/verbosity:quiet /nologo /target:Rebuild /property:Configuration=Development /property:Platform=AnyCPU " + Path.GetFileName(ProjectFileName) + " /property:TargetFrameworkVersion=v4.5 /tv:4.0";
+			string XBuildArgs = "/verbosity:quiet /nologo /target:Rebuild /property:Configuration=Development /property:Platform=AnyCPU " + Path.GetFileName(ProjectFileName);
 
 			if (BuildHostPlatform.Current.Platform != UnrealTargetPlatform.Mac)
 			{
diff --git a/Engine/Source/Programs/UnrealBuildTool/System/XcodeProjectFileGenerator.cs b/Engine/Source/Programs/UnrealBuildTool/System/XcodeProjectFileGenerator.cs
--- a/Engine/Source/Programs/UnrealBuildTool/System/XcodeProjectFileGenerator.cs
+++ b/Engine/Source/Programs/UnrealBuildTool/System/XcodeProjectFileGenerator.cs
@@ -979,7 +979,7 @@ namespace UnrealBuildTool
 			if (!bGeneratingRocketProjectFiles)
 			{
 				// Add UnrealBuildTool to the master project
-				string ProjectPath = System.IO.Path.Combine(System.IO.Path.Combine(EngineRelativePath, "Source"), "Programs", "UnrealBuildTool", "UnrealBuildTool_Mono.csproj");
+				string ProjectPath = System.IO.Path.Combine(System.IO.Path.Combine(EngineRelativePath, "Source"), "Programs", "UnrealBuildTool", "UnrealBuildTool.csproj");
 				AddPreGeneratedProject(ref PBXBuildFileSection, ref PBXFileReferenceSection, ref PBXSourcesBuildPhaseSection, ref Groups, ProjectPath);
 			}
 		}
diff --git a/Engine/Source/Programs/UnrealBuildTool/UnrealBuildTool.csproj b/Engine/Source/Programs/UnrealBuildTool/UnrealBuildTool.csproj
--- a/Engine/Source/Programs/UnrealBuildTool/UnrealBuildTool.csproj
+++ b/Engine/Source/Programs/UnrealBuildTool/UnrealBuildTool.csproj
@@ -122,26 +122,18 @@
       </ItemGroup>
     </When>
   </Choose>
-  <Choose>
-    <When Condition="Exists('Mac\UEBuildMac.cs')">
-      <ItemGroup>
-        <Compile Include="Mac\MacToolChain.cs" />
-        <Compile Include="Mac\UEBuildMac.cs" />
-        <Compile Include="Mac\MacProjectGenerator.cs" />
-        <Compile Include="Mac\UEDeployMac.cs" />
-      </ItemGroup>
-    </When>
-  </Choose>
-  <Choose>
-    <When Condition="Exists('Android\UEBuildAndroid.cs')">
-      <ItemGroup>
-        <Compile Include="Android\UEBuildAndroid.cs" />
-        <Compile Include="Android\UEDeployAndroid.cs" />
-        <Compile Include="Android\AndroidToolChain.cs" />
-        <Compile Include="Android\AndroidProjectGenerator.cs" />
-      </ItemGroup>
-    </When>
-  </Choose>
+  <ItemGroup>
+	<Compile Include="Mac\MacToolChain.cs" />
+	<Compile Include="Mac\UEBuildMac.cs" />
+	<Compile Include="Mac\MacProjectGenerator.cs" />
+	<Compile Include="Mac\UEDeployMac.cs" />
+  </ItemGroup>
+  <ItemGroup>
+	<Compile Include="Android\UEBuildAndroid.cs" />
+	<Compile Include="Android\UEDeployAndroid.cs" />
+	<Compile Include="Android\AndroidToolChain.cs" />
+	<Compile Include="Android\AndroidProjectGenerator.cs" />
+  </ItemGroup>
   <Choose>
     <When Condition="Exists('WinRT\UEBuildWinRT.cs')">
       <ItemGroup>
@@ -152,35 +144,23 @@
       </ItemGroup>
     </When>
   </Choose>
-  <Choose>
-    <When Condition="Exists('IOS\UEBuildIOS.cs')">
-      <ItemGroup>
-        <Compile Include="IOS\IOSProjectGenerator.cs" />
-        <Compile Include="IOS\IOSToolChain.cs" />
-        <Compile Include="IOS\UEBuildIOS.cs" />
-        <Compile Include="IOS\UEDeployIOS.cs" />
-      </ItemGroup>
-    </When>
-  </Choose>
-  <Choose>
-    <When Condition="Exists('HTML5\UEBuildHTML5.cs')">
-      <ItemGroup>
-        <Compile Include="HTML5\HTML5ProjectGenerator.cs" />
-        <Compile Include="HTML5\HTML5ToolChain.cs" />
-        <Compile Include="HTML5\UEBuildHTML5.cs" />
-        <Compile Include="HTML5\HTML5SDKInfo.cs" />
-      </ItemGroup>
-    </When>
-  </Choose>
-  <Choose>
-    <When Condition="Exists('Linux\UEBuildLinux.cs')">
-      <ItemGroup>
-        <Compile Include="Linux\LinuxProjectGenerator.cs" />
-        <Compile Include="Linux\LinuxToolChain.cs" />
-        <Compile Include="Linux\UEBuildLinux.cs" />
-      </ItemGroup>
-    </When>
-  </Choose>
+  <ItemGroup>
+	<Compile Include="IOS\IOSProjectGenerator.cs" />
+	<Compile Include="IOS\IOSToolChain.cs" />
+	<Compile Include="IOS\UEBuildIOS.cs" />
+	<Compile Include="IOS\UEDeployIOS.cs" />
+  </ItemGroup>
+  <ItemGroup>
+	<Compile Include="HTML5\HTML5ProjectGenerator.cs" />
+	<Compile Include="HTML5\HTML5ToolChain.cs" />
+	<Compile Include="HTML5\UEBuildHTML5.cs" />
+	<Compile Include="HTML5\HTML5SDKInfo.cs" />
+  </ItemGroup>
+  <ItemGroup>
+	<Compile Include="Linux\LinuxProjectGenerator.cs" />
+	<Compile Include="Linux\LinuxToolChain.cs" />
+	<Compile Include="Linux\UEBuildLinux.cs" />
+  </ItemGroup>
   <Choose>
     <When Condition="Exists('System\NotForLicensees\TelemetryProviderET.cs')">
       <ItemGroup>
diff --git a/Engine/Source/Programs/UnrealBuildTool/UnrealBuildTool_Mono.csproj b/Engine/Source/Programs/UnrealBuildTool/UnrealBuildTool_Mono.csproj
deleted file mode 100644
