Unreal Build Tool: Enabled remote linking with XGE on Windows platform
- Especially in modular builds, remote linking speeds up build times very significantly
- Also enabled remote resource compilation, which also makes modular builds faster when XGE is installed

diff --git a/Engine/Source/Programs/UnrealBuildTool/Windows/VCToolChain.cs b/Engine/Source/Programs/UnrealBuildTool/Windows/VCToolChain.cs
--- a/Engine/Source/Programs/UnrealBuildTool/Windows/VCToolChain.cs
+++ b/Engine/Source/Programs/UnrealBuildTool/Windows/VCToolChain.cs
@@ -1059,6 +1059,9 @@ namespace UnrealBuildTool
 				CompileAction.CommandPath        = EnvVars.ResourceCompilerPath;
 				CompileAction.StatusDescription  = Path.GetFileName(RCFile.AbsolutePath);
 
+				// Resource tool can run remotely if possible
+				CompileAction.bCanExecuteRemotely = true;
+
 				if( WindowsPlatform.bCompileWithClang )
 				{ 
 					CompileAction.OutputEventHandler = new DataReceivedEventHandler( ClangCompilerOutputFormatter );
@@ -1334,8 +1337,8 @@ namespace UnrealBuildTool
 			// ignored as a prerequisite for other actions
 			LinkAction.bProducesImportLibrary = bBuildImportLibraryOnly || LinkEnvironment.Config.bIsBuildingDLL;
 
-			// Only execute linking on the local PC.
-			LinkAction.bCanExecuteRemotely = false;
+			// Allow remote linking.  Especially in modular builds with many small DLL files, this is almost always very efficient
+			LinkAction.bCanExecuteRemotely = true;
 
 			Log.TraceVerbose( "     Linking: " + LinkAction.StatusDescription );
 			Log.TraceVerbose( "     Command: " + LinkAction.CommandArguments );
