Fix for some compiler warnings in windows xp builds. Seems that there are issues with the windows 7.1a sdk that cause this. The worksaround is to use the _USING_V110_SDK71_ compiler flag as opposed to D_USING_V120_SDK71_ for xp compilation.

Tested that this works with the vs2013 compilers.

diff --git a/Engine/Source/Programs/UnrealBuildTool/Windows/VCToolChain.cs b/Engine/Source/Programs/UnrealBuildTool/Windows/VCToolChain.cs
--- a/Engine/Source/Programs/UnrealBuildTool/Windows/VCToolChain.cs
+++ b/Engine/Source/Programs/UnrealBuildTool/Windows/VCToolChain.cs
@@ -126,14 +126,7 @@ namespace UnrealBuildTool
 			// with Windows XP (http://blogs.msdn.com/b/vcblog/archive/2012/10/08/10357555.aspx)
 			if( WindowsPlatform.SupportWindowsXP )
 			{
-				if( WindowsPlatform.Compiler == WindowsCompiler.VisualStudio2013 )
-				{
-					Arguments.Append(" /D_USING_V120_SDK71_");
-				}
-				else if (WindowsPlatform.Compiler == WindowsCompiler.VisualStudio2012)
-				{
-					Arguments.Append(" /D_USING_V110_SDK71_");
-				}
+				Arguments.Append(" /D_USING_V110_SDK71_");
 			}
 
 
@@ -1081,14 +1074,7 @@ namespace UnrealBuildTool
 				// with Windows XP (http://blogs.msdn.com/b/vcblog/archive/2012/10/08/10357555.aspx)
 				if (WindowsPlatform.SupportWindowsXP)
 				{
-					if (WindowsPlatform.Compiler == WindowsCompiler.VisualStudio2013)
-					{
-						CompileAction.CommandArguments += " /D_USING_V120_SDK71_";
-					}
-					else if (WindowsPlatform.Compiler == WindowsCompiler.VisualStudio2012)
-					{
-						CompileAction.CommandArguments += " /D_USING_V110_SDK71_";
-					}
+					CompileAction.CommandArguments += " /D_USING_V110_SDK71_";
 				}
 
 				// Language
