Allow higher versioned emscripten SDKs than required SDK as Valid SDKs

#HTML5

diff --git a/Engine/Source/Programs/UnrealBuildTool/HTML5/UEBuildHTML5.cs b/Engine/Source/Programs/UnrealBuildTool/HTML5/UEBuildHTML5.cs
--- a/Engine/Source/Programs/UnrealBuildTool/HTML5/UEBuildHTML5.cs
+++ b/Engine/Source/Programs/UnrealBuildTool/HTML5/UEBuildHTML5.cs
@@ -88,13 +88,16 @@ namespace UnrealBuildTool
 			}
             try
             {
-				if (HTML5SDKInfo.EmscriptenVersion().Contains(GetRequiredSDKString()))
+                int InstalledVersion = Convert.ToInt32(HTML5SDKInfo.EmscriptenVersion().Replace(".","")); 
+                int RequiredVersion = Convert.ToInt32(GetRequiredSDKString().Replace(".","")); 
+
+                if (InstalledVersion > RequiredVersion)
                 {
                     return SDKStatus.Valid;
                 }
                 else
                 {
-					Console.WriteLine("EMSCRIPTEN sdk " + HTML5SDKInfo.EmscriptenVersion() + " found which is unsupported, Please install version " + GetRequiredSDKString());
+					Console.WriteLine("EMSCRIPTEN sdk " + HTML5SDKInfo.EmscriptenVersion() + " found which is older than " + RequiredVersion + " Please install the latest emscripten SDK");
                     return SDKStatus.Invalid;
                 }
             }
