Correct comparision (>=) for SDK versions for CL 2390651

diff --git a/Engine/Source/Programs/UnrealBuildTool/HTML5/UEBuildHTML5.cs b/Engine/Source/Programs/UnrealBuildTool/HTML5/UEBuildHTML5.cs
--- a/Engine/Source/Programs/UnrealBuildTool/HTML5/UEBuildHTML5.cs
+++ b/Engine/Source/Programs/UnrealBuildTool/HTML5/UEBuildHTML5.cs
@@ -91,7 +91,7 @@ namespace UnrealBuildTool
                 int InstalledVersion = Convert.ToInt32(HTML5SDKInfo.EmscriptenVersion().Replace(".","")); 
                 int RequiredVersion = Convert.ToInt32(GetRequiredSDKString().Replace(".","")); 
 
-                if (InstalledVersion > RequiredVersion)
+                if (InstalledVersion >= RequiredVersion)
                 {
                     return SDKStatus.Valid;
                 }
