Added Unreal Build Tool performance diagnostic info for shared precompiled header usage
- Requires BuildConfiguration.bPrintPerformanceInfo to be enabled

diff --git a/Engine/Source/Programs/UnrealBuildTool/Configuration/UEBuildModule.cs b/Engine/Source/Programs/UnrealBuildTool/Configuration/UEBuildModule.cs
--- a/Engine/Source/Programs/UnrealBuildTool/Configuration/UEBuildModule.cs
+++ b/Engine/Source/Programs/UnrealBuildTool/Configuration/UEBuildModule.cs
@@ -1322,6 +1322,11 @@ namespace UnrealBuildTool
 						{
 							throw new BuildException( "Module {0} doesn't use a Shared PCH!  Please add a dependency on a Shared PCH module to this module's dependency list", this.Name);
 						}
+
+						// Keep track of how many modules make use of this PCH for performance diagnostics
+						var LargestSharedPCHHeader = GlobalCompileEnvironment.SharedPCHHeaderFiles[ LargestSharedPCHHeaderFileIndex ];
+						++LargestSharedPCHHeader.NumModulesUsingThisPCH;
+
 					}
 					else
 					{
diff --git a/Engine/Source/Programs/UnrealBuildTool/Configuration/UEBuildTarget.cs b/Engine/Source/Programs/UnrealBuildTool/Configuration/UEBuildTarget.cs
--- a/Engine/Source/Programs/UnrealBuildTool/Configuration/UEBuildTarget.cs
+++ b/Engine/Source/Programs/UnrealBuildTool/Configuration/UEBuildTarget.cs
@@ -1569,6 +1569,14 @@ namespace UnrealBuildTool
 				OutputItems.AddRange(Binary.Build(TargetToolChain, GlobalCompileEnvironment, GlobalLinkEnvironment));
 			}
 
+			if( BuildConfiguration.bPrintPerformanceInfo )
+			{
+				foreach( var SharedPCH in GlobalCompileEnvironment.SharedPCHHeaderFiles )
+				{
+					Log.TraceInformation( "Shared PCH '" + SharedPCH.Module.Name + "': Used " + SharedPCH.NumModulesUsingThisPCH + " times" );
+				}
+			}
+
 			return ECompilationResult.Succeeded;
 		}
 
diff --git a/Engine/Source/Programs/UnrealBuildTool/System/CPPEnvironment.cs b/Engine/Source/Programs/UnrealBuildTool/System/CPPEnvironment.cs
--- a/Engine/Source/Programs/UnrealBuildTool/System/CPPEnvironment.cs
+++ b/Engine/Source/Programs/UnrealBuildTool/System/CPPEnvironment.cs
@@ -299,6 +299,9 @@ namespace UnrealBuildTool
 
 		/// All dependent modules
 		public Dictionary<string, UEBuildModule> Dependencies;
+
+		/// Performance diagnostics: The number of modules using this shared PCH header
+		public int NumModulesUsingThisPCH = 0;
 	}
 
 
