Hot reload support for fast build iteration mode in Unreal Build Tool
- Hot reload now works when using bUseExperimentalFastBuildIteration and bUseExperimentalFastDependencyScan
- We use a separate Makefile and built target set for Hot Reload vs. normal compiles
- This allows for extremely quick outdateness checking and incremental builds from the editor

diff --git a/Engine/Source/Programs/UnrealBuildTool/Configuration/UEBuildTarget.cs b/Engine/Source/Programs/UnrealBuildTool/Configuration/UEBuildTarget.cs
--- a/Engine/Source/Programs/UnrealBuildTool/Configuration/UEBuildTarget.cs
+++ b/Engine/Source/Programs/UnrealBuildTool/Configuration/UEBuildTarget.cs
@@ -239,7 +239,6 @@ namespace UnrealBuildTool
 				{
 					switch (Arguments[ArgumentIndex].ToUpperInvariant())
 					{
-						// @todo ubtmake: How does this work with UBTMake?  Do we even need to support it anymore?
 						case "-MODULE":
 							// Specifies a module to recompile.  Can be specified more than once on the command-line to compile multiple specific modules.
 							{
@@ -1135,10 +1134,9 @@ namespace UnrealBuildTool
 
 				// Delete the UBT makefile
 				{
-					// Figure out what to call our action graph based on everything we're building
+					// @todo ubtmake: Does not yet support cleaning Makefiles that were generated for multiple targets (that code path is currently not used though.)
+
 					var TargetDescs = new List<TargetDescriptor>();
-					
-					// @todo ubtmake: Only supports cleaning one target at a time :(
 					TargetDescs.Add( new TargetDescriptor
 						{
 							TargetName = GetTargetName(),
@@ -1146,13 +1144,28 @@ namespace UnrealBuildTool
 							Configuration = this.Configuration
 						} );
 
-					var UBTMakefilePath = UnrealBuildTool.GetUBTMakefilePath( TargetDescs );
-					if (File.Exists(UBTMakefilePath))
-					{
-						Log.TraceVerbose("\tDeleting " + UBTMakefilePath);
-						CleanFile(UBTMakefilePath);
+					// Normal makefile
+					{					
+						var UBTMakefilePath = UnrealBuildTool.GetUBTMakefilePath( TargetDescs );
+						if (File.Exists(UBTMakefilePath))
+						{
+							Log.TraceVerbose("\tDeleting " + UBTMakefilePath);
+							CleanFile(UBTMakefilePath);
+						}
 					}
-				}
+
+					// Hot reload makefile
+					{					
+						UEBuildConfiguration.bHotReloadFromIDE = true;
+						var UBTMakefilePath = UnrealBuildTool.GetUBTMakefilePath( TargetDescs );
+						if (File.Exists(UBTMakefilePath))
+						{
+							Log.TraceVerbose("\tDeleting " + UBTMakefilePath);
+							CleanFile(UBTMakefilePath);
+						}
+						UEBuildConfiguration.bHotReloadFromIDE = false;
+					}
+			}
 
 				// Delete the action history
 				{
@@ -1540,7 +1553,7 @@ namespace UnrealBuildTool
 					Trace.TraceInformation( "UObject discovery time: " + UObjectDiscoveryTime + "s" );
 				}
 
-				// @todo ubtmake: Even in Gather mode, we need to run UHT to make sure the files exist for the static action graph to be setup correctly.  This is because UHT generates .cpp
+				// NOTE: Even in Gather mode, we need to run UHT to make sure the files exist for the static action graph to be setup correctly.  This is because UHT generates .cpp
 				// files that are injected as top level prerequisites.  If UHT only emitted included header files, we wouldn't need to run it during the Gather phase at all.
 				if( UObjectModules.Count > 0 )
 				{
@@ -1849,7 +1862,7 @@ namespace UnrealBuildTool
 					FilteredBinaries.Add(DLLBinary);
 					AnyBinariesAdded = true;
 				}
-				}
+			}
 
 			if (!AnyBinariesAdded)
 			{
diff --git a/Engine/Source/Programs/UnrealBuildTool/System/ActionGraph.cs b/Engine/Source/Programs/UnrealBuildTool/System/ActionGraph.cs
--- a/Engine/Source/Programs/UnrealBuildTool/System/ActionGraph.cs
+++ b/Engine/Source/Programs/UnrealBuildTool/System/ActionGraph.cs
@@ -84,11 +84,11 @@ namespace UnrealBuildTool
 		public bool bProducesImportLibrary = false;
 
 		/** Optional custom event handler for standard output. */
-		public DataReceivedEventHandler OutputEventHandler = null;	// @todo ubtmake urgent: Delegate variables are not saved, but we are comparing against this in ExecutActions() for XGE!
+		public DataReceivedEventHandler OutputEventHandler = null;	// @todo ubtmake urgent: Delegate variables are not saved, but we are comparing against this in ExecuteActions() for XGE!
 
 		/** Callback used to perform a special action instead of a generic command line */
 		public delegate void BlockingActionHandler(Action Action, out int ExitCode, out string Output);
-		public BlockingActionHandler ActionHandler = null;		// @todo ubtmake urgent: Delegate variables are not saved, but we are comparing against this in ExecutActions() for XGE!
+		public BlockingActionHandler ActionHandler = null;		// @todo ubtmake urgent: Delegate variables are not saved, but we are comparing against this in ExecuteActions() for XGE!
 
 
 
diff --git a/Engine/Source/Programs/UnrealBuildTool/System/UnrealBuildTool.cs b/Engine/Source/Programs/UnrealBuildTool/System/UnrealBuildTool.cs
--- a/Engine/Source/Programs/UnrealBuildTool/System/UnrealBuildTool.cs
+++ b/Engine/Source/Programs/UnrealBuildTool/System/UnrealBuildTool.cs
@@ -106,7 +106,7 @@ namespace UnrealBuildTool
                 }
             }
         }
-        static private bool bIsAssemblingBuild_Unsafe = true;	// @todo ubtmake urgent: Doesn't seem to work with hotreload (tries to overwrite same DLL and not use a hotreload DLL name)
+        static private bool bIsAssemblingBuild_Unsafe = true;
 
         /** Used when BuildConfiguration.bUseExperimentalFastDependencyScan is enabled.  If true, it means that our cached includes may not longer be
             valid (or never were), and we need to recover by forcibly scanning included headers for all build prerequisites to make sure that our
@@ -1422,6 +1422,9 @@ namespace UnrealBuildTool
 					}
 				}
 
+				UEBuildConfiguration.bHotReloadFromIDE = UEBuildConfiguration.bAllowHotReloadFromIDE && TargetDescs.Count == 1 && !TargetDescs[0].bIsEditorRecompile && ShouldDoHotReloadFromIDE(TargetDescs[0]);
+				bool bIsHotReload = UEBuildConfiguration.bHotReloadFromIDE || ( TargetDescs.Count == 1 && TargetDescs[0].OnlyModules.Count > 0 );
+				TargetDescriptor HotReloadTargetDesc = bIsHotReload ? TargetDescs[0] : null;
 
                 // If UBT is being used to build something different than it did last time
                 if( !ProjectFileGenerator.bGenerateProjectFiles )
@@ -1450,9 +1453,15 @@ namespace UnrealBuildTool
                         else
                         {
                             string TargetCollectionName = MakeTargetCollectionName( TargetDescs );
-                            string LastBuiltTargetsFilePath = Path.Combine( BuildConfiguration.BaseIntermediatePath, "LastBuiltTargets.txt" );
+
+							string LastBuiltTargetsFileName = bIsHotReload ? "HotReloadLastBuiltTargets.txt" : "LastBuiltTargets.txt";
+                            string LastBuiltTargetsFilePath = Path.Combine( BuildConfiguration.BaseIntermediatePath, LastBuiltTargetsFileName );
                             if( File.Exists( LastBuiltTargetsFilePath ) && Utils.ReadAllText( LastBuiltTargetsFilePath ) == TargetCollectionName )
                             {
+								// @todo ubtmake: Because we're using separate files for hot reload vs. full compiles, it's actually possible that includes will
+								// become out of date without us knowing if the developer ping pongs between hot reloading target A and bulding target B normally.
+								// To fix this we can not use a different file name for last built targets, but the downside is slower performance when
+								// performing the first hot reload after compiling normally (forces full include dependency scan)
                                 bIsBuildingSameTargetsAsLastTime = true;
                             }
 
@@ -1470,7 +1479,7 @@ namespace UnrealBuildTool
                                 // we might not know about all of the C++ include dependencies for already-up-to-date shared build products
                                 // between the targets
                                 bNeedsFullCPPIncludeRescan = true;
-								Log.TraceInformation( "Performing full C++ include scan (building a new target)" );
+								Log.TraceInformation( "Performing full C++ include scan ({0} a new target)", bIsHotReload ? "hot reloading" : "building" );
                             }
                         }
                     }
@@ -1516,7 +1525,11 @@ namespace UnrealBuildTool
                             // a 'gather' and 'assemble' in the same run.  This will take a while longer, but subsequent runs will be fast!
                             UnrealBuildTool.bIsGatheringBuild_Unsafe = true;
 
-                            Log.TraceInformation( "Creating makefile for {0}{1} ({2})", TargetDescs[0].TargetName, TargetDescs.Count > 1 ? ( " (and " + ( TargetDescs.Count - 1 ).ToString() + " more)" ) : "", ReasonNotLoaded );
+                            Log.TraceInformation( "Creating makefile for {0}{1}{2} ({3})", 
+								bIsHotReload ? "hot reloading " : "", 
+								TargetDescs[0].TargetName, 
+								TargetDescs.Count > 1 ? ( " (and " + ( TargetDescs.Count - 1 ).ToString() + " more)" ) : "", 
+								ReasonNotLoaded );
                         }
                     }
 
@@ -1561,9 +1574,7 @@ namespace UnrealBuildTool
                 {
                     var TargetStartTime = DateTime.UtcNow;
 
-                    // Check if we're currently running an editor and this is hot-reload from IDE
-                    UEBuildConfiguration.bHotReloadFromIDE = UEBuildConfiguration.bAllowHotReloadFromIDE && !Target.bEditorRecompile && ShouldDoHotReload(Target);
-                    if (UEBuildConfiguration.bHotReloadFromIDE)
+                    if (bIsHotReload)
                     {
 						// Don't produce new DLLs if there's been no code changes
 						UEBuildConfiguration.bSkipLinkingWhenNothingToCompile = true;
@@ -1685,6 +1696,14 @@ namespace UnrealBuildTool
                         {
                             ActionGraph.AllActions = UBTMakefile.AllActions;
 
+							// Patch action history for hot reload when running in assembler mode.  In assembler mode, the suffix on the output file will be
+							// the same for every invocation on that makefile, but we need a new suffix each time.
+							if( bIsHotReload )
+							{
+								PatchActionHistoryForHotReloadAssembling( HotReloadTargetDesc.OnlyModules );
+							}
+
+
                             foreach( var EnvironmentVariable in UBTMakefile.EnvironmentVariables )
                             {
                                 // @todo ubtmake: There may be some variables we do NOT want to clobber.
@@ -1700,7 +1719,6 @@ namespace UnrealBuildTool
                                     if( TargetUObjectModules.Count > 0 )
                                     {
                                         // Execute the header tool
-										// @todo ubtmake urgent: With a game project, I had a couple of failed builds, then suddenly all generated.cpp's were out of date, and failed to compile (needs repro)
                                         string ModuleInfoFileName = Path.GetFullPath( Path.Combine( Target.ProjectIntermediateDirectory, "UnrealHeaderTool.manifest" ) );
                                         ECompilationResult UHTResult = ECompilationResult.OtherCompilationError;
                                         if (!ExternalExecution.ExecuteHeaderToolIfNecessary(Target, GlobalCompileEnvironment:null, UObjectModules:TargetUObjectModules, ModuleInfoFileName:ModuleInfoFileName, UHTResult:ref UHTResult))
@@ -1862,14 +1880,19 @@ namespace UnrealBuildTool
 		/// </summary>
 		/// <param name="ResetPlatform"></param>
 		/// <param name="ResetConfiguration"></param>
-		/// <param name="Target"></param>
+		/// <param name="TargetDesc"></param>
 		/// <returns></returns>
-		private static bool ShouldDoHotReload(UEBuildTarget Target)
+		private static bool ShouldDoHotReloadFromIDE(TargetDescriptor TargetDesc)
 		{
 			bool bIsRunning = false;
-			if (!ProjectFileGenerator.bGenerateProjectFiles && !UEBuildConfiguration.bGenerateManifest && Target.TargetType == TargetRules.TargetType.Editor)
+			
+			// @todo ubtmake: Kind of cheating here to figure out if an editor target.  At this point we don't have access to the actual target description, and
+			// this code must be able to execute before we create or load module rules DLLs so that hot reload can work with bUseExperimentalFastBuildIteration
+			bool bIsEditorTarget = TargetDesc.TargetName.EndsWith( "Editor", StringComparison.InvariantCultureIgnoreCase ) || TargetDesc.bIsEditorRecompile;	
+
+			if (!ProjectFileGenerator.bGenerateProjectFiles && !UEBuildConfiguration.bGenerateManifest && bIsEditorTarget )
 			{
-				var EditorProcessFilenames = UEBuildTarget.MakeBinaryPaths("Launcher", "UE4Editor", Target.Platform, Target.Configuration, UEBuildBinaryType.Executable, null, null, "UE4Editor");
+				var EditorProcessFilenames = UEBuildTarget.MakeBinaryPaths("Launcher", "UE4Editor", TargetDesc.Platform, TargetDesc.Configuration, UEBuildBinaryType.Executable, null, null, "UE4Editor");
 				if (EditorProcessFilenames.Length != 1)
 				{
 					throw new BuildException("ShouldDoHotReload cannot handle multiple binaries returning from UEBuildTarget.MakeBinaryPaths");
@@ -1879,7 +1902,7 @@ namespace UnrealBuildTool
 				var EditorProcessName = Path.GetFileNameWithoutExtension(EditorProcessFilename);
 				var EditorProcesses = BuildHostPlatform.Current.GetProcessesByName(EditorProcessName);
 				var BinariesPath = Path.GetFullPath(Path.GetDirectoryName(EditorProcessFilename));
-				var PerProjectBinariesPath = Path.Combine(Target.ProjectDirectory, BinariesPath.Substring(BinariesPath.LastIndexOf("Binaries")));
+				var PerProjectBinariesPath = Path.Combine(UnrealBuildTool.GetUProjectPath(), BinariesPath.Substring(BinariesPath.LastIndexOf("Binaries")));
 				bIsRunning = EditorProcesses.FirstOrDefault(EditorProc =>
 					{
 						if(!Path.GetFullPath(EditorProc.Filename).StartsWith(BinariesPath, StringComparison.InvariantCultureIgnoreCase))
@@ -2309,6 +2332,8 @@ namespace UnrealBuildTool
 
             if( TargetDescs.Count == 1 )
             {
+				var TargetDesc = TargetDescs[0];
+
                 // If there's only one target, just save the UBTMakefile in the target's build intermediate directory
                 // under a folder for that target (and platform/config combo.)
                 string PlatformIntermediatePath = BuildConfiguration.PlatformIntermediatePath;
@@ -2317,7 +2342,20 @@ namespace UnrealBuildTool
                     PlatformIntermediatePath = Path.Combine(UnrealBuildTool.GetUProjectPath(), BuildConfiguration.PlatformIntermediateFolder);
                 }
 
-                UBTMakefilePath = Path.Combine(PlatformIntermediatePath, TargetDescs[0].TargetName, TargetDescs[0].Configuration.ToString(), "Makefile.ubt" );
+				bool bIsHotReload = ( UEBuildConfiguration.bHotReloadFromIDE || TargetDesc.OnlyModules.Count > 0 );
+				string UBTMakefileName = bIsHotReload ? "HotReloadMakefile.ubt" : "Makefile.ubt";
+
+				// @todo ubtmake: If this is a compile triggered from the editor it will have passed along the game's target name, not the editor target name.
+				// At this point in Unreal Build Tool, we can't possibly know what the actual editor target name is, but we can take a guess.
+				// Even if we get it wrong, this won't have any side effects aside from not being able to share the exact same cached UBT Makefile
+				// between hot reloads invoked from the editor and hot reloads invoked from within the IDE.
+				string TargetName = TargetDesc.TargetName;
+				if( !TargetName.EndsWith( "Editor", StringComparison.InvariantCultureIgnoreCase ) && TargetDesc.bIsEditorRecompile )
+				{
+					TargetName += "Editor";
+				}
+
+                UBTMakefilePath = Path.Combine(PlatformIntermediatePath, TargetName, TargetDescs[0].Configuration.ToString(), UBTMakefileName );
             }
             else
             {
@@ -2364,8 +2402,18 @@ namespace UnrealBuildTool
                     TargetCollectionName.Append( "_" );
                 }
 
+				// @todo ubtmake: If this is a compile triggered from the editor it will have passed along the game's target name, not the editor target name.
+				// At this point in Unreal Build Tool, we can't possibly know what the actual editor target name is, but we can take a guess.
+				// Even if we get it wrong, this won't have any side effects aside from not being able to share the exact same cached UBT Makefile
+				// between hot reloads invoked from the editor and hot reloads invoked from within the IDE.
+				string TargetName = Target.TargetName;
+				if( !TargetName.EndsWith( "Editor", StringComparison.InvariantCultureIgnoreCase ) && Target.bIsEditorRecompile )
+				{
+					TargetName += "Editor";
+				}
+
 				// @todo ubtmake: Should we also have the platform Architecture in this string?
-                TargetCollectionName.Append( Target.TargetName + "-" + Target.Platform.ToString() + "-" + Target.Configuration.ToString() );
+				TargetCollectionName.Append( TargetName + "-" + Target.Platform.ToString() + "-" + Target.Configuration.ToString() );
             }
 
             return TargetCollectionName.ToString();
@@ -2381,5 +2429,169 @@ namespace UnrealBuildTool
 			BuildConfiguration.RelativeEnginePath = "Engine";
 			SetProjectFile(UProjectFile);
 		}
+
+
+		/// <summary>
+		/// Patch action history for hot reload when running in assembler mode.  In assembler mode, the suffix on the output file will be
+		/// the same for every invocation on that makefile, but we need a new suffix each time.
+		/// </summary>
+		protected static void PatchActionHistoryForHotReloadAssembling( List<OnlyModule> OnlyModules )
+		{
+			// Gather all of the response files for link actions.  We're going to need to patch 'em up after we figure out new
+			// names for all of the output files and import libraries
+			var ResponseFilePaths = new List<string>();
+
+			// Keep a map of the original file names and their new file names, so we can fix up response files after
+			var OriginalFileNameAndNewFileNameList_NoExtensions = new List<Tuple<string, string>>();
+
+			// Finally, we'll keep track of any file items that we had to create counterparts for change file names, so we can fix those up too
+			var AffectedOriginalFileItemAndNewFileItemMap = new Dictionary<FileItem,FileItem>();
+
+			foreach( var Action in ActionGraph.AllActions )
+			{
+				if( Action.ActionType == ActionType.Link )
+				{
+					// Assume that the first produced item (with no extension) is our output file name
+					var OriginalFileNameWithoutExtension = Utils.GetFilenameWithoutAnyExtensions( Action.ProducedItems[0].AbsolutePath );
+
+					// Remove the numbered suffix from the end of the file
+					var OriginalNumberSuffix = OriginalFileNameWithoutExtension.Substring( OriginalFileNameWithoutExtension.LastIndexOf( '-' ) );
+					var OriginalFileNameWithoutNumberSuffix = OriginalFileNameWithoutExtension.Substring( 0, OriginalFileNameWithoutExtension.Length - OriginalNumberSuffix.Length );	// Remove "-####"
+
+					// Figure out which suffix to use
+					string UniqueSuffix = null;
+					{ 
+						int FirstHyphenIndex = OriginalFileNameWithoutNumberSuffix.IndexOf( '-' );
+						if( FirstHyphenIndex == -1 )
+						{
+							throw new BuildException( "Expected produced item file name '{0}' to start with a prefix (such as 'UE4Editor-') when hot reloading", OriginalFileNameWithoutExtension );
+						}
+						string OutputFileNamePrefix = OriginalFileNameWithoutNumberSuffix.Substring( 0, FirstHyphenIndex + 1 );
+
+						// Get the module name from the file name
+						string ModuleName = OriginalFileNameWithoutNumberSuffix.Substring( OutputFileNamePrefix.Length );
+
+						// Add a new random suffix
+						foreach( var OnlyModule in OnlyModules )
+						{
+							if( OnlyModule.OnlyModuleName.Equals( ModuleName, StringComparison.InvariantCultureIgnoreCase ) )
+							{
+								// OK, we found the module!
+								UniqueSuffix = "-" + OnlyModule.OnlyModuleSuffix;
+								break;
+							}
+						}
+						if( UniqueSuffix == null )
+						{
+							UniqueSuffix = "-" + (new Random((int)(DateTime.Now.Ticks % Int32.MaxValue)).Next(10000)).ToString();
+						}
+					}
+					var NewFileNameWithoutExtension = OriginalFileNameWithoutNumberSuffix + UniqueSuffix;
+
+					// Find the response file in the command line.  We'll need to make a copy of it with our new file name.
+					{
+						string ResponseFileExtension = ".response";
+						var ResponseExtensionIndex = Action.CommandArguments.IndexOf( ResponseFileExtension, StringComparison.InvariantCultureIgnoreCase );
+						if( ResponseExtensionIndex != -1 )
+						{ 
+							var ResponseFilePathIndex = Action.CommandArguments.LastIndexOf( "@\"", ResponseExtensionIndex );
+							if( ResponseFilePathIndex == -1 )
+							{
+								throw new BuildException( "Couldn't find response file path in action's command arguments when hot reloading" );
+							}
+
+							var OriginalResponseFilePathWithoutExtension = Action.CommandArguments.Substring( ResponseFilePathIndex + 2, ( ResponseExtensionIndex - ResponseFilePathIndex ) - 2 );
+							var OriginalResponseFilePath = OriginalResponseFilePathWithoutExtension + ResponseFileExtension;
+
+							var NewResponseFilePath = OriginalResponseFilePath.Replace( OriginalFileNameWithoutExtension, NewFileNameWithoutExtension );
+
+							// Copy the old response file to the new path
+							File.Copy( OriginalResponseFilePath, NewResponseFilePath, overwrite: true);
+
+							// Keep track of the new response file name.  We'll have to do some edits afterwards.
+							ResponseFilePaths.Add( NewResponseFilePath );
+						}
+					}
+
+
+					// Go ahead and replace all occurrences of our file name in the command-line (ignoring extensions)
+					Action.CommandArguments = Action.CommandArguments.Replace( OriginalFileNameWithoutExtension, NewFileNameWithoutExtension );
+
+					
+					// Update this action's list of produced items too
+					{
+						for( int ItemIndex = 0; ItemIndex < Action.ProducedItems.Count; ++ItemIndex )
+						{
+							var OriginalProducedItem = Action.ProducedItems[ ItemIndex ];
+
+							var OriginalProducedItemFilePath = OriginalProducedItem.AbsolutePath;
+							var NewProducedItemFilePath = OriginalProducedItemFilePath.Replace( OriginalFileNameWithoutExtension, NewFileNameWithoutExtension );
+						
+							if( OriginalProducedItemFilePath != NewProducedItemFilePath )
+							{
+								// OK, the produced item's file name changed so we'll update it to point to our new file
+								var NewProducedItem = FileItem.GetItemByFullPath( NewProducedItemFilePath );
+								Action.ProducedItems[ ItemIndex ] = NewProducedItem;
+
+								// Copy the other important settings from the original file item
+								NewProducedItem.bNeedsHotReloadNumbersDLLCleanUp = OriginalProducedItem.bNeedsHotReloadNumbersDLLCleanUp;
+								NewProducedItem.ProducingAction = OriginalProducedItem.ProducingAction;
+								NewProducedItem.bIsRemoteFile = OriginalProducedItem.bIsRemoteFile;
+
+								// Keep track of it so we can fix up dependencies in a second pass afterwards
+								AffectedOriginalFileItemAndNewFileItemMap.Add( OriginalProducedItem, NewProducedItem );
+							}
+						}
+					}
+
+					// The status description of the item has the file name, so we'll update it too
+					Action.StatusDescription = Action.StatusDescription.Replace( OriginalFileNameWithoutExtension, NewFileNameWithoutExtension );
+
+		
+					// Keep track of the file names, so we can fix up response files afterwards.
+					OriginalFileNameAndNewFileNameList_NoExtensions.Add( Tuple.Create( OriginalFileNameWithoutExtension, NewFileNameWithoutExtension ) );
+				}
+			}
+
+			
+			// Do another pass and update any actions that depended on the original file names that we changed
+			foreach( var Action in ActionGraph.AllActions )
+			{
+				for( int ItemIndex = 0; ItemIndex < Action.PrerequisiteItems.Count; ++ItemIndex )
+				{
+					var OriginalFileItem = Action.PrerequisiteItems[ ItemIndex ];
+
+					FileItem NewFileItem;
+					if( AffectedOriginalFileItemAndNewFileItemMap.TryGetValue( OriginalFileItem, out NewFileItem ) )
+					{
+						// OK, looks like we need to replace this file item because we've renamed the file
+						Action.PrerequisiteItems[ ItemIndex ] = NewFileItem;
+					}
+				}
+			}
+
+
+			if( OriginalFileNameAndNewFileNameList_NoExtensions.Count > 0 )
+			{	
+				foreach( var ResponseFilePath in ResponseFilePaths )
+				{
+					// Load the file up
+					StringBuilder FileContents = new StringBuilder( Utils.ReadAllText( ResponseFilePath ) );
+
+					// Replace all of the old file names with new ones
+					foreach( var FileNameTuple in OriginalFileNameAndNewFileNameList_NoExtensions )
+					{
+						var OriginalFileNameWithoutExtension = FileNameTuple.Item1;
+						var NewFileNameWithoutExtension = FileNameTuple.Item2;
+
+						FileContents.Replace( OriginalFileNameWithoutExtension, NewFileNameWithoutExtension );
+					}
+
+					// Overwrite the original file
+					var FileContentsString = FileContents.ToString();
+					File.WriteAllText( ResponseFilePath, FileContentsString, System.Text.Encoding.UTF8 );
+				}
+			}
+		}
     }
 }
