UBT: Temporary change (hack) to get the command line length down for UE4LinkerFixups

diff --git a/Engine/Source/Programs/UnrealBuildTool/Configuration/UEBuildModule.cs b/Engine/Source/Programs/UnrealBuildTool/Configuration/UEBuildModule.cs
--- a/Engine/Source/Programs/UnrealBuildTool/Configuration/UEBuildModule.cs
+++ b/Engine/Source/Programs/UnrealBuildTool/Configuration/UEBuildModule.cs
@@ -887,7 +887,7 @@ namespace UnrealBuildTool
 		}
 	};
 
-	/** A module that is compiled from C++ code. */
+	/** A module that is compiled from C++ code. */	
 	public class UEBuildModuleCPP : UEBuildModule
 	{
 		public class AutoGenerateCppInfoClass
@@ -1007,6 +1007,9 @@ namespace UnrealBuildTool
 		/** The processed dependencies for the class */
 		public ProcessedDependenciesClass ProcessedDependencies = null;
 
+		/** @hack to skip adding definitions to compile environment. They will be baked into source code by external code. */
+		public bool bSkipDefinitionsForCompileEnvironment = false;
+
 		/** Categorizes source files into per-extension buckets */
 		private static void CategorizeSourceFiles(IEnumerable<FileItem> InSourceFiles, SourceFilesClass OutSourceFiles)
 		{
@@ -1819,6 +1822,12 @@ namespace UnrealBuildTool
 			// Setup the compile environment for the module.
 			SetupPrivateCompileEnvironment(ref Result.Config.CPPIncludeInfo.IncludePaths, ref Result.Config.CPPIncludeInfo.SystemIncludePaths, ref Result.Config.Definitions, ref Result.Config.AdditionalFrameworks);
 
+			// @hack to skip adding definitions to compile environment, they will be baked into source code files
+			if (bSkipDefinitionsForCompileEnvironment)
+			{
+				Result.Config.Definitions.Clear();
+			}
+
 			return Result;
 		}
 
diff --git a/Engine/Source/Programs/UnrealBuildTool/Configuration/UEBuildTarget.cs b/Engine/Source/Programs/UnrealBuildTool/Configuration/UEBuildTarget.cs
--- a/Engine/Source/Programs/UnrealBuildTool/Configuration/UEBuildTarget.cs
+++ b/Engine/Source/Programs/UnrealBuildTool/Configuration/UEBuildTarget.cs
@@ -1910,50 +1910,40 @@ namespace UnrealBuildTool
 			}
 
 			// We ALWAYS have to write this file as the IMPLEMENT_PRIMARY_GAME_MODULE function depends on the UELinkerFixups function existing.
-			{
-				List<string> LinkerFixupsFileContents = new List<string>();
-
+			{				
 				string LinkerFixupsName = "UELinkerFixups";
 
 				// Include an empty header so UEBuildModule.Compile does not complain about a lack of PCH
 				string HeaderFilename                    = LinkerFixupsName + "Name.h";
 				string LinkerFixupHeaderFilenameWithPath = Path.Combine(GlobalCompileEnvironment.Config.OutputDirectory, HeaderFilename);
 
-				LinkerFixupsFileContents.Add("#include \"" + HeaderFilename + "\"");
-
-				// Add a function that is not referenced by anything that invokes all the empty functions in the different static libraries
-				LinkerFixupsFileContents.Add("void " + LinkerFixupsName + "()");
-				LinkerFixupsFileContents.Add("{");
-
-				// Fill out the body of the function with the empty function calls. This is what causes the static libraries to be considered relevant
-				{
-					var UObjectModules = new List<UEBuildModuleCPP>();
-					foreach (var Binary in AppBinaries)
-					{
-						var DependencyModules = Binary.GetAllDependencyModules(bIncludeDynamicallyLoaded: false, bForceCircular: false);
-						foreach (var Module in DependencyModules.OfType<UEBuildModuleCPP>().Where(CPPModule => CPPModule.AutoGenerateCppInfo != null && !UObjectModules.Any(Module => Module.Name == CPPModule.Name)))
-						{
-							UObjectModules.Add(Module);
-						}
-					}
-					foreach (var Module in UObjectModules)
-					{
-						LinkerFixupsFileContents.Add("    extern void EmptyLinkFunctionForGeneratedCode" + Module.Name + "();");
-						LinkerFixupsFileContents.Add("    EmptyLinkFunctionForGeneratedCode" + Module.Name + "();");
-					}
-				}
-				foreach (var DependencyModuleName in PrivateDependencyModuleNames)
+				// Create the cpp filename
+				string LinkerFixupCPPFilename = Path.Combine(GlobalCompileEnvironment.Config.OutputDirectory, LinkerFixupsName + ".cpp");
+				if (!File.Exists(LinkerFixupCPPFilename))
 				{
-					LinkerFixupsFileContents.Add("    extern void EmptyLinkFunctionForStaticInitialization" + DependencyModuleName + "();");
-					LinkerFixupsFileContents.Add("    EmptyLinkFunctionForStaticInitialization" + DependencyModuleName + "();");
+					// Create a dummy file in case it doesn't exist yet so that the module does not complain it's not there
+					ResponseFile.Create(LinkerFixupHeaderFilenameWithPath, new List<string>());
+					ResponseFile.Create(LinkerFixupCPPFilename, new List<string>(new string[] { String.Format("#include \"{0}\"", LinkerFixupHeaderFilenameWithPath) }));
 				}
 
-				// End the function body that was started above
-				LinkerFixupsFileContents.Add("}");
+				// Create the source file list (just the one cpp file)
+				List<FileItem> SourceFiles = new List<FileItem>();
+				SourceFiles.Add(FileItem.GetItemByPath(LinkerFixupCPPFilename));
 
-				// Create the cpp file
-				string LinkerFixupCPPFilename = Path.Combine(GlobalCompileEnvironment.Config.OutputDirectory, LinkerFixupsName + ".cpp");
+				// Create the CPP module
+				var FakeModuleDirectory = Path.GetDirectoryName( LinkerFixupCPPFilename );
+				var NewModule = CreateArtificialModule(LinkerFixupsName, FakeModuleDirectory, SourceFiles, PrivateDependencyModuleNames);
 
+				// Now bind this new module to the executable binary so it will link the plugin libs correctly
+				NewModule.bSkipDefinitionsForCompileEnvironment = true;
+				BindArtificialModuleToBinary(NewModule, ExecutableBinary);
+
+				// Create the cpp file
+				List<string> LinkerFixupsFileContents = new List<string>();
+				NewModule.bSkipDefinitionsForCompileEnvironment = false;
+				GenerateLinkerFixupsContents(LinkerFixupsFileContents, NewModule.CreateModuleCompileEnvironment(GlobalCompileEnvironment), HeaderFilename, LinkerFixupsName, PrivateDependencyModuleNames);
+				NewModule.bSkipDefinitionsForCompileEnvironment = true;
+				
 				// Determine if the file changed. Write it if it either doesn't exist or the contents are different.
 				bool bShouldWriteFile = true;
 				if (File.Exists(LinkerFixupCPPFilename))
@@ -1970,18 +1960,62 @@ namespace UnrealBuildTool
 					ResponseFile.Create(LinkerFixupHeaderFilenameWithPath, new List<string>());
 					ResponseFile.Create(LinkerFixupCPPFilename, LinkerFixupsFileContents);
 				}
+			}
+		}
 
-				// Create the source file list (just the one cpp file)
-				List<FileItem> SourceFiles = new List<FileItem>();
-				SourceFiles.Add(FileItem.GetItemByPath(LinkerFixupCPPFilename));
+		private void GenerateLinkerFixupsContents(List<string> LinkerFixupsFileContents, CPPEnvironment CompileEnvironment, string HeaderFilename, string LinkerFixupsName, List<string> PrivateDependencyModuleNames)
+		{
+			LinkerFixupsFileContents.Add("#include \"" + HeaderFilename + "\"");
 
-				// Create the CPP module
-				var FakeModuleDirectory = Path.GetDirectoryName( LinkerFixupCPPFilename );
-				var NewModule = CreateArtificialModule(LinkerFixupsName, FakeModuleDirectory, SourceFiles, PrivateDependencyModuleNames);
+			// To reduce the size of the command line for the compiler, we're going to put all definitions inside of the cpp file.
+			foreach (var Definition in CompileEnvironment.Config.Definitions)
+			{
+				string MacroName;
+				string MacroValue = String.Empty;
+				int EqualsIndex = Definition.IndexOf('=');
+				if (EqualsIndex >= 0)
+				{
+					MacroName = Definition.Substring(0, EqualsIndex);
+					MacroValue = Definition.Substring(EqualsIndex + 1);
+				}
+				else
+				{
+					MacroName = Definition;
+				}
+				LinkerFixupsFileContents.Add("#ifndef " + MacroName);
+				LinkerFixupsFileContents.Add(String.Format("\t#define {0} {1}", MacroName, MacroValue));
+				LinkerFixupsFileContents.Add("#endif");
+			}
 
-				// Now bind this new module to the executable binary so it will link the plugin libs correctly
-				BindArtificialModuleToBinary(NewModule, ExecutableBinary);
+			// Add a function that is not referenced by anything that invokes all the empty functions in the different static libraries
+			LinkerFixupsFileContents.Add("void " + LinkerFixupsName + "()");
+			LinkerFixupsFileContents.Add("{");
+
+			// Fill out the body of the function with the empty function calls. This is what causes the static libraries to be considered relevant
+			{
+				var UObjectModules = new List<UEBuildModuleCPP>();
+				foreach (var Binary in AppBinaries)
+				{
+					var DependencyModules = Binary.GetAllDependencyModules(bIncludeDynamicallyLoaded: false, bForceCircular: false);
+					foreach (var Module in DependencyModules.OfType<UEBuildModuleCPP>().Where(CPPModule => CPPModule.AutoGenerateCppInfo != null && !UObjectModules.Any(Module => Module.Name == CPPModule.Name)))
+					{
+						UObjectModules.Add(Module);
+					}
+				}
+				foreach (var Module in UObjectModules)
+				{
+					LinkerFixupsFileContents.Add("    extern void EmptyLinkFunctionForGeneratedCode" + Module.Name + "();");
+					LinkerFixupsFileContents.Add("    EmptyLinkFunctionForGeneratedCode" + Module.Name + "();");
+				}
 			}
+			foreach (var DependencyModuleName in PrivateDependencyModuleNames)
+			{
+				LinkerFixupsFileContents.Add("    extern void EmptyLinkFunctionForStaticInitialization" + DependencyModuleName + "();");
+				LinkerFixupsFileContents.Add("    EmptyLinkFunctionForStaticInitialization" + DependencyModuleName + "();");
+			}
+
+			// End the function body that was started above
+			LinkerFixupsFileContents.Add("}");
 		}
 
 		/// <summary>
