Hot reload support for fast build iteration mode in Unreal Build Tool (fixed assertion on clean)

diff --git a/Engine/Source/Programs/UnrealBuildTool/System/UnrealBuildTool.cs b/Engine/Source/Programs/UnrealBuildTool/System/UnrealBuildTool.cs
--- a/Engine/Source/Programs/UnrealBuildTool/System/UnrealBuildTool.cs
+++ b/Engine/Source/Programs/UnrealBuildTool/System/UnrealBuildTool.cs
@@ -2342,7 +2342,7 @@ namespace UnrealBuildTool
                     PlatformIntermediatePath = Path.Combine(UnrealBuildTool.GetUProjectPath(), BuildConfiguration.PlatformIntermediateFolder);
                 }
 
-				bool bIsHotReload = ( UEBuildConfiguration.bHotReloadFromIDE || TargetDesc.OnlyModules.Count > 0 );
+				bool bIsHotReload = ( UEBuildConfiguration.bHotReloadFromIDE || ( TargetDesc.OnlyModules != null && TargetDesc.OnlyModules.Count > 0 ) );
 				string UBTMakefileName = bIsHotReload ? "HotReloadMakefile.ubt" : "Makefile.ubt";
 
 				// @todo ubtmake: If this is a compile triggered from the editor it will have passed along the game's target name, not the editor target name.
@@ -2472,13 +2472,16 @@ namespace UnrealBuildTool
 						string ModuleName = OriginalFileNameWithoutNumberSuffix.Substring( OutputFileNamePrefix.Length );
 
 						// Add a new random suffix
-						foreach( var OnlyModule in OnlyModules )
-						{
-							if( OnlyModule.OnlyModuleName.Equals( ModuleName, StringComparison.InvariantCultureIgnoreCase ) )
+						if( OnlyModules != null )
+						{ 
+							foreach( var OnlyModule in OnlyModules )
 							{
-								// OK, we found the module!
-								UniqueSuffix = "-" + OnlyModule.OnlyModuleSuffix;
-								break;
+								if( OnlyModule.OnlyModuleName.Equals( ModuleName, StringComparison.InvariantCultureIgnoreCase ) )
+								{
+									// OK, we found the module!
+									UniqueSuffix = "-" + OnlyModule.OnlyModuleSuffix;
+									break;
+								}
 							}
 						}
 						if( UniqueSuffix == null )
