Disabled remote linking in XGE until we verify it isn't causing problems with the build system

diff --git a/Engine/Source/Programs/UnrealBuildTool/Windows/VCToolChain.cs b/Engine/Source/Programs/UnrealBuildTool/Windows/VCToolChain.cs
--- a/Engine/Source/Programs/UnrealBuildTool/Windows/VCToolChain.cs
+++ b/Engine/Source/Programs/UnrealBuildTool/Windows/VCToolChain.cs
@@ -1052,9 +1052,6 @@ namespace UnrealBuildTool
 				CompileAction.CommandPath        = EnvVars.ResourceCompilerPath;
 				CompileAction.StatusDescription  = Path.GetFileName(RCFile.AbsolutePath);
 
-				// Resource tool can run remotely if possible
-				CompileAction.bCanExecuteRemotely = true;
-
 				if( WindowsPlatform.bCompileWithClang )
 				{ 
 					CompileAction.OutputEventHandler = new DataReceivedEventHandler( ClangCompilerOutputFormatter );
@@ -1323,8 +1320,8 @@ namespace UnrealBuildTool
 			// ignored as a prerequisite for other actions
 			LinkAction.bProducesImportLibrary = bBuildImportLibraryOnly || LinkEnvironment.Config.bIsBuildingDLL;
 
-			// Allow remote linking.  Especially in modular builds with many small DLL files, this is almost always very efficient
-			LinkAction.bCanExecuteRemotely = true;
+			// Only execute linking on the local PC.
+			LinkAction.bCanExecuteRemotely = false;
 
 			Log.TraceVerbose( "     Linking: " + LinkAction.StatusDescription );
 			Log.TraceVerbose( "     Command: " + LinkAction.CommandArguments );
