Fixed a few problems when compiling using latest version of Clang on Windows
- Superficial warnings that triggered as errors are now suppressed
- Fixed XGE thinking we were running the Visual C++ compiler

diff --git a/Engine/Source/Programs/UnrealBuildTool/Windows/VCToolChain.cs b/Engine/Source/Programs/UnrealBuildTool/Windows/VCToolChain.cs
--- a/Engine/Source/Programs/UnrealBuildTool/Windows/VCToolChain.cs
+++ b/Engine/Source/Programs/UnrealBuildTool/Windows/VCToolChain.cs
@@ -411,7 +411,9 @@ namespace UnrealBuildTool
 				Arguments.Append(" -Wno-microsoft");								
 
 				// @todo clang: Kind of a shame to turn these off.  We'd like to catch unused variables, but it is tricky with how our assertion macros work.
+				Arguments.Append(" -Wno-inconsistent-missing-override");
 				Arguments.Append(" -Wno-unused-variable");
+				Arguments.Append(" -Wno-unused-local-typedefs");
 				Arguments.Append(" -Wno-unused-function");
 				Arguments.Append(" -Wno-unused-private-field");
 				Arguments.Append(" -Wno-unused-value");
@@ -983,7 +985,7 @@ namespace UnrealBuildTool
 					CompileAction.CommandPath = EnvVars.CompilerPath;
 				}
 
-				if( !WindowsPlatform.bCompileWithClang || !BuildConfiguration.bRunUnrealCodeAnalyzer )
+				if( !WindowsPlatform.bCompileWithClang && !BuildConfiguration.bRunUnrealCodeAnalyzer )
 				{
 					CompileAction.bIsVCCompiler = true;
 				}
