Fix SNDBS build configuration throwing an exception if the SCE_ROOT path could not be found.  This only occurs with XGE disabled.

diff --git a/Engine/Source/Programs/UnrealBuildTool/System/SNDBS.cs b/Engine/Source/Programs/UnrealBuildTool/System/SNDBS.cs
--- a/Engine/Source/Programs/UnrealBuildTool/System/SNDBS.cs
+++ b/Engine/Source/Programs/UnrealBuildTool/System/SNDBS.cs
@@ -289,10 +289,16 @@ namespace UnrealBuildTool
             if (Actions.Count > 0)
             {
                 string SCERoot = Environment.GetEnvironmentVariable("SCE_ROOT_DIR");
-                string SNDBSExecutable = Path.Combine(SCERoot, "Common/SN-DBS/bin/dbsbuild.exe");
 
-                // Check that SN-DBS is available
-                bool bSNDBSExists = File.Exists(SNDBSExecutable);
+		bool bSNDBSExists = false;
+		if( SCERoot != null )
+		{
+			string SNDBSExecutable = Path.Combine(SCERoot, "Common/SN-DBS/bin/dbsbuild.exe");
+
+			// Check that SN-DBS is available
+			bSNDBSExists = File.Exists(SNDBSExecutable);
+		}
+
                 if( bSNDBSExists == false )
                 {
                     return ExecutionResult.Unavailable;
