Shortening UELinkerFixups compiler command line length even more

diff --git a/Engine/Source/Programs/UnrealBuildTool/Configuration/UEBuildModule.cs b/Engine/Source/Programs/UnrealBuildTool/Configuration/UEBuildModule.cs
--- a/Engine/Source/Programs/UnrealBuildTool/Configuration/UEBuildModule.cs
+++ b/Engine/Source/Programs/UnrealBuildTool/Configuration/UEBuildModule.cs
@@ -1826,6 +1826,7 @@ namespace UnrealBuildTool
 			if (bSkipDefinitionsForCompileEnvironment)
 			{
 				Result.Config.Definitions.Clear();
+				Result.Config.CPPIncludeInfo.IncludePaths.Clear();
 			}
 
 			return Result;
