Fixing UE4Game target's UELinkerFixups not adding the PCH include path

diff --git a/Engine/Source/Programs/UnrealBuildTool/Configuration/UEBuildModule.cs b/Engine/Source/Programs/UnrealBuildTool/Configuration/UEBuildModule.cs
--- a/Engine/Source/Programs/UnrealBuildTool/Configuration/UEBuildModule.cs
+++ b/Engine/Source/Programs/UnrealBuildTool/Configuration/UEBuildModule.cs
@@ -1826,7 +1826,7 @@ namespace UnrealBuildTool
 			if (bSkipDefinitionsForCompileEnvironment)
 			{
 				Result.Config.Definitions.Clear();
-				Result.Config.CPPIncludeInfo.IncludePaths.Clear();
+				Result.Config.CPPIncludeInfo.IncludePaths = new HashSet<string>(BaseCompileEnvironment.Config.CPPIncludeInfo.IncludePaths);
 			}
 
 			return Result;
