Improved support for compiling and linking using Clang on Windows platform
- Disclaimer: This is all experimental forward-looking stuff, we don't recommend using this yet!
- Added native support for Clang command-line on Windows ('clang-cl' is not required, but still supported)
- Clang debug symbols are now Visual Studio compatible (use latest Clang 3.6.0 prelease build.)
- Clang linker (LLD.exe) now works with most UE4 binaries, but debug symbols won't work
- New 'WindowsPlatform.bUseVCCompilerArgs' that controls whether to the VC-style wrapper with Clang
- Implemented precompiled header support (requires native Clang binary.)  There are still a few bugs in Clang with struct packing which prevent us from using PCHs on Windows though.
- Editor still can't compile under Clang because of some DX header file issues, possibly system include related.
- Unreal Build Tool: Reduced default performance spew when using the local executor (non-XGE builds.)

diff --git a/Engine/Source/Programs/UnrealBuildTool/System/Distcc.cs b/Engine/Source/Programs/UnrealBuildTool/System/Distcc.cs
--- a/Engine/Source/Programs/UnrealBuildTool/System/Distcc.cs
+++ b/Engine/Source/Programs/UnrealBuildTool/System/Distcc.cs
@@ -231,7 +231,9 @@ namespace UnrealBuildTool
 					TotalThreadSeconds += ThreadSeconds;
 				}
 
-				Log.TraceInformation("-------- End Detailed Actions Stats -----------------------------------------------------------");
+				Log.WriteLineIf(BuildConfiguration.bLogDetailedActionStats || BuildConfiguration.bPrintDebugInfo,
+					TraceEventType.Information, 
+					"-------- End Detailed Actions Stats -----------------------------------------------------------");
 
 				// Log total CPU seconds and numbers of processors involved in tasks.
 				Log.WriteLineIf(BuildConfiguration.bLogDetailedActionStats || BuildConfiguration.bPrintDebugInfo, 
diff --git a/Engine/Source/Programs/UnrealBuildTool/System/LocalExecutor.cs b/Engine/Source/Programs/UnrealBuildTool/System/LocalExecutor.cs
--- a/Engine/Source/Programs/UnrealBuildTool/System/LocalExecutor.cs
+++ b/Engine/Source/Programs/UnrealBuildTool/System/LocalExecutor.cs
@@ -384,7 +384,7 @@ namespace UnrealBuildTool
 				MaxActionsToExecuteInParallel = Math.Min(MaxActionsToExecuteInParallel, MaxActionsAffordedByMemory);
             }
 
-			MaxActionsToExecuteInParallel = Math.Min(MaxActionsToExecuteInParallel, BuildConfiguration.MaxProcessorCount);
+			MaxActionsToExecuteInParallel = Math.Max( 1, Math.Min(MaxActionsToExecuteInParallel, BuildConfiguration.MaxProcessorCount) );
 
             Log.TraceInformation("Performing {0} actions ({1} in parallel)", Actions.Count, MaxActionsToExecuteInParallel);
 
diff --git a/Engine/Source/Programs/UnrealBuildTool/System/UnrealBuildTool.cs b/Engine/Source/Programs/UnrealBuildTool/System/UnrealBuildTool.cs
--- a/Engine/Source/Programs/UnrealBuildTool/System/UnrealBuildTool.cs
+++ b/Engine/Source/Programs/UnrealBuildTool/System/UnrealBuildTool.cs
@@ -1830,7 +1830,9 @@ namespace UnrealBuildTool
             double BuildDuration = (DateTime.UtcNow - StartTime - MutexWaitTime).TotalSeconds;
             if (ExecutorName == "Local" || ExecutorName == "Distcc" || ExecutorName == "SNDBS")
             {
-                Log.TraceInformation("Cumulative action seconds ({0} processors): {1:0.00} building projects, {2:0.00} compiling, {3:0.00} creating app bundles, {4:0.00} generating debug info, {5:0.00} linking, {6:0.00} other",
+                Log.WriteLineIf(BuildConfiguration.bLogDetailedActionStats || BuildConfiguration.bPrintDebugInfo,
+					TraceEventType.Information, 
+					"Cumulative action seconds ({0} processors): {1:0.00} building projects, {2:0.00} compiling, {3:0.00} creating app bundles, {4:0.00} generating debug info, {5:0.00} linking, {6:0.00} other",
                     System.Environment.ProcessorCount,
                     TotalBuildProjectTime,
                     TotalCompileTime,
@@ -1850,7 +1852,7 @@ namespace UnrealBuildTool
                     "TotalOtherActionsThreadTimeSec", TotalOtherActionsTime.ToString("0.00")
                     );
 
-                Log.TraceInformation("UBT execution time: {0:0.00} seconds", BuildDuration);
+                Log.TraceInformation("Total build time: {0:0.00} seconds", BuildDuration);
 
                 // reset statistics
                 TotalBuildProjectTime = 0;
diff --git a/Engine/Source/Programs/UnrealBuildTool/Windows/UEBuildWindows.cs b/Engine/Source/Programs/UnrealBuildTool/Windows/UEBuildWindows.cs
--- a/Engine/Source/Programs/UnrealBuildTool/Windows/UEBuildWindows.cs
+++ b/Engine/Source/Programs/UnrealBuildTool/Windows/UEBuildWindows.cs
@@ -59,11 +59,14 @@ namespace UnrealBuildTool
             }
         }
 
-        /// True if we are using "clang-cl" to compile instead of MSVC on Windows platform
+        /// True if we should use Clang/LLVM instead of MSVC to compile code on Windows platform
         public static readonly bool bCompileWithClang = false;
 
+        /// When using Clang, enabling enables the MSVC-like "clang-cl" wrapper, otherwise we pass arguments to Clang directly
+		public static readonly bool bUseVCCompilerArgs = !bCompileWithClang || true;
+
 		/// True if we should use the Clang linker (LLD) when bCompileWithClang is enabled, otherwise we use the MSVC linker
-		public static readonly bool bAllowClangLinker = false;
+		public static readonly bool bAllowClangLinker = bCompileWithClang && true;
 
         /// True if we're targeting Windows XP as a minimum spec.  In Visual Studio 2012 and higher, this may change how
         /// we compile and link the application (http://blogs.msdn.com/b/vcblog/archive/2012/10/08/10357555.aspx)
@@ -324,7 +327,8 @@ namespace UnrealBuildTool
 
         public override void ValidateBuildConfiguration(CPPTargetConfiguration Configuration, CPPTargetPlatform Platform, bool bCreateDebugInfo)
         {
-            // @todo clang: PCH files aren't quite working yet with "clang-cl" (no /Yc support, and "-x c++-header" cannot be specified)
+            // @todo clang: PCH files aren't supported by "clang-cl" yet (no /Yc support, and "-x c++-header" cannot be specified)
+			// @todo clang: PCH files with regular Clang on Windows have bugs with pack alignment and segfaults occasionally
             if( WindowsPlatform.bCompileWithClang )
             {
                 BuildConfiguration.bUsePCHFiles = false;
diff --git a/Engine/Source/Programs/UnrealBuildTool/Windows/VCEnvironment.cs b/Engine/Source/Programs/UnrealBuildTool/Windows/VCEnvironment.cs
--- a/Engine/Source/Programs/UnrealBuildTool/Windows/VCEnvironment.cs
+++ b/Engine/Source/Programs/UnrealBuildTool/Windows/VCEnvironment.cs
@@ -149,7 +149,17 @@ namespace UnrealBuildTool
 			// If we were asked to use Clang, then we'll redirect the path to the compiler to the LLVM installation directory
 			if (WindowsPlatform.bCompileWithClang)
 			{
-				var Result = Path.Combine( Environment.GetFolderPath( Environment.SpecialFolder.ProgramFilesX86 ), "LLVM", "msbuild-bin", "cl.exe" );
+				string Result;
+				if( WindowsPlatform.bUseVCCompilerArgs )
+				{ 
+					// Use regular Clang compiler on Windows
+					Result = Path.Combine( Environment.GetFolderPath( Environment.SpecialFolder.ProgramFilesX86 ), "LLVM", "msbuild-bin", "cl.exe" );
+				}
+				else
+				{
+					// Use 'clang-cl', a wrapper around Clang that supports Visual C++ compiler command-line arguments
+					Result = Path.Combine( Environment.GetFolderPath( Environment.SpecialFolder.ProgramFilesX86 ), "LLVM", "bin", "clang.exe" );
+				}
 				if (!File.Exists(Result))
 				{
 					throw new BuildException( "Clang was selected as the Windows compiler, but LLVM/Clang does not appear to be installed.  Could not find: " + Result );
diff --git a/Engine/Source/Programs/UnrealBuildTool/Windows/VCToolChain.cs b/Engine/Source/Programs/UnrealBuildTool/Windows/VCToolChain.cs
--- a/Engine/Source/Programs/UnrealBuildTool/Windows/VCToolChain.cs
+++ b/Engine/Source/Programs/UnrealBuildTool/Windows/VCToolChain.cs
@@ -21,36 +21,103 @@ namespace UnrealBuildTool
 			UEToolChain.RegisterPlatformToolChain(CPPTargetPlatform.Win32, this);
 		}
 
-		static void AppendCLArguments_Global(CPPEnvironment CompileEnvironment, VCEnvironment EnvVars, StringBuilder Arguments)
+		
+		static void AddDefinition( ref StringBuilder String, string Variable, string Value = null )
 		{
-			// @todo fastubt: Use should be using StringBuilder instead of concatenation in all of these toolchain files.  On Mono, string builder is dramatically faster.
+			// If the value has a space in it and isn't wrapped in quotes, do that now
+			if( Value != null && !Value.StartsWith( "\"" ) && ( Value.Contains( " " ) || Value.Contains( "$" ) ) )
+			{
+				Value = "\"" + Value + "\"";
+			}
 
+			if( WindowsPlatform.bUseVCCompilerArgs )
+			{
+				if( Value != null )
+				{ 
+					String.Append( " /D" + Variable + "=" + Value );
+				}
+				else
+				{
+					String.Append( " /D" + Variable );
+				}
+			}
+			else
+			{
+				if( Value != null )
+				{ 
+					String.Append( " -D " + Variable + "=" + Value );
+				}
+				else
+				{
+					String.Append( " -D " + Variable );
+				}
+			}
+		}
+
+
+		static void AddIncludePath( ref StringBuilder String, string IncludePath )
+		{
+			// If the value has a space in it and isn't wrapped in quotes, do that now
+			if( !IncludePath.StartsWith( "\"" ) && ( IncludePath.Contains( " " ) || IncludePath.Contains( "$" ) ) )
+			{
+				IncludePath = "\"" + IncludePath + "\"";
+			}
+
+			if( WindowsPlatform.bUseVCCompilerArgs )
+			{
+				String.Append( " /I " + IncludePath );
+			}
+			else
+			{
+				String.Append( " -I" + IncludePath );
+			}
+		}
 
-			// @todo fastubt: We can use '/showIncludes' to accelerate outdatedness checking, as the compiler will discover all indirect includes itself.  But, the spew is pretty noisy!
-			//		-> If no files in source file chain have changed, even if the build product is outdated, we can skip using /showIncludes (relies on cache surviving)
-			// Arguments.Append( " /showIncludes" );
+	
+		static void AppendCLArguments_Global(CPPEnvironment CompileEnvironment, VCEnvironment EnvVars, StringBuilder Arguments)
+		{
+			// NOTE: Uncommenting this line will print includes as they are encountered by the preprocessor.  This can help with diagnosing include order problems.
+			if( WindowsPlatform.bCompileWithClang && !WindowsPlatform.bUseVCCompilerArgs )
+			{ 
+				// Arguments.Append( " -H" );
+			}
+			else
+			{
+				// Arguments.Append( " /showIncludes" );
+			}
 
 			if( WindowsPlatform.bCompileWithClang )
 			{ 
 				// Arguments.Append( " -###" );	// @todo clang: Print Clang command-lines (instead of outputting compile results!)
 
+				if( !WindowsPlatform.bUseVCCompilerArgs )
+				{
+					Arguments.Append( " -std=c++11" );
+					Arguments.Append( " -fdiagnostics-format=msvc" );
+					Arguments.Append( " -Xclang -relaxed-aliasing -Xclang --dependent-lib=msvcrt -Xclang --dependent-lib=oldnames -gline-tables-only -ffunction-sections" );
+				}
+
 				// @todo clang: We're impersonating the Visual C++ compiler by setting MSC_VER and _MSC_FULL_VER to values that MSVC would set
 				string VersionString;
+				string FullVersionString;
 				switch( WindowsPlatform.Compiler )
 				{
 					case WindowsCompiler.VisualStudio2012:
-						VersionString = "1700";
+						VersionString = "17.0";
+						FullVersionString = "1700";
 						break;
 		
 					case WindowsCompiler.VisualStudio2013:
-						VersionString = "1800";
+						VersionString = "18.0";
+						FullVersionString = "1800";
 						break;
 
 					default:
 						throw new BuildException( "Unexpected value for WindowsPlatform.Compiler: " + WindowsPlatform.Compiler.ToString() );
 				}
-				Arguments.Append(" -fmsc-version=" + VersionString);
-				Arguments.Append(" /D_MSC_FULL_VER=" + VersionString + "00000");
+
+				Arguments.Append(" -fms-compatibility-version=" + VersionString);
+				AddDefinition( ref Arguments, "_MSC_FULL_VER", FullVersionString + "00000" );
 			}
 
 			// @todo clang: Clang on Windows doesn't respect "#pragma warning (error: ####)", and we're not passing "/WX", so warnings are not
@@ -72,13 +139,16 @@ namespace UnrealBuildTool
 				//Arguments.Append(" /analyze:only");
 			}
 
-			// Prevents the compiler from displaying its logo for each invocation.
-			Arguments.Append(" /nologo");
+			if( WindowsPlatform.bUseVCCompilerArgs )
+			{ 
+				// Prevents the compiler from displaying its logo for each invocation.
+				Arguments.Append(" /nologo");
 
-			// Enable intrinsic functions.
-			Arguments.Append(" /Oi");
+				// Enable intrinsic functions.
+				Arguments.Append(" /Oi");	// @todo clang: No Clang equivalent to this?
+			}
 
-			
+						
 			if( WindowsPlatform.bCompileWithClang )
 			{
 				// Tell the Clang compiler whether we want to generate 32-bit code or 64-bit code
@@ -92,41 +162,54 @@ namespace UnrealBuildTool
 				}
 			}
 
-			// Separate functions for linker.
-			Arguments.Append(" /Gy");
-
 			// Compile into an .obj file, and skip linking.
-			Arguments.Append(" /c");
+			if( WindowsPlatform.bUseVCCompilerArgs )
+			{
+				Arguments.Append(" /c");
+			}
+			else
+			{
+				Arguments.Append(" -c");
+			}
 
-			// Allow 800% of the default memory allocation limit.
-			Arguments.Append(" /Zm800");
+			if( WindowsPlatform.bUseVCCompilerArgs )
+			{
+				// Separate functions for linker.
+				Arguments.Append(" /Gy");
 
-			// Disable "The file contains a character that cannot be represented in the current code page" warning for non-US windows.
-			Arguments.Append(" /wd4819");
+				// Allow 800% of the default memory allocation limit.
+				Arguments.Append(" /Zm800");
+
+				// Disable "The file contains a character that cannot be represented in the current code page" warning for non-US windows.
+				Arguments.Append(" /wd4819");
+			}
 
 			if( BuildConfiguration.bUseSharedPCHs )
 			{
-				// @todo SharedPCH: Disable warning about PCH defines not matching .cpp defines.  We "cheat" these defines a little
-				// bit to make shared PCHs work.  But it's totally safe.  Trust us.
-				Arguments.Append(" /wd4651");
+				if( WindowsPlatform.bUseVCCompilerArgs )
+				{ 
+					// @todo SharedPCH: Disable warning about PCH defines not matching .cpp defines.  We "cheat" these defines a little
+					// bit to make shared PCHs work.  But it's totally safe.  Trust us.
+					Arguments.Append(" /wd4651");
 
-				// @todo SharedPCH: Disable warning about redefining *API macros.  The PCH header is compiled with various DLLIMPORTs, but
-				// when a module that uses that PCH header *IS* one of those imports, that module is compiled with EXPORTS, so the macro
-				// is redefined on the command-line.  We need to clobber those defines to make shared PCHs work properly!
-				Arguments.Append(" /wd4005");
+					// @todo SharedPCH: Disable warning about redefining *API macros.  The PCH header is compiled with various DLLIMPORTs, but
+					// when a module that uses that PCH header *IS* one of those imports, that module is compiled with EXPORTS, so the macro
+					// is redefined on the command-line.  We need to clobber those defines to make shared PCHs work properly!
+					Arguments.Append(" /wd4005");
+				}
 			}
 
 			// If compiling as a DLL, set the relevant defines
 			if (CompileEnvironment.Config.bIsBuildingDLL)
 			{
-				Arguments.Append(" /D _WINDLL");
+				AddDefinition( ref Arguments, "_WINDLL" );
 			}
 
 			// When targeting Windows XP with Visual Studio 2012+, we need to tell the compiler to use the older Windows SDK that works
 			// with Windows XP (http://blogs.msdn.com/b/vcblog/archive/2012/10/08/10357555.aspx)
 			if( WindowsPlatform.SupportWindowsXP )
 			{
-				Arguments.Append(" /D_USING_V110_SDK71_");
+				AddDefinition( ref Arguments, "_USING_V110_SDK71_");
 			}
 
 
@@ -144,25 +227,28 @@ namespace UnrealBuildTool
 			//
 			if (CompileEnvironment.Config.Target.Configuration == CPPTargetConfiguration.Debug)
 			{
-				// Disable compiler optimization.
-				Arguments.Append(" /Od");
+				if( WindowsPlatform.bUseVCCompilerArgs )
+				{ 
+					// Disable compiler optimization.
+					Arguments.Append(" /Od");
 
-				// Favor code size (especially useful for embedded platforms).
-				Arguments.Append(" /Os");
+					// Favor code size (especially useful for embedded platforms).
+					Arguments.Append(" /Os");
 
-				// Allow inline method expansion unless E&C support is requested
-				if( !BuildConfiguration.bSupportEditAndContinue )
-				{
-					Arguments.Append(" /Ob2");
-				}
+					// Allow inline method expansion unless E&C support is requested
+					if( !BuildConfiguration.bSupportEditAndContinue )
+					{
+						Arguments.Append(" /Ob2");
+					}
 
-				if ((CompileEnvironment.Config.Target.Platform == CPPTargetPlatform.Win32) || 
-					(CompileEnvironment.Config.Target.Platform == CPPTargetPlatform.Win64))
-				{
-					// Runtime stack checks are not allowed when compiling for CLR
-					if (CompileEnvironment.Config.CLRMode == CPPCLRMode.CLRDisabled)
+					if ((CompileEnvironment.Config.Target.Platform == CPPTargetPlatform.Win32) || 
+						(CompileEnvironment.Config.Target.Platform == CPPTargetPlatform.Win64))
 					{
-						Arguments.Append(" /RTCs");
+						// Runtime stack checks are not allowed when compiling for CLR
+						if (CompileEnvironment.Config.CLRMode == CPPCLRMode.CLRDisabled)
+						{
+							Arguments.Append(" /RTCs");
+						}
 					}
 				}
 			}
@@ -171,35 +257,46 @@ namespace UnrealBuildTool
 			//
 			else
 			{
-				// Maximum optimizations if desired.
-				if( CompileEnvironment.Config.OptimizeCode >= ModuleRules.CodeOptimization.InNonDebugBuilds )
-				{
-					Arguments.Append(" /Ox");
-				}
+				if( WindowsPlatform.bUseVCCompilerArgs )
+				{ 
+					// Maximum optimizations if desired.
+					if( CompileEnvironment.Config.OptimizeCode >= ModuleRules.CodeOptimization.InNonDebugBuilds )
+					{
+						Arguments.Append(" /Ox");
+					}
 				
-				// Favor code speed.
-				Arguments.Append(" /Ot");
+					// Favor code speed.
+					Arguments.Append(" /Ot");
 
-				// Only omit frame pointers on the PC (which is implied by /Ox) if wanted.
-				if ( BuildConfiguration.bOmitFramePointers == false
-				&& ((CompileEnvironment.Config.Target.Platform == CPPTargetPlatform.Win32) ||
-					(CompileEnvironment.Config.Target.Platform == CPPTargetPlatform.Win64)))
-				{
-					Arguments.Append(" /Oy-");
-				}
+					// Only omit frame pointers on the PC (which is implied by /Ox) if wanted.
+					if ( BuildConfiguration.bOmitFramePointers == false
+					&& ((CompileEnvironment.Config.Target.Platform == CPPTargetPlatform.Win32) ||
+						(CompileEnvironment.Config.Target.Platform == CPPTargetPlatform.Win64)))
+					{
+						Arguments.Append(" /Oy-");
+					}
 
-				// Allow inline method expansion
-				Arguments.Append(" /Ob2");
+					// Allow inline method expansion
+					Arguments.Append(" /Ob2");
 
-				//
-				// LTCG
-				//
-				if (CompileEnvironment.Config.Target.Configuration == CPPTargetConfiguration.Shipping)
+					//
+					// LTCG
+					//
+					if (CompileEnvironment.Config.Target.Configuration == CPPTargetConfiguration.Shipping)
+					{
+						if( BuildConfiguration.bAllowLTCG )
+						{
+							// Enable link-time code generation.
+							Arguments.Append(" /GL");
+						}
+					}
+				}
+				else
 				{
-					if( BuildConfiguration.bAllowLTCG )
+					// Maximum optimizations if desired.
+					if( CompileEnvironment.Config.OptimizeCode >= ModuleRules.CodeOptimization.InNonDebugBuilds )
 					{
-						// Enable link-time code generation.
-						Arguments.Append(" /GL");
+						Arguments.Append( " -O3");
 					}
 				}
 			}
@@ -215,12 +312,18 @@ namespace UnrealBuildTool
 				if (CompileEnvironment.Config.CLRMode == CPPCLRMode.CLRDisabled &&
 					CompileEnvironment.Config.Target.Platform != CPPTargetPlatform.Win64)
 				{
-					// Allow the compiler to generate SSE2 instructions.
-					Arguments.Append(" /arch:SSE2");
+					if( WindowsPlatform.bUseVCCompilerArgs )
+					{ 
+						// Allow the compiler to generate SSE2 instructions.
+						Arguments.Append(" /arch:SSE2");
+					}
 				}
 
-				// Prompt the user before reporting internal errors to Microsoft.
-				Arguments.Append(" /errorReport:prompt");
+				if( WindowsPlatform.bUseVCCompilerArgs )
+				{ 
+					// Prompt the user before reporting internal errors to Microsoft.
+					Arguments.Append(" /errorReport:prompt");
+				}
 
 				if (CompileEnvironment.Config.CLRMode == CPPCLRMode.CLRDisabled)
 				{
@@ -235,6 +338,11 @@ namespace UnrealBuildTool
 					{
 						// This is required to disable exception handling in VC platform headers.
 						CompileEnvironment.Config.Definitions.Add("_HAS_EXCEPTIONS=0");
+
+						if( !WindowsPlatform.bUseVCCompilerArgs )
+						{ 
+							Arguments.Append( " -fno-exceptions" );
+						}
 					}
 				}
 				else
@@ -243,42 +351,44 @@ namespace UnrealBuildTool
 					Arguments.Append(" /EHa");
 				}
             }
-
-            // @todo - find a better place for this. 
-            if (CompileEnvironment.Config.Target.Platform == CPPTargetPlatform.HTML5)
+    		else if (CompileEnvironment.Config.Target.Platform == CPPTargetPlatform.HTML5)
             {
                 Arguments.Append(" /EHsc"); 
             }
+
             // If enabled, create debug information.
 			if (CompileEnvironment.Config.bCreateDebugInfo)
 			{
-				// Store debug info in .pdb files.
-				// @todo clang: PDB files are emited from Clang but do not fully work with Visual Studio yet (breakpoints won't hit due to "symbol read error")
-				if( BuildConfiguration.bUsePDBFiles )
-				{
-					// Create debug info suitable for E&C if wanted.
-					if( BuildConfiguration.bSupportEditAndContinue &&
-						// We only need to do this in debug as that's the only configuration that supports E&C.
-						CompileEnvironment.Config.Target.Configuration == CPPTargetConfiguration.Debug)
+				if( WindowsPlatform.bUseVCCompilerArgs )
+				{ 
+					// Store debug info in .pdb files.
+					// @todo clang: PDB files are emited from Clang but do not fully work with Visual Studio yet (breakpoints won't hit due to "symbol read error")
+					if( BuildConfiguration.bUsePDBFiles )
 					{
-						Arguments.Append(" /ZI");
+						// Create debug info suitable for E&C if wanted.
+						if( BuildConfiguration.bSupportEditAndContinue &&
+							// We only need to do this in debug as that's the only configuration that supports E&C.
+							CompileEnvironment.Config.Target.Configuration == CPPTargetConfiguration.Debug)
+						{
+							Arguments.Append(" /ZI");
+						}
+						// Regular PDB debug information.
+						else
+						{
+							Arguments.Append(" /Zi");
+						}
+						// We need to add this so VS won't lock the PDB file and prevent synchronous updates. This forces serialization through MSPDBSRV.exe.
+						// See http://msdn.microsoft.com/en-us/library/dn502518.aspx for deeper discussion of /FS switch.
+						if (BuildConfiguration.bUseIncrementalLinking && WindowsPlatform.Compiler == WindowsCompiler.VisualStudio2013)
+						{
+							Arguments.Append(" /FS");
+						}
 					}
-					// Regular PDB debug information.
+					// Store C7-format debug info in the .obj files, which is faster.
 					else
 					{
-						Arguments.Append(" /Zi");
+						Arguments.Append(" /Z7");
 					}
-                    // We need to add this so VS won't lock the PDB file and prevent synchronous updates. This forces serialization through MSPDBSRV.exe.
-                    // See http://msdn.microsoft.com/en-us/library/dn502518.aspx for deeper discussion of /FS switch.
-                    if (BuildConfiguration.bUseIncrementalLinking && WindowsPlatform.Compiler == WindowsCompiler.VisualStudio2013)
-                    {
-                        Arguments.Append(" /FS");
-                    }
-				}
-				// Store C7-format debug info in the .obj files, which is faster.
-				else
-				{
-					Arguments.Append(" /Z7");
 				}
 			}
 
@@ -287,35 +397,68 @@ namespace UnrealBuildTool
 			{
 				if( CompileEnvironment.Config.Target.Configuration == CPPTargetConfiguration.Debug && BuildConfiguration.bDebugBuildsActuallyUseDebugCRT )
 				{
-					Arguments.Append(" /MTd");
+					if( WindowsPlatform.bUseVCCompilerArgs )
+					{ 
+						Arguments.Append(" /MTd");
+					}
+					else
+					{
+						AddDefinition( ref Arguments, "_MT" );
+						AddDefinition( ref Arguments, "_DEBUG" );
+					}
 				}
 				else
 				{
-					Arguments.Append(" /MT");
+					if( WindowsPlatform.bUseVCCompilerArgs )
+					{ 
+						Arguments.Append(" /MT");
+					}
+					else
+					{
+						AddDefinition( ref Arguments, "_MT" );
+					}
 				}
 			}
 			else
 			{
 				if( CompileEnvironment.Config.Target.Configuration == CPPTargetConfiguration.Debug && BuildConfiguration.bDebugBuildsActuallyUseDebugCRT )
 				{
-					Arguments.Append(" /MDd");
+					if( WindowsPlatform.bUseVCCompilerArgs )
+					{ 
+						Arguments.Append(" /MDd");
+					}
+					else
+					{
+						AddDefinition( ref Arguments, "_MT" );
+						AddDefinition( ref Arguments, "_DEBUG" );
+						AddDefinition( ref Arguments, "_DLL" );
+					}
 				}
 				else
 				{
-					Arguments.Append(" /MD");
+					if( WindowsPlatform.bUseVCCompilerArgs )
+					{ 
+						Arguments.Append(" /MD");
+					}
+					else
+					{
+						AddDefinition( ref Arguments, "_MT" );
+						AddDefinition( ref Arguments, "_DLL" );
+					}
 				}
 			}
 
-			//
-			// Options that aren't parsed by clang-cl yet
-			//
-			if (!WindowsPlatform.bCompileWithClang && !BuildConfiguration.bRunUnrealCodeAnalyzer)	// @todo clang: Not supported in clang-cl yet
+			if (WindowsPlatform.bUseVCCompilerArgs && !BuildConfiguration.bRunUnrealCodeAnalyzer)
 			{
-				// Allow large object files to avoid hitting the 2^16 section limit when running with -StressTestUnity.
-				Arguments.Append(" /bigobj");
+				// @todo clang: Not supported in clang-cl yet
+				if( !WindowsPlatform.bCompileWithClang )
+				{ 
+					// Allow large object files to avoid hitting the 2^16 section limit when running with -StressTestUnity.
+					Arguments.Append(" /bigobj");
 
-				// Relaxes floating point precision semantics to allow more optimization.
-				Arguments.Append(" /fp:fast");
+					// Relaxes floating point precision semantics to allow more optimization.
+					Arguments.Append(" /fp:fast");
+				}
 
 				if (CompileEnvironment.Config.OptimizeCode >= ModuleRules.CodeOptimization.InNonDebugBuilds)
 				{
@@ -336,18 +479,31 @@ namespace UnrealBuildTool
 			//
 			// Options skipped for UnrealCodeAnalyzer
 			//
-			if (!BuildConfiguration.bRunUnrealCodeAnalyzer)
+			if (!BuildConfiguration.bRunUnrealCodeAnalyzer)	// Disabled when compiling UnrealCodeAnalyzer as it breaks compilation (some structs in clang/llvm headers require 8-byte alignment in 32-bit compilation)
 			{
 				if (CompileEnvironment.Config.Target.Platform == CPPTargetPlatform.Win64)
 				{
 					// Pack struct members on 8-byte boundaries.
-					Arguments.Append(" /Zp8");
+					if( WindowsPlatform.bUseVCCompilerArgs )
+					{ 
+						Arguments.Append(" /Zp8");
+					}
+					else
+					{
+						Arguments.Append(" -fpack-struct=8" );
+					}
 				}
-				else if (!UnrealBuildTool.CommandLineContains(@"UnrealCodeAnalyzer"))
+				else
 				{
 					// Pack struct members on 4-byte boundaries.
-					// Disabled when compiling UnrealCodeAnalyzer as it breaks compilation (some structs in clang/llvm headers require 8-byte alignment in 32-bit compilation)
-					Arguments.Append(" /Zp4");
+					if( WindowsPlatform.bUseVCCompilerArgs )
+					{ 
+						Arguments.Append(" /Zp4");
+					}
+					else
+					{
+						Arguments.Append(" -fpack-struct=4" );
+					}
 				}
 			}
 
@@ -356,21 +512,40 @@ namespace UnrealBuildTool
 			//
 			if (BuildConfiguration.bRunUnrealCodeAnalyzer)
 			{
-				Arguments.Append(" /DUNREAL_CODE_ANALYZER");
+				AddDefinition( ref Arguments, "UNREAL_CODE_ANALYZER");
 			}
 		}
 
 		static void AppendCLArguments_CPP( CPPEnvironment CompileEnvironment, StringBuilder Arguments )
 		{
-			// Explicitly compile the file as C++.
-			Arguments.Append(" /TP");
+			if( WindowsPlatform.bUseVCCompilerArgs )
+			{
+				// Explicitly compile the file as C++.
+				Arguments.Append(" /TP");
+			}
+			else
+			{
+				if( CompileEnvironment.Config.PrecompiledHeaderAction == PrecompiledHeaderAction.Create )
+				{ 
+					// Tell Clang to generate a PCH header
+					Arguments.Append(" -x c++-header");
+				}
+				else
+				{
+					Arguments.Append(" -x c++" );
+				}
+			}
+
 
 			if (!CompileEnvironment.Config.bEnableBufferSecurityChecks)
 			{
-				// This will disable buffer security checks (which are enabled by default) that the MS compiler adds around arrays on the stack,
-				// Which can add some performance overhead, especially in performance intensive code
-				// Only disable this if you know what you are doing, because it will be disabled for the entire module!
-				Arguments.Append(" /GS-");
+				if( WindowsPlatform.bUseVCCompilerArgs )
+				{ 
+					// This will disable buffer security checks (which are enabled by default) that the MS compiler adds around arrays on the stack,
+					// Which can add some performance overhead, especially in performance intensive code
+					// Only disable this if you know what you are doing, because it will be disabled for the entire module!
+					Arguments.Append(" /GS-");
+				}
 			}
 
 			// C++/CLI requires that RTTI is left enabled
@@ -379,25 +554,46 @@ namespace UnrealBuildTool
 				if (CompileEnvironment.Config.bUseRTTI)
 				{
 					// Enable C++ RTTI.
-					Arguments.Append(" /GR");
+					if( WindowsPlatform.bUseVCCompilerArgs )
+					{ 
+						Arguments.Append(" /GR");
+					}
+					else
+					{
+						Arguments.Append(" -Xclang -frtti-data" );
+					}
 				}
 				else
 				{
 					// Disable C++ RTTI.
-					Arguments.Append(" /GR-");
+					if( WindowsPlatform.bUseVCCompilerArgs )
+					{ 
+						Arguments.Append(" /GR-");
+					}
+					else
+					{
+						Arguments.Append(" -Xclang -fno-rtti-data" );
+					}
 				}
 			}
 
 			// Set warning level.
-			if (!BuildConfiguration.bRunUnrealCodeAnalyzer)
-			{
-				// Restrictive during regular compilation.
-				Arguments.Append(" /W4");
+			if( WindowsPlatform.bUseVCCompilerArgs )
+			{ 
+				if (!BuildConfiguration.bRunUnrealCodeAnalyzer)
+				{
+					// Restrictive during regular compilation.
+					Arguments.Append(" /W4");
+				}
+				else
+				{
+					// If we had /W4 with clang on windows we would be flooded with warnings. This will be fixed incrementally.
+					Arguments.Append(" /W0");
+				}
 			}
 			else
 			{
-				// If we had /W4 with clang on windows we would be flooded with warnings. This will be fixed incrementally.
-				Arguments.Append(" /W0");
+				Arguments.Append(" -Wall");	
 			}
 
 			if( WindowsPlatform.bCompileWithClang )
@@ -407,8 +603,12 @@ namespace UnrealBuildTool
 
 				// @todo clang: Ideally we want as few warnings disabled as possible
 				// 
+	
 				// Allow Microsoft-specific syntax to slide, even though it may be non-standard.  Needed for Windows headers.
-				Arguments.Append(" -Wno-microsoft");								
+				Arguments.Append(" -Wno-microsoft");	
+				
+				// @todo clang: Hack due to how we have our 'DummyPCH' wrappers setup when using unity builds.  This warning should not be disabled!!
+				// Arguments.Append(" -Wno-msvc-include");
 
 				// @todo clang: Kind of a shame to turn these off.  We'd like to catch unused variables, but it is tricky with how our assertion macros work.
 				Arguments.Append(" -Wno-inconsistent-missing-override");
@@ -442,11 +642,21 @@ namespace UnrealBuildTool
 		
 		static void AppendCLArguments_C(StringBuilder Arguments)
 		{
-			// Explicitly compile the file as C.
-			Arguments.Append(" /TC");
+			if( WindowsPlatform.bUseVCCompilerArgs )
+			{
+				Arguments.Append( " -x c");
+			}
+			else
+			{ 
+				// Explicitly compile the file as C.
+				Arguments.Append(" /TC");
+			}
 		
-			// Level 0 warnings.  Needed for external C projects that produce warnings at higher warning levels.
-			Arguments.Append(" /W0");
+			if( WindowsPlatform.bUseVCCompilerArgs )
+			{ 
+				// Level 0 warnings.  Needed for external C projects that produce warnings at higher warning levels.
+				Arguments.Append(" /W0");
+			}
 		}
 
 		static void AppendLinkArguments(LinkEnvironment LinkEnvironment, StringBuilder Arguments)
@@ -459,6 +669,13 @@ namespace UnrealBuildTool
 				// @todo clang: The following static libraries aren't linking correctly with Clang:
 				//		tbbmalloc.lib, zlib_64.lib, libpng_64.lib, freetype2412MT.lib, IlmImf.lib
 				//		LLD: Assertion failed: result.size() == 1, file ..\tools\lld\lib\ReaderWriter\FileArchive.cpp, line 71
+				//		
+			
+				// Only omit frame pointers on the PC (which is implied by /Ox) if wanted.
+				if ( !BuildConfiguration.bOmitFramePointers )
+				{
+					Arguments.Append(" --disable-fp-elim");
+				}
 			}
 
 			// Don't create a side-by-side manifest file for the executable.
@@ -690,22 +907,22 @@ namespace UnrealBuildTool
 		{
 			var EnvVars = VCEnvironment.SetEnvironment(CompileEnvironment.Config.Target.Platform);
 
-			StringBuilder Arguments = new StringBuilder();
-			AppendCLArguments_Global(CompileEnvironment, EnvVars, Arguments);
+			StringBuilder SharedArguments = new StringBuilder();
+			AppendCLArguments_Global(CompileEnvironment, EnvVars, SharedArguments);
 
 			// Add include paths to the argument list.
 			if (!BuildConfiguration.bRunUnrealCodeAnalyzer)
 			{
 				foreach (string IncludePath in CompileEnvironment.Config.CPPIncludeInfo.IncludePaths)
 				{
-					Arguments.AppendFormat(" /I \"{0}\"", IncludePath);
+					AddIncludePath( ref SharedArguments, IncludePath );
 				}
 			}
 			else
 			{
 				foreach (string IncludePath in CompileEnvironment.Config.CPPIncludeInfo.IncludePaths)
 				{
-					Arguments.AppendFormat(" /I \"{0}\"", System.IO.Path.GetFullPath(IncludePath));
+					AddIncludePath( ref SharedArguments, System.IO.Path.GetFullPath(IncludePath) );
 				}
 			}
 			foreach (string IncludePath in CompileEnvironment.Config.CPPIncludeInfo.SystemIncludePaths)
@@ -718,11 +935,18 @@ namespace UnrealBuildTool
 					// this case Clang will ignore any earlier occurrence of the include path, preventing a system header include path from overriding a different system 
 					// include path set later on by a module.  NOTE: When passing "-Xclang", these options will always appear at the end of the command-line string, meaning
 					// they will be forced to appear *after* all environment-variable-extracted includes.  This is technically okay though.
-					Arguments.AppendFormat(" -Xclang -internal-isystem -Xclang \"{0}\"", IncludePath);
+					if( WindowsPlatform.bUseVCCompilerArgs )
+					{ 
+						SharedArguments.AppendFormat(" -Xclang -internal-isystem -Xclang \"{0}\"", IncludePath);
+					}
+					else
+					{
+						SharedArguments.AppendFormat(" -isystem \"{0}\"", IncludePath);
+					}
 				}
 				else
 				{
-					Arguments.AppendFormat(" /I \"{0}\"", IncludePath);
+					AddIncludePath( ref SharedArguments, IncludePath );
 				}
 			}
 
@@ -734,19 +958,19 @@ namespace UnrealBuildTool
 				// .cpp file to the assembly's location.				
 				foreach (string AssemblyPath in CompileEnvironment.Config.SystemDotNetAssemblyPaths)
 				{
-					Arguments.AppendFormat(" /AI \"{0}\"", AssemblyPath);
+					SharedArguments.AppendFormat(" /AI \"{0}\"", AssemblyPath);
 				}
 
 				// Add explicit .NET framework assembly references				
 				foreach (string AssemblyName in CompileEnvironment.Config.FrameworkAssemblyDependencies)
 				{
-					Arguments.AppendFormat( " /FU \"{0}\"", AssemblyName );
+					SharedArguments.AppendFormat( " /FU \"{0}\"", AssemblyName );
 				}
 
 				// Add private assembly references				
 				foreach( PrivateAssemblyInfo CurAssemblyInfo in CompileEnvironment.PrivateAssemblyDependencies )
 				{
-					Arguments.AppendFormat( " /FU \"{0}\"", CurAssemblyInfo.FileItem.AbsolutePath );
+					SharedArguments.AppendFormat( " /FU \"{0}\"", CurAssemblyInfo.FileItem.AbsolutePath );
 				}
 			}
 
@@ -756,7 +980,7 @@ namespace UnrealBuildTool
 			{
 				// Escape all quotation marks so that they get properly passed with the command line.
 				var DefinitionArgument = Definition.Contains("\"") ? Definition.Replace("\"", "\\\"") : Definition;
-				Arguments.AppendFormat(" /D \"{0}\"", DefinitionArgument);
+				AddDefinition( ref SharedArguments, DefinitionArgument);
 			}
 
 			var BuildPlatform = UEBuildPlatform.GetBuildPlatformForCPPTargetPlatform(CompileEnvironment.Config.Target.Platform);
@@ -803,7 +1027,7 @@ namespace UnrealBuildTool
 					// Make sure the original source directory the PCH header file existed in is added as an include
 					// path -- it might be a private PCH header and we need to make sure that its found!
 					string OriginalPCHHeaderDirectory = Path.GetDirectoryName( SourceFile.AbsolutePath );
-					FileArguments.AppendFormat(" /I \"{0}\"", OriginalPCHHeaderDirectory);
+					AddIncludePath( ref FileArguments, OriginalPCHHeaderDirectory );
 
 					var PrecompiledFileExtension = UEBuildPlatform.BuildPlatformDictionary[UnrealTargetPlatform.Win64].GetBinaryExtension(UEBuildBinaryType.PrecompiledHeader);
 					// Add the precompiled header file to the produced items list.
@@ -816,17 +1040,7 @@ namespace UnrealBuildTool
 					CompileAction.ProducedItems.Add(PrecompiledHeaderFile);
 					Result.PrecompiledHeaderFile = PrecompiledHeaderFile;
 
-					// @todo clang: Ideally clang-cl would support the regular MSVC arguments for PCH files, and we wouldn't need to do this manually
-					if( WindowsPlatform.bCompileWithClang )
-					{
-						// Tell Clang to generate a PCH header
-						FileArguments.Append(" -Xclang -x -Xclang c++-header");	// @todo clang: Doesn't work due to Clang-cl overriding us at the end of the command-line (see -### option output)
-						FileArguments.AppendFormat(" /Fo\"{0}\"", PrecompiledHeaderFile.AbsolutePath);
-
-						// Clang PCH generation doesn't create an .obj file to link in, unlike MSVC
-						bEmitsObjectFile = false;
-					}
-					else
+					if( WindowsPlatform.bUseVCCompilerArgs )
 					{
 						// Add the parameters needed to compile the precompiled header file to the command-line.
 						FileArguments.AppendFormat(" /Yc\"{0}\"", CompileEnvironment.Config.PrecompiledHeaderIncludeFilename);
@@ -844,6 +1058,15 @@ namespace UnrealBuildTool
 							FileArguments.AppendFormat(" /Yl{0}", FakeUniquePCHSymbolName);
 						}
 					}
+					else
+					{
+						FileArguments.Append( " -Xclang -emit-pch" );
+
+						FileArguments.AppendFormat(" -o\"{0}\"", PrecompiledHeaderFile.AbsolutePath);
+
+						// Clang PCH generation doesn't create an .obj file to link in, unlike MSVC
+						bEmitsObjectFile = false;
+					}
 
 					FileArguments.AppendFormat(" \"{0}\"", PCHCPPFile.AbsolutePath);
 
@@ -860,7 +1083,15 @@ namespace UnrealBuildTool
 						{
 							// NOTE: With Clang, PCH headers are ALWAYS forcibly included!
 							// NOTE: This needs to be before the other include paths to ensure Clang uses it instead of the source header file.
-							FileArguments.AppendFormat(" /FI\"{0}\"", Path.ChangeExtension(CompileEnvironment.PrecompiledHeaderFile.AbsolutePath, null));
+							if( WindowsPlatform.bUseVCCompilerArgs )
+							{ 
+								FileArguments.AppendFormat(" /FI\"{0}\"", Path.ChangeExtension(CompileEnvironment.PrecompiledHeaderFile.AbsolutePath, null));
+							}
+							else
+							{
+								// FileArguments.Append( " -Xclang -fno-validate-pch" );	// @todo clang: Interesting option for faster performance
+								FileArguments.AppendFormat(" -include-pch \"{0}\"", CompileEnvironment.PrecompiledHeaderFile.AbsolutePath);
+							}
 						}
 						else
 						{
@@ -905,7 +1136,14 @@ namespace UnrealBuildTool
 					// UnrealCodeAnalyzer requires specific position of output file.
 					if (!BuildConfiguration.bRunUnrealCodeAnalyzer)
 					{
-						FileArguments.AppendFormat(" /Fo\"{0}\"", ObjectFile.AbsolutePath);
+						if( WindowsPlatform.bUseVCCompilerArgs )
+						{ 
+							FileArguments.AppendFormat(" /Fo\"{0}\"", ObjectFile.AbsolutePath);
+						}
+						else
+						{
+							FileArguments.AppendFormat(" -o\"{0}\"", ObjectFile.AbsolutePath);
+						}
 					}
 				}
 
@@ -968,7 +1206,7 @@ namespace UnrealBuildTool
 				// Add C or C++ specific compiler arguments.
 				if (bIsPlainCFile)
 				{
-					AppendCLArguments_C(FileArguments);
+					AppendCLArguments_C( FileArguments );
 				}
 				else
 				{
@@ -999,7 +1237,7 @@ namespace UnrealBuildTool
 					UnrealCodeAnalyzerArguments = @"-CreateIncludeFiles " + SourceFile.AbsolutePath + @" -OutputFile=""" + ObjectFile.AbsolutePath + @""" -- " + ClangPath + @" --driver-mode=cl ";
 				}
 
-				CompileAction.CommandArguments = UnrealCodeAnalyzerArguments + Arguments.ToString() + FileArguments.ToString() + CompileEnvironment.Config.AdditionalArguments;
+				CompileAction.CommandArguments = UnrealCodeAnalyzerArguments + SharedArguments.ToString() + FileArguments.ToString() + CompileEnvironment.Config.AdditionalArguments;
 
 				if( CompileEnvironment.Config.PrecompiledHeaderAction == PrecompiledHeaderAction.Create )
 				{
@@ -1054,47 +1292,49 @@ namespace UnrealBuildTool
 				CompileAction.CommandPath        = EnvVars.ResourceCompilerPath;
 				CompileAction.StatusDescription  = Path.GetFileName(RCFile.AbsolutePath);
 
+				var Arguments = new StringBuilder();
+
 				if( WindowsPlatform.bCompileWithClang )
 				{ 
 					CompileAction.OutputEventHandler = new DataReceivedEventHandler( ClangCompilerOutputFormatter );
 				}
 
 				// Suppress header spew
-				CompileAction.CommandArguments += " /nologo";
+				Arguments.Append( " /nologo" );
 
 				// If we're compiling for 64-bit Windows, also add the _WIN64 definition to the resource
 				// compiler so that we can switch on that in the .rc file using #ifdef.
 				if (Environment.Config.Target.Platform == CPPTargetPlatform.Win64)
 				{
-					CompileAction.CommandArguments += " /D _WIN64";
+					AddDefinition( ref Arguments, "_WIN64" );
 				}
 
 				// When targeting Windows XP with Visual Studio 2012+, we need to tell the compiler to use the older Windows SDK that works
 				// with Windows XP (http://blogs.msdn.com/b/vcblog/archive/2012/10/08/10357555.aspx)
 				if (WindowsPlatform.SupportWindowsXP)
 				{
-					CompileAction.CommandArguments += " /D_USING_V110_SDK71_";
+					AddDefinition( ref Arguments, "_USING_V110_SDK71_" );
 				}
 
 				// Language
-				CompileAction.CommandArguments += " /l 0x409";
+				Arguments.AppendFormat( " /l 0x409" );
 
 				// Include paths.
 				foreach (string IncludePath in Environment.Config.CPPIncludeInfo.IncludePaths)
 				{
-					CompileAction.CommandArguments += string.Format( " /i \"{0}\"", IncludePath );
+					AddIncludePath( ref Arguments, IncludePath );
 				}
 
 				// System include paths.
 				foreach( var SystemIncludePath in Environment.Config.CPPIncludeInfo.SystemIncludePaths )
 				{
-					CompileAction.CommandArguments += string.Format( " /i \"{0}\"", SystemIncludePath );
+					AddIncludePath( ref Arguments, SystemIncludePath );
 				}
 
 				// Preprocessor definitions.
 				foreach (string Definition in Environment.Config.Definitions)
 				{
-					CompileAction.CommandArguments += string.Format(" /d \"{0}\"", Definition);
+					AddDefinition( ref Arguments, Definition);
 				}
 
 				// Add the RES file to the produced item list.
@@ -1105,11 +1345,13 @@ namespace UnrealBuildTool
 						)
 					);
 				CompileAction.ProducedItems.Add(CompiledResourceFile);
-				CompileAction.CommandArguments += string.Format(" /fo \"{0}\"", CompiledResourceFile.AbsolutePath);
+				Arguments.AppendFormat(" /fo \"{0}\"", CompiledResourceFile.AbsolutePath);
 				Result.ObjectFiles.Add(CompiledResourceFile);
 
 				// Add the RC file as a prerequisite of the action.
-				CompileAction.CommandArguments += string.Format(" \"{0}\"", RCFile.AbsolutePath);
+				Arguments.AppendFormat(" \"{0}\"", RCFile.AbsolutePath);
+
+				CompileAction.CommandArguments = Arguments.ToString();
 
 				// Add the C++ source file and its included files to the prerequisite item list.
 				AddPrerequisiteSourceFile( Target, BuildPlatform, Environment, RCFile, CompileAction.PrerequisiteItems );
@@ -1292,7 +1534,7 @@ namespace UnrealBuildTool
 					{
 						string MAPFilePath = Path.Combine(LinkEnvironment.Config.OutputDirectory, Path.GetFileNameWithoutExtension(OutputFile.AbsolutePath) + ".map");
 						FileItem MAPFile = FileItem.GetItemByPath(MAPFilePath);
-						LinkAction.CommandArguments += string.Format(" /MAP:\"{0}\"", MAPFilePath);
+						Arguments.AppendFormat(" /MAP:\"{0}\"", MAPFilePath);
 						LinkAction.ProducedItems.Add(MAPFile);
 					}
 #endif
