Re-enable building all plugins for Rocket, but prevent attempts to link them in. Linker fixups file should only contained linked dependent modules.

diff --git a/Engine/Source/Programs/UnrealBuildTool/Configuration/UEBuildTarget.cs b/Engine/Source/Programs/UnrealBuildTool/Configuration/UEBuildTarget.cs
--- a/Engine/Source/Programs/UnrealBuildTool/Configuration/UEBuildTarget.cs
+++ b/Engine/Source/Programs/UnrealBuildTool/Configuration/UEBuildTarget.cs
@@ -1941,7 +1941,7 @@ namespace UnrealBuildTool
 				// Create the cpp file
 				List<string> LinkerFixupsFileContents = new List<string>();
 				NewModule.bSkipDefinitionsForCompileEnvironment = false;
-				GenerateLinkerFixupsContents(LinkerFixupsFileContents, NewModule.CreateModuleCompileEnvironment(GlobalCompileEnvironment), HeaderFilename, LinkerFixupsName, PrivateDependencyModuleNames);
+				GenerateLinkerFixupsContents(ExecutableBinary, LinkerFixupsFileContents, NewModule.CreateModuleCompileEnvironment(GlobalCompileEnvironment), HeaderFilename, LinkerFixupsName, PrivateDependencyModuleNames);
 				NewModule.bSkipDefinitionsForCompileEnvironment = true;
 				
 				// Determine if the file changed. Write it if it either doesn't exist or the contents are different.
@@ -1963,7 +1963,7 @@ namespace UnrealBuildTool
 			}
 		}
 
-		private void GenerateLinkerFixupsContents(List<string> LinkerFixupsFileContents, CPPEnvironment CompileEnvironment, string HeaderFilename, string LinkerFixupsName, List<string> PrivateDependencyModuleNames)
+		private void GenerateLinkerFixupsContents(UEBuildBinary ExecutableBinary, List<string> LinkerFixupsFileContents, CPPEnvironment CompileEnvironment, string HeaderFilename, string LinkerFixupsName, List<string> PrivateDependencyModuleNames)
 		{
 			LinkerFixupsFileContents.Add("#include \"" + HeaderFilename + "\"");
 
@@ -1992,21 +1992,11 @@ namespace UnrealBuildTool
 			LinkerFixupsFileContents.Add("{");
 
 			// Fill out the body of the function with the empty function calls. This is what causes the static libraries to be considered relevant
+			var DependencyModules = ExecutableBinary.GetAllDependencyModules(bIncludeDynamicallyLoaded: false, bForceCircular: false);
+			foreach(string ModuleName in DependencyModules.OfType<UEBuildModuleCPP>().Where(CPPModule => CPPModule.AutoGenerateCppInfo != null).Select(CPPModule => CPPModule.Name).Distinct())
 			{
-				var UObjectModules = new List<UEBuildModuleCPP>();
-				foreach (var Binary in AppBinaries)
-				{
-					var DependencyModules = Binary.GetAllDependencyModules(bIncludeDynamicallyLoaded: false, bForceCircular: false);
-					foreach (var Module in DependencyModules.OfType<UEBuildModuleCPP>().Where(CPPModule => CPPModule.AutoGenerateCppInfo != null && !UObjectModules.Any(Module => Module.Name == CPPModule.Name)))
-					{
-						UObjectModules.Add(Module);
-					}
-				}
-				foreach (var Module in UObjectModules)
-				{
-					LinkerFixupsFileContents.Add("    extern void EmptyLinkFunctionForGeneratedCode" + Module.Name + "();");
-					LinkerFixupsFileContents.Add("    EmptyLinkFunctionForGeneratedCode" + Module.Name + "();");
-				}
+				LinkerFixupsFileContents.Add("    extern void EmptyLinkFunctionForGeneratedCode" + ModuleName + "();");
+				LinkerFixupsFileContents.Add("    EmptyLinkFunctionForGeneratedCode" + ModuleName + "();");
 			}
 			foreach (var DependencyModuleName in PrivateDependencyModuleNames)
 			{
@@ -2508,7 +2498,7 @@ namespace UnrealBuildTool
 			// Set the list of plugins that should be built
 			if (UnrealBuildTool.BuildingRocket() && TargetType != TargetRules.TargetType.Program)
 			{
-				BuildPlugins.AddRange(EnabledPlugins);//ValidPlugins);
+				BuildPlugins.AddRange(ValidPlugins);
 			}
 			else if (ShouldCompileMonolithic() || TargetType == TargetRules.TargetType.Program)
 			{
