Fixed UnrealBuildTool failing to harvest environment variables when SET command produced empty lines

diff --git a/Engine/Source/Programs/UnrealBuildTool/System/Utils.cs b/Engine/Source/Programs/UnrealBuildTool/System/Utils.cs
--- a/Engine/Source/Programs/UnrealBuildTool/System/Utils.cs
+++ b/Engine/Source/Programs/UnrealBuildTool/System/Utils.cs
@@ -222,7 +222,7 @@ namespace UnrealBuildTool
 				}
 
 				// Load environment variables (the file will be encoded without a BOM, so we need to manually specify the encoding)
-				var EnvStringsFromFile = File.ReadAllLines( EnvOutputFileName, Encoding.Unicode );
+				var EnvStringsFromFile = File.ReadAllLines( EnvOutputFileName, Encoding.Unicode ).Where( Line => !String.IsNullOrEmpty( Line ) );
 				foreach( var EnvString in EnvStringsFromFile )
 				{
 					// Parse the environment variable name and value from the string ("name=value")
