Improved IntelliSense: Generate definitions for editor targets instead of cooked monolithic games
- This changes improves the usefulness of IntelliSense on Visual Studio or Xcode code completion
- IntelliSense compiler definitions and include paths are now configured for editor targets which are inclusive of much more code
- For example, all code that was wrapped with WITH_EDITORONLY_DATA is now visible to the IntelliSense compiler

diff --git a/Engine/Source/Programs/UnrealBuildTool/System/ProjectFileGenerator.cs b/Engine/Source/Programs/UnrealBuildTool/System/ProjectFileGenerator.cs
--- a/Engine/Source/Programs/UnrealBuildTool/System/ProjectFileGenerator.cs
+++ b/Engine/Source/Programs/UnrealBuildTool/System/ProjectFileGenerator.cs
@@ -528,7 +528,12 @@ namespace UnrealBuildTool
 							{
 								if( !String.IsNullOrEmpty( ProjectTarget.TargetFilePath ) )
 								{
-									IntelliSenseTargetFiles.Add( ProjectTarget.TargetFilePath );
+									// Only bother with the editor target.  We want to make sure that definitions are setup to be as inclusive as possible
+									// for good quality IntelliSense.  For example, we want WITH_EDITORONLY_DATA=1, so using the editor targets works well.
+									if( ProjectTarget.TargetRules.Type == TargetRules.TargetType.Editor )
+									{ 
+										IntelliSenseTargetFiles.Add( ProjectTarget.TargetFilePath );
+									}
 								}
 							}
 						}
@@ -552,7 +557,12 @@ namespace UnrealBuildTool
 							{
 								if( !String.IsNullOrEmpty( ProjectTarget.TargetFilePath ) )
 								{
-									IntelliSenseTargetFiles.Add( ProjectTarget.TargetFilePath );
+									// Only bother with the editor target.  We want to make sure that definitions are setup to be as inclusive as possible
+									// for good quality IntelliSense.  For example, we want WITH_EDITORONLY_DATA=1, so using the editor targets works well.
+									if( ProjectTarget.TargetRules.Type == TargetRules.TargetType.Editor )
+									{ 
+										IntelliSenseTargetFiles.Add( ProjectTarget.TargetFilePath );
+									}
 								}
 							}
 						}
