Fix for non development windows xp builds
- We need to allow for binariy names to be overriden as opposed to appname.
- - Appname was being used as a part of the binaryname, which was fine for development but not other configs. AppName_xp-Win32-Shipping, and we needed it to be AppName-Win32-Shipping_xp.

Added the capacity to specify that XP is supported in an ini.

diff --git a/Engine/Source/Programs/UnrealBuildTool/Configuration/UEBuildTarget.cs b/Engine/Source/Programs/UnrealBuildTool/Configuration/UEBuildTarget.cs
--- a/Engine/Source/Programs/UnrealBuildTool/Configuration/UEBuildTarget.cs
+++ b/Engine/Source/Programs/UnrealBuildTool/Configuration/UEBuildTarget.cs
@@ -659,6 +659,9 @@ namespace UnrealBuildTool
 		/** The name of the application the target is part of. */
 		public string AppName;
 
+		/** AppName overriden from the commandline */
+		private static string OverridenAppName;
+
 		/** The name of the game the target is part of - can be empty */
 		public string GameName;
 
@@ -766,13 +769,10 @@ namespace UnrealBuildTool
 			const string OverrideTargetAppNameSwitch = "-overridetargetappname=";
 			if ((CmdlineAppName = InAdditionalDefinitions.Find(x => x.StartsWith(OverrideTargetAppNameSwitch))) != null)
 			{
-				AppName = CmdlineAppName.Substring(OverrideTargetAppNameSwitch.Length);
-			}
-			else
-			{
-				AppName = InAppName;
+				OverridenAppName = CmdlineAppName.Substring(OverrideTargetAppNameSwitch.Length);
 			}
 
+			AppName = InAppName;
 			GameName = InGameName;
 			Platform = InPlatform;
 			Configuration = InConfiguration;
@@ -831,7 +831,7 @@ namespace UnrealBuildTool
 			TargetTypeOrNull = (Rules != null) ? Rules.Type : (TargetRules.TargetType?)null;
 
 			// Construct the output path based on configuration, platform, game if not specified.
-            OutputPaths = MakeBinaryPaths("", AppName, UEBuildBinaryType.Executable, TargetType, null, InAppName, Configuration == UnrealTargetConfiguration.Shipping ? Rules.ForceNameAsForDevelopment() : false, Rules.ExeBinariesSubFolder);
+            OutputPaths = MakeBinaryPaths("", AppName, UEBuildBinaryType.Executable, TargetType, null, AppName, Configuration == UnrealTargetConfiguration.Shipping ? Rules.ForceNameAsForDevelopment() : false, Rules.ExeBinariesSubFolder);
 			for (int Index = 0; Index < OutputPaths.Length; Index++)
 			{
 				OutputPaths[Index] = Path.GetFullPath(OutputPaths[Index]);
@@ -2395,7 +2395,11 @@ namespace UnrealBuildTool
 			{
 				Prefix = "lib";
 			}
-			if (LocalConfig == UnrealTargetConfiguration.Development || bForceNameAsForDevelopment)
+			if (!string.IsNullOrEmpty(OverridenAppName) && BinaryType == UEBuildBinaryType.Executable)
+			{
+				OutBinaryPath = Path.Combine(BaseDirectory, String.Format("{2}{0}{1}", OverridenAppName, BinaryExtension, Prefix));
+			}
+			else if (LocalConfig == UnrealTargetConfiguration.Development || bForceNameAsForDevelopment)
 			{
 				OutBinaryPath = Path.Combine(BaseDirectory, String.Format("{3}{0}{1}{2}", BinaryName, BinarySuffix, BinaryExtension, Prefix));
 			}
diff --git a/Engine/Source/Programs/UnrealBuildTool/Windows/UEBuildWindows.cs b/Engine/Source/Programs/UnrealBuildTool/Windows/UEBuildWindows.cs
--- a/Engine/Source/Programs/UnrealBuildTool/Windows/UEBuildWindows.cs
+++ b/Engine/Source/Programs/UnrealBuildTool/Windows/UEBuildWindows.cs
@@ -392,8 +392,20 @@ namespace UnrealBuildTool
         {
             InBuildTarget.GlobalCompileEnvironment.Config.Definitions.Add("WIN32=1");
 
-            // Win32 XP is only supported at this time.
-            SupportWindowsXP = SupportWindowsXP && InBuildTarget.Platform == UnrealTargetPlatform.Win32;
+			{
+				ConfigCacheIni Ini = new ConfigCacheIni(UnrealTargetPlatform.Win64, "Engine", UnrealBuildTool.GetUProjectPath());
+				string MinimumOS;
+				if (Ini.GetString("/Script/WindowsTargetPlatform.WindowsTargetSettings", "MinimumOSVersion", out MinimumOS))
+				{
+					if(string.IsNullOrEmpty(MinimumOS) == false)
+					{
+						SupportWindowsXP = MinimumOS == "MSOS_XP";
+					}
+				}
+			}
+			// Win32 XP is only supported at this time.
+			SupportWindowsXP = SupportWindowsXP && InBuildTarget.Platform == UnrealTargetPlatform.Win32;
+
 
             if( SupportWindowsXP )
             {
