Don't infect a project's IntelliSense definitions with definitions from other project's targets
- This change improves the quality of IntelliSense and error list results in Visual Studio and Xcode
- This fixes definitions from Program targets from being set on the UE4 project that caused undesirable IntelliSense results
- For example, we never want WITH_EDITORONLY_DATA=1 or IS_PROGRAM=1 to be set for engine modules

diff --git a/Engine/Source/Programs/UnrealBuildTool/System/Project.cs b/Engine/Source/Programs/UnrealBuildTool/System/Project.cs
--- a/Engine/Source/Programs/UnrealBuildTool/System/Project.cs
+++ b/Engine/Source/Programs/UnrealBuildTool/System/Project.cs
@@ -250,20 +250,22 @@ namespace UnrealBuildTool
 		/// Splits the definition text into macro name and value (if any).
 		/// </summary>
 		/// <param name="Definition">Definition text</param>
+		/// <param name="Key">Out: The definition name</param>
+		/// <param name="Value">Out: The definition value or null if it has none</param>
 		/// <returns>Pair representing macro name and value.</returns>
-		private KeyValuePair<string, string> SplitDefinitionAndValue(string Definition)
+		private void SplitDefinitionAndValue(string Definition, out String Key, out String Value)
 		{
 			int EqualsIndex = Definition.IndexOf('=');
-			KeyValuePair<string, string> DefineAndValue;
 			if (EqualsIndex >= 0)
 			{
-				DefineAndValue = new KeyValuePair<string, string>(Definition.Substring(0, EqualsIndex), Definition.Substring(EqualsIndex + 1));
+				Key = Definition.Substring(0, EqualsIndex);
+				Value = Definition.Substring(EqualsIndex + 1);
 			}
 			else
 			{
-				DefineAndValue = new KeyValuePair<string, string>(Definition, "");
+				Key = Definition;
+				Value = "";
 			}
-			return DefineAndValue;
 		}
 
 		/// <summary>
@@ -272,31 +274,35 @@ namespace UnrealBuildTool
 		/// <param name="NewPreprocessorDefinitions">List of preprocessor definitons to add</param>
 		public void AddIntelliSensePreprocessorDefinitions( List<string> NewPreprocessorDefinitions ) 
 		{
-			foreach( var CurDef in NewPreprocessorDefinitions )
-			{
-				if( KnownIntelliSensePreprocessorDefinitions.Add( CurDef ) )
+			if( ProjectFileGenerator.OnlyGenerateIntelliSenseDataForProject == null || 
+				ProjectFileGenerator.OnlyGenerateIntelliSenseDataForProject == this )
+			{ 
+				foreach( var CurDef in NewPreprocessorDefinitions )
 				{
-					// Make sure that it doesn't exist already
-					var AlreadyExists = false;
-					var CurDefAndValue = SplitDefinitionAndValue(CurDef);
-					for (int DefineIndex = 0; DefineIndex < IntelliSensePreprocessorDefinitions.Count; ++DefineIndex)
+					// Don't add definitions and value combinations that have already been added for this project
+					if( KnownIntelliSensePreprocessorDefinitions.Add( CurDef ) )
 					{
-						var ExistingDefAndValue = SplitDefinitionAndValue(IntelliSensePreprocessorDefinitions[DefineIndex]);
-						if (ExistingDefAndValue.Key == CurDefAndValue.Key)
+						// Go ahead and check to see if the definition already exists, but the value is different
+						var AlreadyExists = false;
+
+						string Def, Value;
+						SplitDefinitionAndValue(CurDef, out Def, out Value);
+						for (int DefineIndex = 0; DefineIndex < IntelliSensePreprocessorDefinitions.Count; ++DefineIndex)
 						{
-							// Favor defines that set their value to 1 and replace the old one if it was set to "0"
-							if (CurDefAndValue.Value == "1" && ExistingDefAndValue.Value == "0")
+							string ExistingDef, ExistingValue;
+							SplitDefinitionAndValue(IntelliSensePreprocessorDefinitions[DefineIndex], out ExistingDef, out ExistingValue);
+							if (ExistingDef == Def)
 							{
-								IntelliSensePreprocessorDefinitions[DefineIndex] = CurDef;
+								// Already exists, but the value is changing.  We don't bother clobbering values for existing defines for this project.
+								AlreadyExists = true;
+								break;
 							}
-							AlreadyExists = true;
-							break;
 						}
-					}
 
-					if( !AlreadyExists )
-					{
-						IntelliSensePreprocessorDefinitions.Add( CurDef );
+						if( !AlreadyExists )
+						{
+							IntelliSensePreprocessorDefinitions.Add( CurDef );
+						}
 					}
 				}
 			}
@@ -309,43 +315,47 @@ namespace UnrealBuildTool
 		/// <param name="NewIncludePaths">List of include paths to add</param>
 		public void AddInteliiSenseIncludePaths(HashSet<string> NewIncludePaths, bool bAddingSystemIncludes) 
 		{
-			foreach( var CurPath in NewIncludePaths )
-			{
-				if( bAddingSystemIncludes ? KnownIntelliSenseSystemIncludeSearchPaths.Add( CurPath ) : KnownIntelliSenseIncludeSearchPaths.Add( CurPath ) )
+			if( ProjectFileGenerator.OnlyGenerateIntelliSenseDataForProject == null || 
+				ProjectFileGenerator.OnlyGenerateIntelliSenseDataForProject == this )
+			{ 
+				foreach( var CurPath in NewIncludePaths )
 				{
-					string PathRelativeToProjectFile;
-
-					// If the include string is an environment variable (e.g. $(DXSDK_DIR)), then we never want to
-					// give it a relative path
-					if( CurPath.StartsWith( "$(" ) )
-					{
-						PathRelativeToProjectFile = CurPath;
-					}
-					else
+					if( bAddingSystemIncludes ? KnownIntelliSenseSystemIncludeSearchPaths.Add( CurPath ) : KnownIntelliSenseIncludeSearchPaths.Add( CurPath ) )
 					{
-						// Incoming include paths are relative to the solution directory, but we need these paths to be
-						// relative to the project file's directory
-						PathRelativeToProjectFile = NormalizeProjectPath( CurPath );
-					}
+						string PathRelativeToProjectFile;
 
-					// Trim any trailing slash
-					PathRelativeToProjectFile = PathRelativeToProjectFile.TrimEnd('/', '\\');
+						// If the include string is an environment variable (e.g. $(DXSDK_DIR)), then we never want to
+						// give it a relative path
+						if( CurPath.StartsWith( "$(" ) )
+						{
+							PathRelativeToProjectFile = CurPath;
+						}
+						else
+						{
+							// Incoming include paths are relative to the solution directory, but we need these paths to be
+							// relative to the project file's directory
+							PathRelativeToProjectFile = NormalizeProjectPath( CurPath );
+						}
 
-					// Make sure that it doesn't exist already
-					var AlreadyExists = false;
-					List<string> SearchPaths = bAddingSystemIncludes ? IntelliSenseSystemIncludeSearchPaths : IntelliSenseIncludeSearchPaths;
-					foreach( var ExistingPath in SearchPaths )
-					{
-						if( PathRelativeToProjectFile == ExistingPath )
+						// Trim any trailing slash
+						PathRelativeToProjectFile = PathRelativeToProjectFile.TrimEnd('/', '\\');
+
+						// Make sure that it doesn't exist already
+						var AlreadyExists = false;
+						List<string> SearchPaths = bAddingSystemIncludes ? IntelliSenseSystemIncludeSearchPaths : IntelliSenseIncludeSearchPaths;
+						foreach( var ExistingPath in SearchPaths )
 						{
-							AlreadyExists = true;
-							break;
+							if( PathRelativeToProjectFile == ExistingPath )
+							{
+								AlreadyExists = true;
+								break;
+							}
 						}
-					}
 
-					if( !AlreadyExists )
-					{
-						SearchPaths.Add( PathRelativeToProjectFile );
+						if( !AlreadyExists )
+						{
+							SearchPaths.Add( PathRelativeToProjectFile );
+						}
 					}
 				}
 			}
diff --git a/Engine/Source/Programs/UnrealBuildTool/System/ProjectFileGenerator.cs b/Engine/Source/Programs/UnrealBuildTool/System/ProjectFileGenerator.cs
--- a/Engine/Source/Programs/UnrealBuildTool/System/ProjectFileGenerator.cs
+++ b/Engine/Source/Programs/UnrealBuildTool/System/ProjectFileGenerator.cs
@@ -198,6 +198,15 @@ namespace UnrealBuildTool
 		/// @todo projectfiles: Nasty global static list.  This is only really used for IntelliSense, and to avoid extra folder searches for projects we've already cached source files for.
 		public static readonly Dictionary<string, ProjectFile> ModuleToProjectFileMap = new Dictionary<string, ProjectFile>( StringComparer.InvariantCultureIgnoreCase );
 
+		/// When generating IntelliSense data, we may want to only generate data for a specific project file, even if other targets make use of modules
+		/// in this project file.  This is useful to prevent unusual or hacky global definitions from Programs affecting the Editor/Engine modules.  We
+		/// always want the most common and useful definitions to be set when working with solutions with many modules.
+		public static ProjectFile OnlyGenerateIntelliSenseDataForProject
+		{
+			get;
+			private set;
+		}
+
 	
 		/// File extension for project files we'll be generating (e.g. ".vcxproj")
 		abstract public string ProjectFileExtension
@@ -519,7 +528,7 @@ namespace UnrealBuildTool
 				{
 					// Figure out which targets we need about IntelliSense for.  We only need to worry about targets for projects
 					// that we're actually generating in this session.
-					List<string> IntelliSenseTargetFiles = new List<string>();
+					var IntelliSenseTargetFiles = new List<Tuple<ProjectFile, string>>();
 					{
 						// Engine targets
 						if( EngineProject != null )
@@ -532,7 +541,7 @@ namespace UnrealBuildTool
 									// for good quality IntelliSense.  For example, we want WITH_EDITORONLY_DATA=1, so using the editor targets works well.
 									if( ProjectTarget.TargetRules.Type == TargetRules.TargetType.Editor )
 									{ 
-										IntelliSenseTargetFiles.Add( ProjectTarget.TargetFilePath );
+										IntelliSenseTargetFiles.Add( Tuple.Create(EngineProject, ProjectTarget.TargetFilePath) );
 									}
 								}
 							}
@@ -545,7 +554,7 @@ namespace UnrealBuildTool
 							{
 								if( !String.IsNullOrEmpty( ProjectTarget.TargetFilePath ) )
 								{
-									IntelliSenseTargetFiles.Add( ProjectTarget.TargetFilePath );
+									IntelliSenseTargetFiles.Add( Tuple.Create( ProgramProject, ProjectTarget.TargetFilePath ) );
 								}
 							}
 						}
@@ -561,7 +570,7 @@ namespace UnrealBuildTool
 									// for good quality IntelliSense.  For example, we want WITH_EDITORONLY_DATA=1, so using the editor targets works well.
 									if( ProjectTarget.TargetRules.Type == TargetRules.TargetType.Editor )
 									{ 
-										IntelliSenseTargetFiles.Add( ProjectTarget.TargetFilePath );
+										IntelliSenseTargetFiles.Add( Tuple.Create( GameProject, ProjectTarget.TargetFilePath ) );
 									}
 								}
 							}
@@ -1289,7 +1298,7 @@ namespace UnrealBuildTool
 		/// <param name="Arguments">Incoming command-line arguments to UBT</param>
 		/// <param name="TargetFiles">Target files</param>
 		/// <return>Whether the process was successful or not</return>
-		private bool GenerateIntelliSenseData( String[] Arguments, List<string> TargetFiles )
+		private bool GenerateIntelliSenseData( String[] Arguments, List<Tuple<ProjectFile, string>> TargetFiles )
 		{
 			var bSuccess = true;
 			if( ShouldGenerateIntelliSenseData() && TargetFiles.Count > 0 )
@@ -1298,8 +1307,10 @@ namespace UnrealBuildTool
 				{
 					for(int TargetIndex = 0; TargetIndex < TargetFiles.Count; ++TargetIndex)
 					{
-						var CurTarget = TargetFiles[ TargetIndex ];
-						var TargetName = Utils.GetFilenameWithoutAnyExtensions( CurTarget );	// Twice, to remove both extensions from *.Target.cs file
+						var TargetProjectFile = TargetFiles[ TargetIndex ].Item1;
+						var CurTarget = TargetFiles[ TargetIndex ].Item2;
+
+						var TargetName = Utils.GetFilenameWithoutAnyExtensions( CurTarget );
 
 						Log.TraceVerbose( "Found target: " + TargetName );
 
@@ -1311,7 +1322,14 @@ namespace UnrealBuildTool
 						// projects from being able to do the same. Capture the state beforehand, and reset if after running UBT. 
 						bool bHasProjectFile = UnrealBuildTool.HasUProjectFile();
 
+						// We only want to update definitions and include paths for modules that are part of this target's project file.
+						ProjectFileGenerator.OnlyGenerateIntelliSenseDataForProject = TargetProjectFile;
+
+						// Run UnrealBuildTool, pretending to build this target but instead only gathering data for IntelliSense (include paths and definitions).
+						// No actual compiling or linking will happen because we early out using the ProjectFileGenerator.bGenerateProjectFiles global
 						bSuccess = UnrealBuildTool.RunUBT( ArgumentsCopy ) == ECompilationResult.Succeeded;
+						ProjectFileGenerator.OnlyGenerateIntelliSenseDataForProject = null;
+
 						if( !bSuccess )
 						{
 							break;
