Unreal Build Tool: Enabled remote linking with XGE on Windows platform (renabled)
- Especially in modular builds, remote linking speeds up build times very significantly
- Also enabled remote resource compilation, which also makes modular builds faster when XGE is installed

diff --git a/Engine/Source/Programs/UnrealBuildTool/Windows/VCToolChain.cs b/Engine/Source/Programs/UnrealBuildTool/Windows/VCToolChain.cs
--- a/Engine/Source/Programs/UnrealBuildTool/Windows/VCToolChain.cs
+++ b/Engine/Source/Programs/UnrealBuildTool/Windows/VCToolChain.cs
@@ -68,12 +68,12 @@ namespace UnrealBuildTool
 				String.Append( " /I " + IncludePath );
 			}
 			else
-			{
+		{
 				String.Append( " -I" + IncludePath );
 			}
 		}
 
-	
+
 		static void AppendCLArguments_Global(CPPEnvironment CompileEnvironment, VCEnvironment EnvVars, StringBuilder Arguments)
 		{
 			// NOTE: Uncommenting this line will print includes as they are encountered by the preprocessor.  This can help with diagnosing include order problems.
@@ -83,7 +83,7 @@ namespace UnrealBuildTool
 			}
 			else
 			{
-				// Arguments.Append( " /showIncludes" );
+			// Arguments.Append( " /showIncludes" );
 			}
 
 			if( WindowsPlatform.bCompileWithClang )
@@ -141,14 +141,14 @@ namespace UnrealBuildTool
 
 			if( WindowsPlatform.bUseVCCompilerArgs )
 			{ 
-				// Prevents the compiler from displaying its logo for each invocation.
-				Arguments.Append(" /nologo");
+			// Prevents the compiler from displaying its logo for each invocation.
+			Arguments.Append(" /nologo");
 
-				// Enable intrinsic functions.
+			// Enable intrinsic functions.
 				Arguments.Append(" /Oi");	// @todo clang: No Clang equivalent to this?
 			}
 
-						
+			
 			if( WindowsPlatform.bCompileWithClang )
 			{
 				// Tell the Clang compiler whether we want to generate 32-bit code or 64-bit code
@@ -174,29 +174,29 @@ namespace UnrealBuildTool
 
 			if( WindowsPlatform.bUseVCCompilerArgs )
 			{
-				// Separate functions for linker.
-				Arguments.Append(" /Gy");
+			// Separate functions for linker.
+			Arguments.Append(" /Gy");
 
-				// Allow 800% of the default memory allocation limit.
-				Arguments.Append(" /Zm800");
+			// Allow 800% of the default memory allocation limit.
+			Arguments.Append(" /Zm800");
 
-				// Disable "The file contains a character that cannot be represented in the current code page" warning for non-US windows.
-				Arguments.Append(" /wd4819");
+			// Disable "The file contains a character that cannot be represented in the current code page" warning for non-US windows.
+			Arguments.Append(" /wd4819");
 			}
 
 			if( BuildConfiguration.bUseSharedPCHs )
 			{
 				if( WindowsPlatform.bUseVCCompilerArgs )
 				{ 
-					// @todo SharedPCH: Disable warning about PCH defines not matching .cpp defines.  We "cheat" these defines a little
-					// bit to make shared PCHs work.  But it's totally safe.  Trust us.
-					Arguments.Append(" /wd4651");
+				// @todo SharedPCH: Disable warning about PCH defines not matching .cpp defines.  We "cheat" these defines a little
+				// bit to make shared PCHs work.  But it's totally safe.  Trust us.
+				Arguments.Append(" /wd4651");
 
-					// @todo SharedPCH: Disable warning about redefining *API macros.  The PCH header is compiled with various DLLIMPORTs, but
-					// when a module that uses that PCH header *IS* one of those imports, that module is compiled with EXPORTS, so the macro
-					// is redefined on the command-line.  We need to clobber those defines to make shared PCHs work properly!
-					Arguments.Append(" /wd4005");
-				}
+				// @todo SharedPCH: Disable warning about redefining *API macros.  The PCH header is compiled with various DLLIMPORTs, but
+				// when a module that uses that PCH header *IS* one of those imports, that module is compiled with EXPORTS, so the macro
+				// is redefined on the command-line.  We need to clobber those defines to make shared PCHs work properly!
+				Arguments.Append(" /wd4005");
+			}
 			}
 
 			// If compiling as a DLL, set the relevant defines
@@ -229,29 +229,29 @@ namespace UnrealBuildTool
 			{
 				if( WindowsPlatform.bUseVCCompilerArgs )
 				{ 
-					// Disable compiler optimization.
-					Arguments.Append(" /Od");
+				// Disable compiler optimization.
+				Arguments.Append(" /Od");
 
-					// Favor code size (especially useful for embedded platforms).
-					Arguments.Append(" /Os");
+				// Favor code size (especially useful for embedded platforms).
+				Arguments.Append(" /Os");
 
-					// Allow inline method expansion unless E&C support is requested
-					if( !BuildConfiguration.bSupportEditAndContinue )
-					{
-						Arguments.Append(" /Ob2");
-					}
+				// Allow inline method expansion unless E&C support is requested
+				if( !BuildConfiguration.bSupportEditAndContinue )
+				{
+					Arguments.Append(" /Ob2");
+				}
 
-					if ((CompileEnvironment.Config.Target.Platform == CPPTargetPlatform.Win32) || 
-						(CompileEnvironment.Config.Target.Platform == CPPTargetPlatform.Win64))
+				if ((CompileEnvironment.Config.Target.Platform == CPPTargetPlatform.Win32) || 
+					(CompileEnvironment.Config.Target.Platform == CPPTargetPlatform.Win64))
+				{
+					// Runtime stack checks are not allowed when compiling for CLR
+					if (CompileEnvironment.Config.CLRMode == CPPCLRMode.CLRDisabled)
 					{
-						// Runtime stack checks are not allowed when compiling for CLR
-						if (CompileEnvironment.Config.CLRMode == CPPCLRMode.CLRDisabled)
-						{
-							Arguments.Append(" /RTCs");
-						}
+						Arguments.Append(" /RTCs");
 					}
 				}
 			}
+			}
 			//
 			//	Development and LTCG
 			//
@@ -259,38 +259,38 @@ namespace UnrealBuildTool
 			{
 				if( WindowsPlatform.bUseVCCompilerArgs )
 				{ 
-					// Maximum optimizations if desired.
-					if( CompileEnvironment.Config.OptimizeCode >= ModuleRules.CodeOptimization.InNonDebugBuilds )
-					{
-						Arguments.Append(" /Ox");
-					}
+				// Maximum optimizations if desired.
+				if( CompileEnvironment.Config.OptimizeCode >= ModuleRules.CodeOptimization.InNonDebugBuilds )
+				{
+					Arguments.Append(" /Ox");
+				}
 				
-					// Favor code speed.
-					Arguments.Append(" /Ot");
+				// Favor code speed.
+				Arguments.Append(" /Ot");
 
-					// Only omit frame pointers on the PC (which is implied by /Ox) if wanted.
-					if ( BuildConfiguration.bOmitFramePointers == false
-					&& ((CompileEnvironment.Config.Target.Platform == CPPTargetPlatform.Win32) ||
-						(CompileEnvironment.Config.Target.Platform == CPPTargetPlatform.Win64)))
-					{
-						Arguments.Append(" /Oy-");
-					}
+				// Only omit frame pointers on the PC (which is implied by /Ox) if wanted.
+				if ( BuildConfiguration.bOmitFramePointers == false
+				&& ((CompileEnvironment.Config.Target.Platform == CPPTargetPlatform.Win32) ||
+					(CompileEnvironment.Config.Target.Platform == CPPTargetPlatform.Win64)))
+				{
+					Arguments.Append(" /Oy-");
+				}
 
-					// Allow inline method expansion
-					Arguments.Append(" /Ob2");
+				// Allow inline method expansion
+				Arguments.Append(" /Ob2");
 
-					//
-					// LTCG
-					//
-					if (CompileEnvironment.Config.Target.Configuration == CPPTargetConfiguration.Shipping)
+				//
+				// LTCG
+				//
+				if (CompileEnvironment.Config.Target.Configuration == CPPTargetConfiguration.Shipping)
+				{
+					if( BuildConfiguration.bAllowLTCG )
 					{
-						if( BuildConfiguration.bAllowLTCG )
-						{
-							// Enable link-time code generation.
-							Arguments.Append(" /GL");
-						}
+						// Enable link-time code generation.
+						Arguments.Append(" /GL");
 					}
 				}
+			}
 				else
 				{
 					// Maximum optimizations if desired.
@@ -314,15 +314,15 @@ namespace UnrealBuildTool
 				{
 					if( WindowsPlatform.bUseVCCompilerArgs )
 					{ 
-						// Allow the compiler to generate SSE2 instructions.
-						Arguments.Append(" /arch:SSE2");
-					}
+					// Allow the compiler to generate SSE2 instructions.
+					Arguments.Append(" /arch:SSE2");
+				}
 				}
 
 				if( WindowsPlatform.bUseVCCompilerArgs )
 				{ 
-					// Prompt the user before reporting internal errors to Microsoft.
-					Arguments.Append(" /errorReport:prompt");
+				// Prompt the user before reporting internal errors to Microsoft.
+				Arguments.Append(" /errorReport:prompt");
 				}
 
 				if (CompileEnvironment.Config.CLRMode == CPPCLRMode.CLRDisabled)
@@ -361,36 +361,36 @@ namespace UnrealBuildTool
 			{
 				if( WindowsPlatform.bUseVCCompilerArgs )
 				{ 
-					// Store debug info in .pdb files.
-					// @todo clang: PDB files are emited from Clang but do not fully work with Visual Studio yet (breakpoints won't hit due to "symbol read error")
-					if( BuildConfiguration.bUsePDBFiles )
+				// Store debug info in .pdb files.
+				// @todo clang: PDB files are emited from Clang but do not fully work with Visual Studio yet (breakpoints won't hit due to "symbol read error")
+				if( BuildConfiguration.bUsePDBFiles )
+				{
+					// Create debug info suitable for E&C if wanted.
+					if( BuildConfiguration.bSupportEditAndContinue &&
+						// We only need to do this in debug as that's the only configuration that supports E&C.
+						CompileEnvironment.Config.Target.Configuration == CPPTargetConfiguration.Debug)
 					{
-						// Create debug info suitable for E&C if wanted.
-						if( BuildConfiguration.bSupportEditAndContinue &&
-							// We only need to do this in debug as that's the only configuration that supports E&C.
-							CompileEnvironment.Config.Target.Configuration == CPPTargetConfiguration.Debug)
-						{
-							Arguments.Append(" /ZI");
-						}
-						// Regular PDB debug information.
-						else
-						{
-							Arguments.Append(" /Zi");
-						}
-						// We need to add this so VS won't lock the PDB file and prevent synchronous updates. This forces serialization through MSPDBSRV.exe.
-						// See http://msdn.microsoft.com/en-us/library/dn502518.aspx for deeper discussion of /FS switch.
-						if (BuildConfiguration.bUseIncrementalLinking && WindowsPlatform.Compiler == WindowsCompiler.VisualStudio2013)
-						{
-							Arguments.Append(" /FS");
-						}
+						Arguments.Append(" /ZI");
 					}
-					// Store C7-format debug info in the .obj files, which is faster.
+					// Regular PDB debug information.
 					else
 					{
-						Arguments.Append(" /Z7");
+						Arguments.Append(" /Zi");
 					}
+                    // We need to add this so VS won't lock the PDB file and prevent synchronous updates. This forces serialization through MSPDBSRV.exe.
+                    // See http://msdn.microsoft.com/en-us/library/dn502518.aspx for deeper discussion of /FS switch.
+                    if (BuildConfiguration.bUseIncrementalLinking && WindowsPlatform.Compiler == WindowsCompiler.VisualStudio2013)
+                    {
+                        Arguments.Append(" /FS");
+                    }
+				}
+				// Store C7-format debug info in the .obj files, which is faster.
+				else
+				{
+					Arguments.Append(" /Z7");
 				}
 			}
+			}
 
 			// Specify the appropriate runtime library based on the platform and config.
 			if( CompileEnvironment.Config.bUseStaticCRT )
@@ -399,10 +399,10 @@ namespace UnrealBuildTool
 				{
 					if( WindowsPlatform.bUseVCCompilerArgs )
 					{ 
-						Arguments.Append(" /MTd");
-					}
-					else
-					{
+					Arguments.Append(" /MTd");
+				}
+				else
+				{
 						AddDefinition( ref Arguments, "_MT" );
 						AddDefinition( ref Arguments, "_DEBUG" );
 					}
@@ -411,8 +411,8 @@ namespace UnrealBuildTool
 				{
 					if( WindowsPlatform.bUseVCCompilerArgs )
 					{ 
-						Arguments.Append(" /MT");
-					}
+					Arguments.Append(" /MT");
+				}
 					else
 					{
 						AddDefinition( ref Arguments, "_MT" );
@@ -425,10 +425,10 @@ namespace UnrealBuildTool
 				{
 					if( WindowsPlatform.bUseVCCompilerArgs )
 					{ 
-						Arguments.Append(" /MDd");
-					}
-					else
-					{
+					Arguments.Append(" /MDd");
+				}
+				else
+				{
 						AddDefinition( ref Arguments, "_MT" );
 						AddDefinition( ref Arguments, "_DEBUG" );
 						AddDefinition( ref Arguments, "_DLL" );
@@ -438,8 +438,8 @@ namespace UnrealBuildTool
 				{
 					if( WindowsPlatform.bUseVCCompilerArgs )
 					{ 
-						Arguments.Append(" /MD");
-					}
+					Arguments.Append(" /MD");
+				}
 					else
 					{
 						AddDefinition( ref Arguments, "_MT" );
@@ -452,12 +452,12 @@ namespace UnrealBuildTool
 			{
 				// @todo clang: Not supported in clang-cl yet
 				if( !WindowsPlatform.bCompileWithClang )
-				{ 
-					// Allow large object files to avoid hitting the 2^16 section limit when running with -StressTestUnity.
-					Arguments.Append(" /bigobj");
+			{
+				// Allow large object files to avoid hitting the 2^16 section limit when running with -StressTestUnity.
+				Arguments.Append(" /bigobj");
 
-					// Relaxes floating point precision semantics to allow more optimization.
-					Arguments.Append(" /fp:fast");
+				// Relaxes floating point precision semantics to allow more optimization.
+				Arguments.Append(" /fp:fast");
 				}
 
 				if (CompileEnvironment.Config.OptimizeCode >= ModuleRules.CodeOptimization.InNonDebugBuilds)
@@ -486,8 +486,8 @@ namespace UnrealBuildTool
 					// Pack struct members on 8-byte boundaries.
 					if( WindowsPlatform.bUseVCCompilerArgs )
 					{ 
-						Arguments.Append(" /Zp8");
-					}
+					Arguments.Append(" /Zp8");
+				}
 					else
 					{
 						Arguments.Append(" -fpack-struct=8" );
@@ -498,8 +498,8 @@ namespace UnrealBuildTool
 					// Pack struct members on 4-byte boundaries.
 					if( WindowsPlatform.bUseVCCompilerArgs )
 					{ 
-						Arguments.Append(" /Zp4");
-					}
+					Arguments.Append(" /Zp4");
+				}
 					else
 					{
 						Arguments.Append(" -fpack-struct=4" );
@@ -520,8 +520,8 @@ namespace UnrealBuildTool
 		{
 			if( WindowsPlatform.bUseVCCompilerArgs )
 			{
-				// Explicitly compile the file as C++.
-				Arguments.Append(" /TP");
+			// Explicitly compile the file as C++.
+			Arguments.Append(" /TP");
 			}
 			else
 			{
@@ -541,11 +541,11 @@ namespace UnrealBuildTool
 			{
 				if( WindowsPlatform.bUseVCCompilerArgs )
 				{ 
-					// This will disable buffer security checks (which are enabled by default) that the MS compiler adds around arrays on the stack,
-					// Which can add some performance overhead, especially in performance intensive code
-					// Only disable this if you know what you are doing, because it will be disabled for the entire module!
-					Arguments.Append(" /GS-");
-				}
+				// This will disable buffer security checks (which are enabled by default) that the MS compiler adds around arrays on the stack,
+				// Which can add some performance overhead, especially in performance intensive code
+				// Only disable this if you know what you are doing, because it will be disabled for the entire module!
+				Arguments.Append(" /GS-");
+			}
 			}
 
 			// C++/CLI requires that RTTI is left enabled
@@ -556,10 +556,10 @@ namespace UnrealBuildTool
 					// Enable C++ RTTI.
 					if( WindowsPlatform.bUseVCCompilerArgs )
 					{ 
-						Arguments.Append(" /GR");
-					}
-					else
-					{
+					Arguments.Append(" /GR");
+				}
+				else
+				{
 						Arguments.Append(" -Xclang -frtti-data" );
 					}
 				}
@@ -568,8 +568,8 @@ namespace UnrealBuildTool
 					// Disable C++ RTTI.
 					if( WindowsPlatform.bUseVCCompilerArgs )
 					{ 
-						Arguments.Append(" /GR-");
-					}
+					Arguments.Append(" /GR-");
+				}
 					else
 					{
 						Arguments.Append(" -Xclang -fno-rtti-data" );
@@ -580,16 +580,16 @@ namespace UnrealBuildTool
 			// Set warning level.
 			if( WindowsPlatform.bUseVCCompilerArgs )
 			{ 
-				if (!BuildConfiguration.bRunUnrealCodeAnalyzer)
-				{
-					// Restrictive during regular compilation.
-					Arguments.Append(" /W4");
-				}
-				else
-				{
-					// If we had /W4 with clang on windows we would be flooded with warnings. This will be fixed incrementally.
-					Arguments.Append(" /W0");
-				}
+			if (!BuildConfiguration.bRunUnrealCodeAnalyzer)
+			{
+				// Restrictive during regular compilation.
+				Arguments.Append(" /W4");
+			}
+			else
+			{
+				// If we had /W4 with clang on windows we would be flooded with warnings. This will be fixed incrementally.
+				Arguments.Append(" /W0");
+			}
 			}
 			else
 			{
@@ -605,8 +605,8 @@ namespace UnrealBuildTool
 				// 
 	
 				// Allow Microsoft-specific syntax to slide, even though it may be non-standard.  Needed for Windows headers.
-				Arguments.Append(" -Wno-microsoft");	
-				
+				Arguments.Append(" -Wno-microsoft");								
+
 				// @todo clang: Hack due to how we have our 'DummyPCH' wrappers setup when using unity builds.  This warning should not be disabled!!
 				// Arguments.Append(" -Wno-msvc-include");
 
@@ -648,15 +648,15 @@ namespace UnrealBuildTool
 			}
 			else
 			{ 
-				// Explicitly compile the file as C.
-				Arguments.Append(" /TC");
+			// Explicitly compile the file as C.
+			Arguments.Append(" /TC");
 			}
 		
 			if( WindowsPlatform.bUseVCCompilerArgs )
 			{ 
-				// Level 0 warnings.  Needed for external C projects that produce warnings at higher warning levels.
-				Arguments.Append(" /W0");
-			}
+			// Level 0 warnings.  Needed for external C projects that produce warnings at higher warning levels.
+			Arguments.Append(" /W0");
+		}
 		}
 
 		static void AppendLinkArguments(LinkEnvironment LinkEnvironment, StringBuilder Arguments)
@@ -1085,10 +1085,10 @@ namespace UnrealBuildTool
 							// NOTE: This needs to be before the other include paths to ensure Clang uses it instead of the source header file.
 							if( WindowsPlatform.bUseVCCompilerArgs )
 							{ 
-								FileArguments.AppendFormat(" /FI\"{0}\"", Path.ChangeExtension(CompileEnvironment.PrecompiledHeaderFile.AbsolutePath, null));
-							}
-							else
-							{
+							FileArguments.AppendFormat(" /FI\"{0}\"", Path.ChangeExtension(CompileEnvironment.PrecompiledHeaderFile.AbsolutePath, null));
+						}
+						else
+						{
 								// FileArguments.Append( " -Xclang -fno-validate-pch" );	// @todo clang: Interesting option for faster performance
 								FileArguments.AppendFormat(" -include-pch \"{0}\"", CompileEnvironment.PrecompiledHeaderFile.AbsolutePath);
 							}
@@ -1138,8 +1138,8 @@ namespace UnrealBuildTool
 					{
 						if( WindowsPlatform.bUseVCCompilerArgs )
 						{ 
-							FileArguments.AppendFormat(" /Fo\"{0}\"", ObjectFile.AbsolutePath);
-						}
+						FileArguments.AppendFormat(" /Fo\"{0}\"", ObjectFile.AbsolutePath);
+					}
 						else
 						{
 							FileArguments.AppendFormat(" -o\"{0}\"", ObjectFile.AbsolutePath);
@@ -1206,7 +1206,7 @@ namespace UnrealBuildTool
 				// Add C or C++ specific compiler arguments.
 				if (bIsPlainCFile)
 				{
-					AppendCLArguments_C( FileArguments );
+					AppendCLArguments_C(FileArguments);
 				}
 				else
 				{
@@ -1292,6 +1292,9 @@ namespace UnrealBuildTool
 				CompileAction.CommandPath        = EnvVars.ResourceCompilerPath;
 				CompileAction.StatusDescription  = Path.GetFileName(RCFile.AbsolutePath);
 
+				// Resource tool can run remotely if possible
+				CompileAction.bCanExecuteRemotely = true;
+
 				var Arguments = new StringBuilder();
 
 				if( WindowsPlatform.bCompileWithClang )
@@ -1564,8 +1567,8 @@ namespace UnrealBuildTool
 			// ignored as a prerequisite for other actions
 			LinkAction.bProducesImportLibrary = bBuildImportLibraryOnly || LinkEnvironment.Config.bIsBuildingDLL;
 
-			// Only execute linking on the local PC.
-			LinkAction.bCanExecuteRemotely = false;
+			// Allow remote linking.  Especially in modular builds with many small DLL files, this is almost always very efficient
+			LinkAction.bCanExecuteRemotely = true;
 
 			Log.TraceVerbose( "     Linking: " + LinkAction.StatusDescription );
 			Log.TraceVerbose( "     Command: " + LinkAction.CommandArguments );
