Disable PS4 and XboxOne plugins when building Rocket.

diff --git a/Engine/Source/Programs/UnrealBuildTool/Configuration/UEBuildTarget.cs b/Engine/Source/Programs/UnrealBuildTool/Configuration/UEBuildTarget.cs
--- a/Engine/Source/Programs/UnrealBuildTool/Configuration/UEBuildTarget.cs
+++ b/Engine/Source/Programs/UnrealBuildTool/Configuration/UEBuildTarget.cs
@@ -772,7 +772,7 @@ namespace UnrealBuildTool
 				OverridenAppName = CmdlineAppName.Substring(OverrideTargetAppNameSwitch.Length);
 			}
 
-			AppName = InAppName;
+				AppName = InAppName;
 			GameName = InGameName;
 			Platform = InPlatform;
 			Configuration = InConfiguration;
@@ -1146,13 +1146,13 @@ namespace UnrealBuildTool
 
 					// Normal makefile
 					{					
-						var UBTMakefilePath = UnrealBuildTool.GetUBTMakefilePath( TargetDescs );
-						if (File.Exists(UBTMakefilePath))
-						{
-							Log.TraceVerbose("\tDeleting " + UBTMakefilePath);
-							CleanFile(UBTMakefilePath);
-						}
+					var UBTMakefilePath = UnrealBuildTool.GetUBTMakefilePath( TargetDescs );
+					if (File.Exists(UBTMakefilePath))
+					{
+						Log.TraceVerbose("\tDeleting " + UBTMakefilePath);
+						CleanFile(UBTMakefilePath);
 					}
+				}
 
 					// Hot reload makefile
 					{					
@@ -1862,7 +1862,7 @@ namespace UnrealBuildTool
 					FilteredBinaries.Add(DLLBinary);
 					AnyBinariesAdded = true;
 				}
-			}
+				}
 
 			if (!AnyBinariesAdded)
 			{
@@ -1943,7 +1943,7 @@ namespace UnrealBuildTool
 				NewModule.bSkipDefinitionsForCompileEnvironment = false;
 				GenerateLinkerFixupsContents(ExecutableBinary, LinkerFixupsFileContents, NewModule.CreateModuleCompileEnvironment(GlobalCompileEnvironment), HeaderFilename, LinkerFixupsName, PrivateDependencyModuleNames);
 				NewModule.bSkipDefinitionsForCompileEnvironment = true;
-				
+
 				// Determine if the file changed. Write it if it either doesn't exist or the contents are different.
 				bool bShouldWriteFile = true;
 				if (File.Exists(LinkerFixupCPPFilename))
@@ -2502,7 +2502,7 @@ namespace UnrealBuildTool
 			// Set the list of plugins that should be built
 			if (UnrealBuildTool.BuildingRocket() && TargetType != TargetRules.TargetType.Program)
 			{
-				BuildPlugins.AddRange(ValidPlugins);
+				BuildPlugins.AddRange(ValidPlugins.Where(x => x.Directory.IndexOf("\\PS4\\", StringComparison.InvariantCultureIgnoreCase) == -1 && x.Directory.IndexOf("\\XboxOne\\", StringComparison.InvariantCultureIgnoreCase) == -1));
 			}
 			else if (ShouldCompileMonolithic() || TargetType == TargetRules.TargetType.Program)
 			{
