Prevent PS4 and XboxOne plugins being enabled by default in Rocket builds.

diff --git a/Engine/Source/Programs/UnrealBuildTool/Configuration/UEBuildTarget.cs b/Engine/Source/Programs/UnrealBuildTool/Configuration/UEBuildTarget.cs
--- a/Engine/Source/Programs/UnrealBuildTool/Configuration/UEBuildTarget.cs
+++ b/Engine/Source/Programs/UnrealBuildTool/Configuration/UEBuildTarget.cs
@@ -2476,6 +2476,12 @@ namespace UnrealBuildTool
 				}
 			}
 
+			// For Rocket builds, remove plugin that's designed for a NDA console
+			if(UnrealBuildTool.BuildingRocket())
+			{
+				ValidPlugins.RemoveAll(x => x.Directory.IndexOf("\\PS4\\", StringComparison.InvariantCultureIgnoreCase) != -1 || x.Directory.IndexOf("\\XboxOne\\", StringComparison.InvariantCultureIgnoreCase) != -1);
+			}
+
 			// Build a list of enabled plugins
 			List<string> EnabledPluginNames = new List<string>(Rules.AdditionalPlugins);
 
@@ -2502,7 +2508,7 @@ namespace UnrealBuildTool
 			// Set the list of plugins that should be built
 			if (UnrealBuildTool.BuildingRocket() && TargetType != TargetRules.TargetType.Program)
 			{
-				BuildPlugins.AddRange(ValidPlugins.Where(x => x.Directory.IndexOf("\\PS4\\", StringComparison.InvariantCultureIgnoreCase) == -1 && x.Directory.IndexOf("\\XboxOne\\", StringComparison.InvariantCultureIgnoreCase) == -1));
+				BuildPlugins.AddRange(ValidPlugins);
 			}
 			else if (ShouldCompileMonolithic() || TargetType == TargetRules.TargetType.Program)
 			{
