Minor dependency performance improvements to UnrealBuildTool
- Faster include dependency caching (ToLowerInvariant() is faster than StringComparer.InvariantCultureIgnoreCase on all platforms)
- Cached off the list of FileInfos when requesting cached include dependencies (reduces string hash lookup spam)
- Moved an object lookup out of an inner loop

diff --git a/Engine/Source/Programs/UnrealBuildTool/System/CPPHeaders.cs b/Engine/Source/Programs/UnrealBuildTool/System/CPPHeaders.cs
--- a/Engine/Source/Programs/UnrealBuildTool/System/CPPHeaders.cs
+++ b/Engine/Source/Programs/UnrealBuildTool/System/CPPHeaders.cs
@@ -175,7 +175,7 @@ namespace UnrealBuildTool
 
 		public static List<FileItem> FindAndCacheAllIncludedFiles( UEBuildTarget Target, FileItem SourceFile, IUEBuildPlatform BuildPlatform, CPPIncludeInfo CPPIncludeInfo, bool bOnlyCachedDependencies )
 		{
-			var Result = new List<FileItem>();
+			List<FileItem> Result = null;
 			
 			if( CPPIncludeInfo.IncludeFileSearchDictionary == null )
 			{
@@ -188,15 +188,8 @@ namespace UnrealBuildTool
 
 			if( bUseFlatCPPIncludeDependencyCache && bOnlyCachedDependencies )
 			{ 
-				var Dependencies = FlatCPPIncludeDependencyCache[Target].GetDependenciesForFile( SourceFile.AbsolutePath );
-				if( Dependencies != null )
-				{ 
-					foreach( string Dependency in Dependencies )
-					{
-						Result.Add( FileItem.GetItemByFullPath( Dependency ) );	// @todo fastubt: Make sure this is as fast as possible (convert to FileItem)
-					}
-				}
-				else
+				Result = FlatCPPIncludeDependencyCache[Target].GetDependenciesForFile( SourceFile.AbsolutePath );
+				if( Result == null )
 				{
 					// Nothing cached for this file!  It is new to us.  This is the expected flow when our CPPIncludeDepencencyCache is missing.
 				}
@@ -205,6 +198,8 @@ namespace UnrealBuildTool
 			{
 				// @todo fastubt: HeaderParser.h is missing from the include set for Module.UnrealHeaderTool.cpp (failed to find include using:  FileItem DirectIncludeResolvedFile = CPPEnvironment.FindIncludedFile(DirectInclude.IncludeName, !BuildConfiguration.bCheckExternalHeadersForModification, IncludePathsToSearch, IncludeFileSearchDictionary );)
 
+				Result = new List<FileItem>();
+
 				var IncludedFileList = new IncludedFilesSet();
 				CPPEnvironment.FindAndCacheAllIncludedFiles( Target, SourceFile, BuildPlatform, CPPIncludeInfo, ref IncludedFileList, bOnlyCachedDependencies:bOnlyCachedDependencies );
 				foreach( FileItem IncludedFile in IncludedFileList )
diff --git a/Engine/Source/Programs/UnrealBuildTool/System/DependencyCache.cs b/Engine/Source/Programs/UnrealBuildTool/System/DependencyCache.cs
--- a/Engine/Source/Programs/UnrealBuildTool/System/DependencyCache.cs
+++ b/Engine/Source/Programs/UnrealBuildTool/System/DependencyCache.cs
@@ -157,7 +157,7 @@ namespace UnrealBuildTool
 			CacheCreateDate = DateTimeOffset.Now;
 			CacheUpdateDate = DateTimeOffset.Now;
 			CachePath = Cache.AbsolutePath;
-			DependencyMap = new Dictionary<string, DependencyInfo>(StringComparer.InvariantCultureIgnoreCase);
+			DependencyMap = new Dictionary<string, DependencyInfo>();
 			bIsDirty = false;
 			CreateFileExistsInfo();
 		}
@@ -216,7 +216,7 @@ namespace UnrealBuildTool
 		{
 			// Check whether File is in cache.
 			DependencyInfo DependencyInfo;
-			if (DependencyMap.TryGetValue(File.AbsolutePath, out DependencyInfo))
+			if (DependencyMap.TryGetValue(File.AbsolutePath.ToLowerInvariant(), out DependencyInfo))
 			{
 				// File is in cache, now check whether last write time is prior to cache creation time.
 				if (File.LastWriteTime < CacheCreateDate)
@@ -293,7 +293,7 @@ namespace UnrealBuildTool
 		/// <param name="AbsolutePath"></param>
 		private void RemoveFileFromCache(string AbsolutePath)
 		{
-			DependencyMap.Remove(AbsolutePath);
+			DependencyMap.Remove(AbsolutePath.ToLowerInvariant());
 			bIsDirty = true;
 		}
 
@@ -306,7 +306,7 @@ namespace UnrealBuildTool
 		 */
 		public void SetDependencyInfo(FileItem File, List<DependencyInclude> DirectDependencies, bool HasUObjects)
 		{
-			DependencyMap[File.AbsolutePath] = new DependencyInfo() { Includes = DirectDependencies, HasUObjects = HasUObjects };
+			DependencyMap[File.AbsolutePath.ToLowerInvariant()] = new DependencyInfo() { Includes = DirectDependencies, HasUObjects = HasUObjects };
 			bIsDirty = true;
 		}
 
@@ -349,7 +349,7 @@ namespace UnrealBuildTool
 		{
 			if (BuildConfiguration.bUseIncludeDependencyResolveCache)
 			{
-				var Includes = DependencyMap[File.AbsolutePath].Includes;
+				var Includes = DependencyMap[File.AbsolutePath.ToLowerInvariant()].Includes;
 				var IncludeToResolve = Includes[DirectlyIncludedFileNameIndex];
 				if (BuildConfiguration.bTestIncludeDependencyResolveCache)
 				{
diff --git a/Engine/Source/Programs/UnrealBuildTool/System/FlatCPPIncludeDepencencyCache.cs b/Engine/Source/Programs/UnrealBuildTool/System/FlatCPPIncludeDepencencyCache.cs
--- a/Engine/Source/Programs/UnrealBuildTool/System/FlatCPPIncludeDepencencyCache.cs
+++ b/Engine/Source/Programs/UnrealBuildTool/System/FlatCPPIncludeDepencencyCache.cs
@@ -13,13 +13,18 @@ namespace UnrealBuildTool
 	/// Stores an ordered list of header files
 	/// </summary>
 	[Serializable]
-	struct FlatCPPIncludeDependencyInfo
+	class FlatCPPIncludeDependencyInfo
 	{
 		/// The PCH header this file is dependent on.
 		public string PCHName;
 
 		/// List of files this file includes, excluding headers that were included from a PCH.
 		public List<string> Includes;
+
+		/// Transient cached list of FileItems for all of the includes for a specific files.  This just saves a bunch of string
+		/// hash lookups as we locate FileItems for files that we've already requested dependencies for
+		[NonSerialized]
+		public List<FileItem> IncludeFileItems;
 	}
 
 
@@ -105,7 +110,7 @@ namespace UnrealBuildTool
 		protected FlatCPPIncludeDependencyCache(FileItem Cache)
 		{
 			CacheFileItem = Cache;
-			DependencyMap = new Dictionary<string, FlatCPPIncludeDependencyInfo>(StringComparer.InvariantCultureIgnoreCase);	// @todo fastubt: Apparently StringComparer is not very fast on Mono, better to use ToLowerInvariant(), even though it is an extra string copy
+			DependencyMap = new Dictionary<string, FlatCPPIncludeDependencyInfo>();
 		}
 
 
@@ -158,12 +163,13 @@ namespace UnrealBuildTool
 		/// <param name="Dependencies">List of source dependencies</param>
 		public void SetDependenciesForFile( string AbsoluteFilePath, string PCHName, List<string> Dependencies )
 		{
-			FlatCPPIncludeDependencyInfo DependencyInfo;
+			var DependencyInfo = new FlatCPPIncludeDependencyInfo();
 			DependencyInfo.PCHName = PCHName;	// @todo fastubt: Not actually used yet.  The idea is to use this to reduce the number of indirect includes we need to store in the cache.
 			DependencyInfo.Includes = Dependencies;
+			DependencyInfo.IncludeFileItems = null;
 
 			// @todo fastubt: We could shrink this file by not storing absolute paths (project and engine relative paths only, except for system headers.)  May affect load times.
-			DependencyMap[ AbsoluteFilePath ] = DependencyInfo;
+			DependencyMap[ AbsoluteFilePath.ToLowerInvariant() ] = DependencyInfo;
 			bIsDirty = true;
 		}
 
@@ -173,12 +179,21 @@ namespace UnrealBuildTool
 		/// </summary>
 		/// <param name="AbsoluteFilePath">Path to the file</param>
 		/// <returns>The list of includes</returns>
-		public List<string> GetDependenciesForFile( string AbsoluteFilePath )
+		public List<FileItem> GetDependenciesForFile( string AbsoluteFilePath )
 		{
 			FlatCPPIncludeDependencyInfo DependencyInfo;
-			if( DependencyMap.TryGetValue( AbsoluteFilePath, out DependencyInfo ) )
+			if( DependencyMap.TryGetValue( AbsoluteFilePath.ToLowerInvariant(), out DependencyInfo ) )
 			{
-				return DependencyInfo.Includes;
+				// Update our transient cache of FileItems for each of the included files
+				if( DependencyInfo.IncludeFileItems == null )
+				{
+					DependencyInfo.IncludeFileItems = new List<FileItem>( DependencyInfo.Includes.Count );
+					foreach( string Dependency in DependencyInfo.Includes )
+					{
+						DependencyInfo.IncludeFileItems.Add( FileItem.GetItemByFullPath( Dependency ) );	// @todo fastubt: Make sure this is as fast as possible (convert to FileItem)
+					}
+				}
+				return DependencyInfo.IncludeFileItems;
 			}
 
 			return null;
