UnrealBuildTool: Cleaning up UBT Makefile support in preparation for wide testing
- Collapsed and renamed some of the new UBT settings for fast build iteration
- New 'BuildConfiguration.bUseUBTMakefiles' that replaces the old experimental settings (defaults to off in this changelist.)
- Cleaned up some comments
- No effective changes

diff --git a/Engine/Source/Programs/UnrealBuildTool/Configuration/BuildConfiguration.cs b/Engine/Source/Programs/UnrealBuildTool/Configuration/BuildConfiguration.cs
--- a/Engine/Source/Programs/UnrealBuildTool/Configuration/BuildConfiguration.cs
+++ b/Engine/Source/Programs/UnrealBuildTool/Configuration/BuildConfiguration.cs
@@ -359,30 +359,28 @@ namespace UnrealBuildTool
 		public static bool bBreakBuildOnLicenseViolation;
 
 		/// <summary>
-		/// True if we should use the new 'experimental' C++ header dependency scanning and caching system that allows Unreal Build Tool to
-		/// detect out of date dependencies very quickly.  When enabled, a deep C++ include graph does not have to be generated, and instead
-		/// we only scan and cache indirect includes for after a dependent build product was already found to be out of date.  During the
-		/// next build, we'll load those cached indirect includes and check for outdatedness.
-		/// </summary>
-		[XmlConfig]
-		public static bool bUseExperimentalFastDependencyScan;
-
-		/// <summary>
-		/// Enables a new 'experimental' support for very fast iterative builds.  Turning this on causes Unreal Build Tool to emit
+		/// Enables support for very fast iterative builds by caching target data.  Turning this on causes Unreal Build Tool to emit
 		/// 'UBT Makefiles' for targets when they're built the first time.  Subsequent builds will load these Makefiles and begin
 		/// outdatedness checking and build invocation very quickly.  The caveat is that if source files are added or removed to
 		/// the project, UBT will need to gather information about those in order for your build to complete successfully.  Currently,
 		/// you must run the project file generator after adding/removing source files to tell UBT to re-gather this information.
+		/// 
 		/// Events that can invalidate the 'UBT Makefile':  
 		///		- Adding/removing .cpp files
 		///		- Adding/removing .h files with UObjects
 		///		- Adding new UObject types to a file that didn't previously have any
 		///		- Changing global build settings (most settings in this file qualify.)
 		///		- Changed code that affects how Unreal Header Tool works
+		///	
 		///	You can force regeneration of the 'UBT Makefile' by passing the '-gather' argument, or simply regenerating project files
+		///
+		///	This also enables the fast include file dependency scanning and caching system that allows Unreal Build Tool to detect out 
+		/// of date dependencies very quickly.  When enabled, a deep C++ include graph does not have to be generated, and instead
+		/// we only scan and cache indirect includes for after a dependent build product was already found to be out of date.  During the
+		/// next build, we'll load those cached indirect includes and check for outdatedness.
 		/// </summary>
 		[XmlConfig]
-		public static bool bUseExperimentalFastBuildIteration;
+		public static bool bUseUBTMakefiles;
 	
 		/// <summary>
 		/// Whether DMUCS/Distcc may be used.
@@ -528,13 +526,10 @@ namespace UnrealBuildTool
 			bCheckLicenseViolations = false;
 			bBreakBuildOnLicenseViolation = true;
 
-			// This feature reduces the time spent tracking C++ header dependencies every run.  This feature is still being tested, 
-			// so we disable it by default
-			bUseExperimentalFastDependencyScan = false;
-
-			// Experimental support for 'gathering' and 'assembling' builds in separate invocations is still being tested, so we
-			// leave it disabled by default
-			bUseExperimentalFastBuildIteration = false;
+			// Enables support for fast include dependency scanning, as well as gathering data for 'UBT Makefiles', then quickly
+			// assembling builds in subsequent runs using data in those cached makefiles
+			// NOTE: This feature is new and has a number of known issues (search the code for '@todo ubtmake')
+			bUseUBTMakefiles = false;
 
 			// Distcc requires some setup - so by default disable it so we don't break local or remote building
 			bAllowDistcc = false;
diff --git a/Engine/Source/Programs/UnrealBuildTool/Configuration/UEBuildModule.cs b/Engine/Source/Programs/UnrealBuildTool/Configuration/UEBuildModule.cs
--- a/Engine/Source/Programs/UnrealBuildTool/Configuration/UEBuildModule.cs
+++ b/Engine/Source/Programs/UnrealBuildTool/Configuration/UEBuildModule.cs
@@ -955,9 +955,6 @@ namespace UnrealBuildTool
 		/** A list of the absolute paths of source files to be built in this module. */
 		public readonly SourceFilesClass SourceFilesToBuild = new SourceFilesClass();
 
-		/** A list of the source files that were found for the module. */
-		public readonly SourceFilesClass SourceFilesFound = new SourceFilesClass();
-
 		/** The preprocessor definitions used to compile this module's private implementation. */
 		HashSet<string> Definitions;
 
@@ -1113,7 +1110,6 @@ namespace UnrealBuildTool
 		{
 			IntelliSenseGatherer = InIntelliSenseGatherer;
 
-			CategorizeSourceFiles(InSourceFiles, SourceFilesFound);
 			if (bInBuildSourceFiles)
 			{
 				CategorizeSourceFiles(InSourceFiles, SourceFilesToBuild);
@@ -1514,7 +1510,7 @@ namespace UnrealBuildTool
 
 					CachePCHUsageForModuleSourceFile(this.Target, CPPCompileEnvironment, GeneratedCppFileItem);
 
-					// @todo fastubt: Check for ALL other places where we might be injecting .cpp or .rc files for compiling without caching CachedCPPIncludeInfo first (anything platform specfic?)
+					// @todo ubtmake: Check for ALL other places where we might be injecting .cpp or .rc files for compiling without caching CachedCPPIncludeInfo first (anything platform specific?)
 					LinkInputFiles.AddRange(CPPCompileEnvironment.CompileFiles(Target, new List<FileItem> { GeneratedCppFileItem }, Name).ObjectFiles);
 				}
 			}
@@ -1642,10 +1638,10 @@ namespace UnrealBuildTool
 				bool bFoundAProblemWithPCHs = false;
 
 				FileItem UniquePCH = null;
-				foreach( var CPPFile in SourceFilesFound.CPPFiles )	// @todo fastubt: We're not caching CPPEnvironments for .c/.mm files, etc.  Even though they don't use PCHs, they still have #includes!  This can break dependency checking!
+				foreach( var CPPFile in SourceFilesToBuild.CPPFiles )	// @todo ubtmake: We're not caching CPPEnvironments for .c/.mm files, etc.  Even though they don't use PCHs, they still have #includes!  This can break dependency checking!
 				{
 					// Build a single list of include paths to search.
-					var IncludePathsToSearch = ModuleCompileEnvironment.Config.CPPIncludeInfo.GetIncludesPathsToSearch( CPPFile );   // @todo fastubt: This call could be made faster (list caching)
+					var IncludePathsToSearch = ModuleCompileEnvironment.Config.CPPIncludeInfo.GetIncludesPathsToSearch( CPPFile );
 
 					// Store the module compile environment along with the .cpp file.  This is so that we can use it later on when looking
 					// for header dependencies
@@ -1662,7 +1658,7 @@ namespace UnrealBuildTool
 					{
 						UniquePCH = PCH;
 					}
-					else if( !UniquePCH.Info.Name.Equals( PCH.Info.Name, StringComparison.InvariantCultureIgnoreCase ) )		// @todo fastubt: We do a string compare on the file name (not path) here, because sometimes the include resolver will pick an Intermediate copy of a PCH header file and throw off our comparisons
+					else if( !UniquePCH.Info.Name.Equals( PCH.Info.Name, StringComparison.InvariantCultureIgnoreCase ) )		// @todo ubtmake: We do a string compare on the file name (not path) here, because sometimes the include resolver will pick an Intermediate copy of a PCH header file and throw off our comparisons
 					{
 						// OK, looks like we have multiple source files including a different header file first.  We'll keep track of this and print out
 						// helpful information afterwards.
@@ -1677,7 +1673,7 @@ namespace UnrealBuildTool
 				{
 					// Map from pch header string to the source files that use that PCH
 					var UsageMapPCH = new Dictionary<string, List<FileItem>>( StringComparer.InvariantCultureIgnoreCase );
-					foreach( var CPPFile in SourceFilesToBuild.CPPFiles )	// @todo fastubt: Using SourceFilesToBuild here, but SourceFilesFound up above.  Why are there two arrays again?
+					foreach( var CPPFile in SourceFilesToBuild.CPPFiles )
 					{
 						// Create a new entry if not in the pch usage map
 						UsageMapPCH.GetOrAddNew( CPPFile.PrecompiledHeaderIncludeFilename ).Add( CPPFile );
@@ -1740,8 +1736,8 @@ namespace UnrealBuildTool
 
 		private static FileItem CachePCHUsageForCPPFile(UEBuildTarget Target, FileItem CPPFile, IUEBuildPlatform BuildPlatform, List<string> IncludePathsToSearch, Dictionary<string, FileItem> IncludeFileSearchDictionary)
 		{
-			// @todo fastubt: We don't really need to scan every file looking for PCH headers, just need one.  The rest is just for error checking.
-			// @todo fastubt: We don't need all of the direct includes either.  We just need the first, unless we want to check for errors.
+			// @todo ubtmake: We don't really need to scan every file looking for PCH headers, just need one.  The rest is just for error checking.
+			// @todo ubtmake: We don't need all of the direct includes either.  We just need the first, unless we want to check for errors.
 			bool HasUObjects;
 			List<DependencyInclude> DirectIncludeFilenames = CPPEnvironment.GetDirectIncludeDependencies(Target, CPPFile, BuildPlatform, bOnlyCachedDependencies:false, HasUObjects:out HasUObjects);
 			if (BuildConfiguration.bPrintDebugInfo)
@@ -2125,7 +2121,7 @@ namespace UnrealBuildTool
 			}
 			else
 			{
-				var AllProjectFolders = UEBuildTarget.DiscoverAllGameFolders();		// @todo fastubt: This will be called again and again for every UObject module (80+)
+				var AllProjectFolders = UEBuildTarget.DiscoverAllGameFolders();		// @todo ubtmake: This will be called again and again for every UObject module (80+)
 				BaseDirectory = AllProjectFolders.Find(ProjectFolder => Utils.IsFileUnderDirectory( ModuleDirectory, ProjectFolder ));
 				if (BaseDirectory == null)
 				{
diff --git a/Engine/Source/Programs/UnrealBuildTool/System/ActionGraph.cs b/Engine/Source/Programs/UnrealBuildTool/System/ActionGraph.cs
--- a/Engine/Source/Programs/UnrealBuildTool/System/ActionGraph.cs
+++ b/Engine/Source/Programs/UnrealBuildTool/System/ActionGraph.cs
@@ -202,8 +202,6 @@ namespace UnrealBuildTool
 
 		public static void FinalizeActionGraph()
 		{
-			// @todo fastubt: Can we use directory changed notifications or directory timestamps to accelerate C++ file outdatedness checking?
-			
 			// Link producing actions to the items they produce.
 			LinkActionsAndItems();
 
@@ -791,7 +789,7 @@ namespace UnrealBuildTool
 					// includes before getting this far.  However, if we find them to be outdated after processing includes, we'll do a deep scan later
 					// on and cache all of the includes so that we have them for a quick outdatedness check the next run.
 					if( !bIsOutdated && 
-						BuildConfiguration.bUseExperimentalFastBuildIteration && 
+						BuildConfiguration.bUseUBTMakefiles && 
 						UnrealBuildTool.IsAssemblingBuild && 
 						RootAction.ActionType == ActionType.Compile )
 					{
@@ -800,7 +798,7 @@ namespace UnrealBuildTool
 
 					// Were we asked to force an update of our cached includes BEFORE we try to build?  This may be needed if our cache can no longer
 					// be trusted and we need to fill it with perfectly valid data (even if we're in assembler only mode)
-					if( BuildConfiguration.bUseExperimentalFastDependencyScan && 
+					if( BuildConfiguration.bUseUBTMakefiles && 
 						UnrealBuildTool.bNeedsFullCPPIncludeRescan )
 					{
 						bFindCPPIncludePrerequisites = true;
@@ -819,8 +817,8 @@ namespace UnrealBuildTool
 						// @todo ubtmake: Make sure we are catching RC files here too.  Anything that the toolchain would have tried it on.  Logic should match the CACHING stuff below
 						if( PrerequisiteItem.CachedCPPIncludeInfo != null )
 						{
-							var IncludedFileList = CPPEnvironment.FindAndCacheAllIncludedFiles( Target, PrerequisiteItem, BuildPlatform, PrerequisiteItem.CachedCPPIncludeInfo, bOnlyCachedDependencies:BuildConfiguration.bUseExperimentalFastDependencyScan );
-							foreach( var IncludedFile in IncludedFileList )	// @todo fastubt: @todo ubtmake: Optimization: This is "retesting" a lot of the same files over and over in a single run (common indirect includes)
+							var IncludedFileList = CPPEnvironment.FindAndCacheAllIncludedFiles( Target, PrerequisiteItem, BuildPlatform, PrerequisiteItem.CachedCPPIncludeInfo, bOnlyCachedDependencies:BuildConfiguration.bUseUBTMakefiles );
+							foreach( var IncludedFile in IncludedFileList )
 							{
 								if( IncludedFile.bExists )
 								{
@@ -915,12 +913,12 @@ namespace UnrealBuildTool
 				// know about the set of header files that are included for files that are already determined to be out of date (such as if the file
 				// is missing or was modified.)  In the case that the file is out of date, we'll perform a deep scan to update our cached set of
 				// includes for this file, so that we'll be able to determine whether it is out of date next time very quickly.
-				if( BuildConfiguration.bUseExperimentalFastDependencyScan )
+				if( BuildConfiguration.bUseUBTMakefiles )
 				{ 
 					var DeepIncludeScanStartTime = DateTime.UtcNow;
 
-					// @todo fastubt: we may be scanning more files than we need to here -- indirectly outdated files are bIsOutdated=true by this point (for example basemost includes when deeper includes are dirty)
-					if( bIsOutdated && RootAction.ActionType == ActionType.Compile )	// @todo fastubt: Does this work with RC files?  See above too.
+					// @todo ubtmake: we may be scanning more files than we need to here -- indirectly outdated files are bIsOutdated=true by this point (for example basemost includes when deeper includes are dirty)
+					if( bIsOutdated && RootAction.ActionType == ActionType.Compile )	// @todo ubtmake: Does this work with RC files?  See above too.
 					{
 						Log.TraceVerbose( "Outdated action: {0}", RootAction.StatusDescription );
 						foreach (FileItem PrerequisiteItem in RootAction.PrerequisiteItems)
diff --git a/Engine/Source/Programs/UnrealBuildTool/System/CPPHeaders.cs b/Engine/Source/Programs/UnrealBuildTool/System/CPPHeaders.cs
--- a/Engine/Source/Programs/UnrealBuildTool/System/CPPHeaders.cs
+++ b/Engine/Source/Programs/UnrealBuildTool/System/CPPHeaders.cs
@@ -152,7 +152,7 @@ namespace UnrealBuildTool
 				IncludeFileSearchDictionary.Add( InvariantPath, Result );
 			}
 
-			// @todo fastubt: The old UBT tried to skip 'external' (STABLE) headers here.  But it didn't work.  We might want to do this though!  Skip system headers and source/thirdparty headers!
+			// @todo ubtmake: The old UBT tried to skip 'external' (STABLE) headers here.  But it didn't work.  We might want to do this though!  Skip system headers and source/thirdparty headers!
 
 			if( Result != null )
 			{
@@ -182,9 +182,9 @@ namespace UnrealBuildTool
 				CPPIncludeInfo.IncludeFileSearchDictionary = new Dictionary<string,FileItem>();
 			}
 
-			bool bUseFlatCPPIncludeDependencyCache =
-				( BuildConfiguration.bUseExperimentalFastDependencyScan &&
-				  ( !BuildConfiguration.bUseExperimentalFastBuildIteration || UnrealBuildTool.IsAssemblingBuild ) );
+			bool bUseFlatCPPIncludeDependencyCache = 
+				( BuildConfiguration.bUseUBTMakefiles &&
+				  ( !BuildConfiguration.bUseUBTMakefiles || UnrealBuildTool.IsAssemblingBuild ) );
 
 			if( bUseFlatCPPIncludeDependencyCache && bOnlyCachedDependencies )
 			{ 
@@ -196,7 +196,7 @@ namespace UnrealBuildTool
 			}
 			else
 			{
-				// @todo fastubt: HeaderParser.h is missing from the include set for Module.UnrealHeaderTool.cpp (failed to find include using:  FileItem DirectIncludeResolvedFile = CPPEnvironment.FindIncludedFile(DirectInclude.IncludeName, !BuildConfiguration.bCheckExternalHeadersForModification, IncludePathsToSearch, IncludeFileSearchDictionary );)
+				// @todo ubtmake: HeaderParser.h is missing from the include set for Module.UnrealHeaderTool.cpp (failed to find include using:  FileItem DirectIncludeResolvedFile = CPPEnvironment.FindIncludedFile(DirectInclude.IncludeName, !BuildConfiguration.bCheckExternalHeadersForModification, IncludePathsToSearch, IncludeFileSearchDictionary );)
 
 				Result = new List<FileItem>();
 
diff --git a/Engine/Source/Programs/UnrealBuildTool/System/ExternalExecution.cs b/Engine/Source/Programs/UnrealBuildTool/System/ExternalExecution.cs
--- a/Engine/Source/Programs/UnrealBuildTool/System/ExternalExecution.cs
+++ b/Engine/Source/Programs/UnrealBuildTool/System/ExternalExecution.cs
@@ -389,7 +389,7 @@ namespace UnrealBuildTool
 								// of the source headers.  We don't care if source files were added or removed in this mode, because we're only
 								// able to process the known UObject headers that are in the Makefile.  If UObject header files are added/removed,
 								// we expect the user to re-run GenerateProjectFiles which will force UBTMakefile outdatedness.
-								// @todo fastubt: Possibly, we should never be doing this check these days.
+								// @todo ubtmake: Possibly, we should never be doing this check these days.
 								if( UnrealBuildTool.IsGatheringBuild || !UnrealBuildTool.IsAssemblingBuild )
 								{
 									// Also check the timestamp on the directory the source file is in.  If the directory timestamp has
@@ -570,7 +570,7 @@ namespace UnrealBuildTool
 					}
 				}
 
-				// @todo fastubt: @todo ubtmake: Optimization: Ideally we could avoid having to generate this data in the case where UHT doesn't even need to run!  Can't we use the existing copy?  (see below use of Manifest)
+				// @todo ubtmake: Optimization: Ideally we could avoid having to generate this data in the case where UHT doesn't even need to run!  Can't we use the existing copy?  (see below use of Manifest)
 				UHTManifest Manifest = new UHTManifest(Target, RootLocalPath, ToolChain.ConvertPath(RootLocalPath + '\\'), UObjectModules);
 
 				if( !bIsBuildingUHT && bUHTNeedsToRun )
diff --git a/Engine/Source/Programs/UnrealBuildTool/System/FileItem.cs b/Engine/Source/Programs/UnrealBuildTool/System/FileItem.cs
--- a/Engine/Source/Programs/UnrealBuildTool/System/FileItem.cs
+++ b/Engine/Source/Programs/UnrealBuildTool/System/FileItem.cs
@@ -48,7 +48,7 @@ namespace UnrealBuildTool
 					// Uh oh.  We're clobbering our cached CompileEnvironment for this file with a different CompileEnvironment.  This means
 					// that the same source file is being compiled into more than one module. (e.g. PCLaunch.rc)
 
-					// @todo fastubt: The only expected offender here is PCLaunch.rc and friends, which are injected by UBT into every module when not compiling monolithic.
+					// @todo ubtmake: The only expected offender here is PCLaunch.rc and friends, which are injected by UBT into every module when not compiling monolithic.
 					// PCLaunch.rc and ModuleVersionResource.rc.inl are "safe" because they do not include any headers that would be affected by include path order.
 					// ==> Ideally we would use a different "shared" CompileEnvironment for these injected .rc files, so their include paths would not change
 					// ==> OR, we can make an Intermediate copy of the .rc file for each module (easier)
@@ -56,7 +56,7 @@ namespace UnrealBuildTool
 						!AbsolutePath.EndsWith( "ModuleVersionResource.rc.inl", StringComparison.InvariantCultureIgnoreCase ) )
 					{ 					
 						// Let's make sure the include paths are the same
-						// @todo fastubt: We have not seen examples of this actually firing off, so we could probably remove the check for matching includes and simply always make this an error case
+						// @todo ubtmake: We have not seen examples of this actually firing off, so we could probably remove the check for matching includes and simply always make this an error case
 						var CachedIncludePathsToSearch = _CachedCPPIncludeInfo.GetIncludesPathsToSearch( this );
 						var NewIncludePathsToSearch = value.GetIncludesPathsToSearch( this );
 
diff --git a/Engine/Source/Programs/UnrealBuildTool/System/FlatCPPIncludeDepencencyCache.cs b/Engine/Source/Programs/UnrealBuildTool/System/FlatCPPIncludeDepencencyCache.cs
--- a/Engine/Source/Programs/UnrealBuildTool/System/FlatCPPIncludeDepencencyCache.cs
+++ b/Engine/Source/Programs/UnrealBuildTool/System/FlatCPPIncludeDepencencyCache.cs
@@ -83,8 +83,8 @@ namespace UnrealBuildTool
 			{
 				using (FileStream Stream = new FileStream(Cache.AbsolutePath, FileMode.Open, FileAccess.Read))
 				{	
-					// @todo fastubt: We can store the cache in a cheaper/smaller way using hash file names and indices into included headers, but it might actually slow down load times
-					// @todo fastubt: If we can index PCHs here, we can avoid storing all of the PCH's included headers (PCH's action should have been invalidated, so we shouldn't even have to report the PCH's includes as our indirect includes)
+					// @todo ubtmake: We can store the cache in a cheaper/smaller way using hash file names and indices into included headers, but it might actually slow down load times
+					// @todo ubtmake: If we can index PCHs here, we can avoid storing all of the PCH's included headers (PCH's action should have been invalidated, so we shouldn't even have to report the PCH's includes as our indirect includes)
 					BinaryFormatter Formatter = new BinaryFormatter();
 					Result = Formatter.Deserialize(Stream) as FlatCPPIncludeDependencyCache;
 					Result.CacheFileItem = Cache;
@@ -164,11 +164,11 @@ namespace UnrealBuildTool
 		public void SetDependenciesForFile( string AbsoluteFilePath, string PCHName, List<string> Dependencies )
 		{
 			var DependencyInfo = new FlatCPPIncludeDependencyInfo();
-			DependencyInfo.PCHName = PCHName;	// @todo fastubt: Not actually used yet.  The idea is to use this to reduce the number of indirect includes we need to store in the cache.
+			DependencyInfo.PCHName = PCHName;	// @todo ubtmake: Not actually used yet.  The idea is to use this to reduce the number of indirect includes we need to store in the cache.
 			DependencyInfo.Includes = Dependencies;
 			DependencyInfo.IncludeFileItems = null;
 
-			// @todo fastubt: We could shrink this file by not storing absolute paths (project and engine relative paths only, except for system headers.)  May affect load times.
+			// @todo ubtmake: We could shrink this file by not storing absolute paths (project and engine relative paths only, except for system headers.)  May affect load times.
 			DependencyMap[ AbsoluteFilePath.ToLowerInvariant() ] = DependencyInfo;
 			bIsDirty = true;
 		}
@@ -190,7 +190,7 @@ namespace UnrealBuildTool
 					DependencyInfo.IncludeFileItems = new List<FileItem>( DependencyInfo.Includes.Count );
 					foreach( string Dependency in DependencyInfo.Includes )
 					{
-						DependencyInfo.IncludeFileItems.Add( FileItem.GetItemByFullPath( Dependency ) );	// @todo fastubt: Make sure this is as fast as possible (convert to FileItem)
+						DependencyInfo.IncludeFileItems.Add( FileItem.GetItemByFullPath( Dependency ) );
 					}
 				}
 				return DependencyInfo.IncludeFileItems;
diff --git a/Engine/Source/Programs/UnrealBuildTool/System/Unity.cs b/Engine/Source/Programs/UnrealBuildTool/System/Unity.cs
--- a/Engine/Source/Programs/UnrealBuildTool/System/Unity.cs
+++ b/Engine/Source/Programs/UnrealBuildTool/System/Unity.cs
@@ -231,7 +231,7 @@ namespace UnrealBuildTool
                 UnityCPPFiles.Add(UnityCPPFile);
 
 				// Cache information about the unity .cpp dependencies
-				// @todo fastubt urgent: Fails when building remotely for Mac because unity .cpp has an include for a PCH on the REMOTE machine
+				// @todo ubtmake urgent: Fails when building remotely for Mac because unity .cpp has an include for a PCH on the REMOTE machine
 				UEBuildModuleCPP.CachePCHUsageForModuleSourceFile( Target, CompileEnvironment, UnityCPPFile );
 			}
 
diff --git a/Engine/Source/Programs/UnrealBuildTool/System/UnrealBuildTool.cs b/Engine/Source/Programs/UnrealBuildTool/System/UnrealBuildTool.cs
--- a/Engine/Source/Programs/UnrealBuildTool/System/UnrealBuildTool.cs
+++ b/Engine/Source/Programs/UnrealBuildTool/System/UnrealBuildTool.cs
@@ -108,7 +108,7 @@ namespace UnrealBuildTool
         }
         static private bool bIsAssemblingBuild_Unsafe = true;
 
-        /** Used when BuildConfiguration.bUseExperimentalFastDependencyScan is enabled.  If true, it means that our cached includes may not longer be
+        /** Used when BuildConfiguration.bUseUBTMakefiles is enabled.  If true, it means that our cached includes may not longer be
             valid (or never were), and we need to recover by forcibly scanning included headers for all build prerequisites to make sure that our
             cached set of includes is actually correct, before determining which files are outdated. */
         static public bool bNeedsFullCPPIncludeRescan = false;
@@ -789,7 +789,7 @@ namespace UnrealBuildTool
 
                     // If we were asked to enable fast build iteration, we want the 'gather' phase to default to off (unless it is overridden below
                     // using a command-line option.)
-                    if( BuildConfiguration.bUseExperimentalFastBuildIteration )
+                    if( BuildConfiguration.bUseUBTMakefiles )
                     {
                         bIsGatheringBuild_Unsafe = false;
                     }
@@ -1138,16 +1138,8 @@ namespace UnrealBuildTool
                         {
                             if (UEBuildConfiguration.bPrepForDeployment == false)
                             {
-                                int PassCount = 1;	// @todo ubtmake: It's useful to set this to a high number when running UBT under a sampling profiler
-                                for( int PassIndex = 0; PassIndex < PassCount; ++PassIndex )
-                                { 
-                                    // If we are only prepping for deployment, assume the build already occurred.
-                                    Result = RunUBT(Arguments);
-                                    if( Result != ECompilationResult.Succeeded )
-                                    {
-                                        break;
-                                    }
-                                }
+								// If we are only prepping for deployment, assume the build already occurred.
+								Result = RunUBT(Arguments);
                             }
                             else
                             {
@@ -1426,10 +1418,9 @@ namespace UnrealBuildTool
 				bool bIsHotReload = UEBuildConfiguration.bHotReloadFromIDE || ( TargetDescs.Count == 1 && TargetDescs[0].OnlyModules.Count > 0 );
 				TargetDescriptor HotReloadTargetDesc = bIsHotReload ? TargetDescs[0] : null;
 
-                // If UBT is being used to build something different than it did last time
                 if( !ProjectFileGenerator.bGenerateProjectFiles )
                 {
-					if( BuildConfiguration.bUseExperimentalFastDependencyScan )
+					if( BuildConfiguration.bUseUBTMakefiles )
                     {
                         // If we're building UHT without a Mutex, we'll need to assume that we're building the same targets and that no caches
                         // should be invalidated for this run.  This is important when UBT is invoked from within UBT in order to compile
@@ -1459,7 +1450,7 @@ namespace UnrealBuildTool
                             if( File.Exists( LastBuiltTargetsFilePath ) && Utils.ReadAllText( LastBuiltTargetsFilePath ) == TargetCollectionName )
                             {
 								// @todo ubtmake: Because we're using separate files for hot reload vs. full compiles, it's actually possible that includes will
-								// become out of date without us knowing if the developer ping pongs between hot reloading target A and bulding target B normally.
+								// become out of date without us knowing if the developer ping-pongs between hot reloading target A and building target B normally.
 								// To fix this we can not use a different file name for last built targets, but the downside is slower performance when
 								// performing the first hot reload after compiling normally (forces full include dependency scan)
                                 bIsBuildingSameTargetsAsLastTime = true;
@@ -1503,7 +1494,7 @@ namespace UnrealBuildTool
                     }
 
                     // Make sure the gather phase is executed if we're not actually building anything
-                    if( ProjectFileGenerator.bGenerateProjectFiles || UEBuildConfiguration.bGenerateManifest || UEBuildConfiguration.bCleanProject )
+                    if( ProjectFileGenerator.bGenerateProjectFiles || UEBuildConfiguration.bGenerateManifest || UEBuildConfiguration.bCleanProject || BuildConfiguration.bXGEExport || UEBuildConfiguration.bGenerateExternalFileList || GeneratingActionGraph)
                     {
                         UnrealBuildTool.bIsGatheringBuild_Unsafe = true;						
                     }
@@ -1589,11 +1580,11 @@ namespace UnrealBuildTool
                     }
 
                     // We don't need this dependency cache in 'gather only' mode
-                    if( BuildConfiguration.bUseExperimentalFastDependencyScan &&
+                    if( BuildConfiguration.bUseUBTMakefiles &&
                         !( UnrealBuildTool.IsGatheringBuild && !UnrealBuildTool.IsAssemblingBuild ) )
                     { 
                         // Load the cache that contains the list of flattened resolved includes for resolved source files
-                        // @todo fastubt: Ideally load this asynchronously at startup and only block when it is first needed and not finished loading
+                        // @todo ubtmake: Ideally load this asynchronously at startup and only block when it is first needed and not finished loading
                         CPPEnvironment.FlatCPPIncludeDependencyCache.Add( Target, FlatCPPIncludeDependencyCache.Create( Target ) );
                     }
 
@@ -1679,7 +1670,7 @@ namespace UnrealBuildTool
                         UBTMakefile.TargetNameToUObjectModules = TargetNameToUObjectModules;
 						UBTMakefile.Targets = Targets;
 
-						if( BuildConfiguration.bUseExperimentalFastBuildIteration )
+						if( BuildConfiguration.bUseUBTMakefiles )
 						{ 
 							// We've been told to prepare to build, so let's go ahead and save out our action graph so that we can use in a later invocation 
 							// to assemble the build.  Even if we are configured to assemble the build in this same invocation, we want to save out the
@@ -1754,7 +1745,7 @@ namespace UnrealBuildTool
                             // Cache indirect includes for all outdated C++ files.  We kick this off as a background thread so that it can
                             // perform the scan while we're compiling.  It usually only takes up to a few seconds, but we don't want to hurt
                             // our best case UBT iteration times for this task which can easily be performed asynchronously
-                            if( BuildConfiguration.bUseExperimentalFastDependencyScan && TargetToOutdatedPrerequisitesMap.Count > 0 )
+                            if( BuildConfiguration.bUseUBTMakefiles && TargetToOutdatedPrerequisitesMap.Count > 0 )
                             {
                                 CPPIncludesThread = CreateThreadForCachingCPPIncludes( TargetToOutdatedPrerequisitesMap );
                                 CPPIncludesThread.Start();
@@ -1889,7 +1880,7 @@ namespace UnrealBuildTool
 			bool bIsRunning = false;
 			
 			// @todo ubtmake: Kind of cheating here to figure out if an editor target.  At this point we don't have access to the actual target description, and
-			// this code must be able to execute before we create or load module rules DLLs so that hot reload can work with bUseExperimentalFastBuildIteration
+			// this code must be able to execute before we create or load module rules DLLs so that hot reload can work with bUseUBTMakefiles
 			bool bIsEditorTarget = TargetDesc.TargetName.EndsWith( "Editor", StringComparison.InvariantCultureIgnoreCase ) || TargetDesc.bIsEditorRecompile;	
 
 			if (!ProjectFileGenerator.bGenerateProjectFiles && !UEBuildConfiguration.bGenerateManifest && bIsEditorTarget )
@@ -2088,7 +2079,7 @@ namespace UnrealBuildTool
         {
             return new Thread( new ThreadStart( () =>
                 {
-                    // @todo fastubt: This thread will access data structures that are also used on the main UBT thread, but during this time UBT
+                    // @todo ubtmake: This thread will access data structures that are also used on the main UBT thread, but during this time UBT
                     // is only invoking the build executor, so should not be touching this stuff.  However, we need to at some guards to make sure.
 
 
@@ -2134,7 +2125,7 @@ namespace UnrealBuildTool
             public Action[] PrerequisiteActions;			
 
             /** Environment variables that we'll need in order to invoke the platform's compiler and linker */
-			// @todo ubtmake: Really we want to allow a different set of environment variables for every Action.  This would allow for targets on multiple platforms to be built in a single assembling phase.  Really, we'd only have unique variables for each platform that has actions, so we'd want to make sure we only store the minimum set.
+			// @todo ubtmake: Really we want to allow a different set of environment variables for every Action.  This would allow for targets on multiple platforms to be built in a single assembling phase.  We'd only have unique variables for each platform that has actions, so we'd want to make sure we only store the minimum set.
             public readonly List<Tuple<string,string>> EnvironmentVariables = new List<Tuple<string,string>>();
 
             /** Maps each target to a list of UObject module info structures */
diff --git a/Engine/Source/Programs/UnrealBuildTool/System/Utils.cs b/Engine/Source/Programs/UnrealBuildTool/System/Utils.cs
--- a/Engine/Source/Programs/UnrealBuildTool/System/Utils.cs
+++ b/Engine/Source/Programs/UnrealBuildTool/System/Utils.cs
@@ -169,7 +169,7 @@ namespace UnrealBuildTool
 		 */
 		public static void SetEnvironmentVariablesFromBatchFile(string BatchFileName)
 		{
-			// @todo fastubt: Experiment with changing this to run asynchronously at startup, and only blocking if accessed before the .bat file finishes
+			// @todo ubtmake: Experiment with changing this to run asynchronously at startup, and only blocking if accessed before the .bat file finishes
 			if( File.Exists( BatchFileName ) )
 			{
 				// Create a wrapper batch file that echoes environment variables to a text file
diff --git a/Engine/Source/Programs/UnrealBuildTool/ToolChain/UEToolChain.cs b/Engine/Source/Programs/UnrealBuildTool/ToolChain/UEToolChain.cs
--- a/Engine/Source/Programs/UnrealBuildTool/ToolChain/UEToolChain.cs
+++ b/Engine/Source/Programs/UnrealBuildTool/ToolChain/UEToolChain.cs
@@ -270,19 +270,19 @@ namespace UnrealBuildTool
 				RemoteThis.QueueFileForBatchUpload(SourceFile);
 			}
 
-			if( !BuildConfiguration.bUseExperimentalFastBuildIteration )	// In fast build iteration mode, we'll gather includes later on
+			if( !BuildConfiguration.bUseUBTMakefiles )	// In fast build iteration mode, we'll gather includes later on
 			{
-				// @todo fastubt: What if one of the prerequisite files has become missing since it was updated in our cache? (usually, because a coder eliminated the source file)
+				// @todo ubtmake: What if one of the prerequisite files has become missing since it was updated in our cache? (usually, because a coder eliminated the source file)
 				//		-> Two CASES:
 				//				1) NOT WORKING: Non-unity file went away (SourceFile in this context).  That seems like an existing old use case.  Compile params or Response file should have changed?
 				//				2) WORKING: Indirect file went away (unity'd original source file or include).  This would return a file that no longer exists and adds to the prerequiteitems list
-				var IncludedFileList = CPPEnvironment.FindAndCacheAllIncludedFiles( Target, SourceFile, BuildPlatform, CompileEnvironment.Config.CPPIncludeInfo, bOnlyCachedDependencies:BuildConfiguration.bUseExperimentalFastDependencyScan );
+				var IncludedFileList = CPPEnvironment.FindAndCacheAllIncludedFiles( Target, SourceFile, BuildPlatform, CompileEnvironment.Config.CPPIncludeInfo, bOnlyCachedDependencies:BuildConfiguration.bUseUBTMakefiles );
 				foreach (FileItem IncludedFile in IncludedFileList)
 				{
 					PrerequisiteItems.Add( IncludedFile );
 
 					if( bAllowUploading &&
-						!BuildConfiguration.bUseExperimentalFastDependencyScan )	// With fast dependency scanning, we will not have an exhaustive list of dependencies here.  We rely on PostCodeGeneration() to upload these files.
+						!BuildConfiguration.bUseUBTMakefiles )	// With fast dependency scanning, we will not have an exhaustive list of dependencies here.  We rely on PostCodeGeneration() to upload these files.
 					{
 						RemoteThis.QueueFileForBatchUpload(IncludedFile);
 					}
