UnrealBuildTool: Fixed crash with UBT Makefiles
- Fixed crash after trying to build a target different than the one you built last time when DependencyCache was missing

diff --git a/Engine/Source/Programs/UnrealBuildTool/System/ActionGraph.cs b/Engine/Source/Programs/UnrealBuildTool/System/ActionGraph.cs
--- a/Engine/Source/Programs/UnrealBuildTool/System/ActionGraph.cs
+++ b/Engine/Source/Programs/UnrealBuildTool/System/ActionGraph.cs
@@ -818,26 +818,29 @@ namespace UnrealBuildTool
 						if( PrerequisiteItem.CachedCPPIncludeInfo != null )
 						{
 							var IncludedFileList = CPPEnvironment.FindAndCacheAllIncludedFiles( Target, PrerequisiteItem, BuildPlatform, PrerequisiteItem.CachedCPPIncludeInfo, bOnlyCachedDependencies:BuildConfiguration.bUseUBTMakefiles );
-							foreach( var IncludedFile in IncludedFileList )
-							{
-								if( IncludedFile.bExists )
+							if( IncludedFileList != null )
+							{ 
+								foreach( var IncludedFile in IncludedFileList )
 								{
-									// allow a 1 second slop for network copies
-									TimeSpan TimeDifference = IncludedFile.LastWriteTime - LastExecutionTime;
-									bool bPrerequisiteItemIsNewerThanLastExecution = TimeDifference.TotalSeconds > 1;
-									if (bPrerequisiteItemIsNewerThanLastExecution)
+									if( IncludedFile.bExists )
 									{
-										Log.TraceVerbose(
-											"{0}: Included file {1} is newer than the last execution of the action: {2} vs {3}",
-											RootAction.StatusDescription,
-											Path.GetFileName(IncludedFile.AbsolutePath),
-											IncludedFile.LastWriteTime.LocalDateTime,
-											LastExecutionTime.LocalDateTime
-											);
-										bIsOutdated = true;
-
-										// Don't bother checking every single include if we've found one that is out of date
-										break;
+										// allow a 1 second slop for network copies
+										TimeSpan TimeDifference = IncludedFile.LastWriteTime - LastExecutionTime;
+										bool bPrerequisiteItemIsNewerThanLastExecution = TimeDifference.TotalSeconds > 1;
+										if (bPrerequisiteItemIsNewerThanLastExecution)
+										{
+											Log.TraceVerbose(
+												"{0}: Included file {1} is newer than the last execution of the action: {2} vs {3}",
+												RootAction.StatusDescription,
+												Path.GetFileName(IncludedFile.AbsolutePath),
+												IncludedFile.LastWriteTime.LocalDateTime,
+												LastExecutionTime.LocalDateTime
+												);
+											bIsOutdated = true;
+
+											// Don't bother checking every single include if we've found one that is out of date
+											break;
+										}
 									}
 								}
 							}
diff --git a/Engine/Source/Programs/UnrealBuildTool/ToolChain/UEToolChain.cs b/Engine/Source/Programs/UnrealBuildTool/ToolChain/UEToolChain.cs
--- a/Engine/Source/Programs/UnrealBuildTool/ToolChain/UEToolChain.cs
+++ b/Engine/Source/Programs/UnrealBuildTool/ToolChain/UEToolChain.cs
@@ -277,14 +277,17 @@ namespace UnrealBuildTool
 				//				1) NOT WORKING: Non-unity file went away (SourceFile in this context).  That seems like an existing old use case.  Compile params or Response file should have changed?
 				//				2) WORKING: Indirect file went away (unity'd original source file or include).  This would return a file that no longer exists and adds to the prerequiteitems list
 				var IncludedFileList = CPPEnvironment.FindAndCacheAllIncludedFiles( Target, SourceFile, BuildPlatform, CompileEnvironment.Config.CPPIncludeInfo, bOnlyCachedDependencies:BuildConfiguration.bUseUBTMakefiles );
-				foreach (FileItem IncludedFile in IncludedFileList)
+				if( IncludedFileList != null )
 				{
-					PrerequisiteItems.Add( IncludedFile );
-
-					if( bAllowUploading &&
-						!BuildConfiguration.bUseUBTMakefiles )	// With fast dependency scanning, we will not have an exhaustive list of dependencies here.  We rely on PostCodeGeneration() to upload these files.
+					foreach (FileItem IncludedFile in IncludedFileList)
 					{
-						RemoteThis.QueueFileForBatchUpload(IncludedFile);
+						PrerequisiteItems.Add( IncludedFile );
+
+						if( bAllowUploading &&
+							!BuildConfiguration.bUseUBTMakefiles )	// With fast dependency scanning, we will not have an exhaustive list of dependencies here.  We rely on PostCodeGeneration() to upload these files.
+						{
+							RemoteThis.QueueFileForBatchUpload(IncludedFile);
+						}
 					}
 				}
 			}
