Produce an error if there is no VS2013 toolchain installed for Windows. All the VS2012 plumbing is still available for project file generation and hosting the XboxOne toolchain.

diff --git a/Engine/Source/Programs/UnrealBuildTool/Windows/UEBuildWindows.cs b/Engine/Source/Programs/UnrealBuildTool/Windows/UEBuildWindows.cs
--- a/Engine/Source/Programs/UnrealBuildTool/Windows/UEBuildWindows.cs
+++ b/Engine/Source/Programs/UnrealBuildTool/Windows/UEBuildWindows.cs
@@ -14,7 +14,7 @@ namespace UnrealBuildTool
     /// </summary>
     public enum WindowsCompiler
     {
-        /// Visual Studio 2012 (Visual C++ 11.0)
+        /// Visual Studio 2012 (Visual C++ 11.0). No longer supported for building on Windows, but required for other platform toolchains.
         VisualStudio2012,
 
         /// Visual Studio 2013 (Visual C++ 12.0)
@@ -168,6 +168,13 @@ namespace UnrealBuildTool
          */
         protected override void RegisterBuildPlatformInternal()
         {
+			// All the plumbing for VS2012 is still in place to support console toolchains that piggy-back on the VS2012 environment, but we require VS2013 for building the editor
+			// and host tools. We always prefer VS2013 over VS2012, so if we detected VS2012 as the only installed compiler, halt now.
+            if (Compiler == WindowsCompiler.VisualStudio2012)
+            {
+				throw new BuildException("Visual Studio 2012 is no longer supported for building the editor and host platform tools. Please also install Visual Studio 2013.");
+            }
+
             // Register this build platform for both Win64 and Win32
             Log.TraceVerbose("        Registering for {0}", UnrealTargetPlatform.Win64.ToString());
             UEBuildPlatform.RegisterBuildPlatform(UnrealTargetPlatform.Win64, this);
