fix for crash when running on a Mac with mono
#mac

diff --git a/Engine/Source/Programs/UnrealBuildTool/Windows/UEBuildWindows.cs b/Engine/Source/Programs/UnrealBuildTool/Windows/UEBuildWindows.cs
--- a/Engine/Source/Programs/UnrealBuildTool/Windows/UEBuildWindows.cs
+++ b/Engine/Source/Programs/UnrealBuildTool/Windows/UEBuildWindows.cs
@@ -170,7 +170,7 @@ namespace UnrealBuildTool
         {
 			// All the plumbing for VS2012 is still in place to support console toolchains that piggy-back on the VS2012 environment, but we require VS2013 for building the editor
 			// and host tools. We always prefer VS2013 over VS2012, so if we detected VS2012 as the only installed compiler, halt now.
-            if (Compiler == WindowsCompiler.VisualStudio2012)
+            if (!Utils.IsRunningOnMono && Compiler == WindowsCompiler.VisualStudio2012)
             {
 				throw new BuildException("Visual Studio 2012 is no longer supported for building the editor and host platform tools. Please also install Visual Studio 2013.");
             }
