Enabled UBTMakefiles by default on Windows
- Fast UnrealBuildTool startup and dependency checking
- Preflighted, but needs wider testing and debugging.  More info to come soon.

diff --git a/Engine/Source/Programs/UnrealBuildTool/Configuration/BuildConfiguration.cs b/Engine/Source/Programs/UnrealBuildTool/Configuration/BuildConfiguration.cs
--- a/Engine/Source/Programs/UnrealBuildTool/Configuration/BuildConfiguration.cs
+++ b/Engine/Source/Programs/UnrealBuildTool/Configuration/BuildConfiguration.cs
@@ -529,7 +529,7 @@ namespace UnrealBuildTool
 			// Enables support for fast include dependency scanning, as well as gathering data for 'UBT Makefiles', then quickly
 			// assembling builds in subsequent runs using data in those cached makefiles
 			// NOTE: This feature is new and has a number of known issues (search the code for '@todo ubtmake')
-			bUseUBTMakefiles = false;
+			bUseUBTMakefiles = !Utils.IsRunningOnMono;	// @todo ubtmake: Needs support for Mac
 
 			// Distcc requires some setup - so by default disable it so we don't break local or remote building
 			bAllowDistcc = false;
