Made a tweak in GUBP_IncludeProjectInPromotedBuild_EditorTypeOnly to accomodate FP targets as well as TP ones so they will get promoted.

diff --git a/Engine/Source/Programs/UnrealBuildTool/System/RulesCompiler.cs b/Engine/Source/Programs/UnrealBuildTool/System/RulesCompiler.cs
--- a/Engine/Source/Programs/UnrealBuildTool/System/RulesCompiler.cs
+++ b/Engine/Source/Programs/UnrealBuildTool/System/RulesCompiler.cs
@@ -752,7 +752,8 @@ namespace UnrealBuildTool
         {
             var Result = new GUBPProjectOptions();
             // hack to set up the templates without adding anything to their .targets.cs files
-            if (!String.IsNullOrEmpty(TargetName) && TargetName.StartsWith("TP_"))
+			// tweaked to include FP_ folders too - which are temporary
+            if (!String.IsNullOrEmpty(TargetName) && ( TargetName.StartsWith("TP_") || TargetName.StartsWith("FP_")) )
             {
                 Result.bTestWithShared = true;
             }
