fix for not getting the right settings for the project when generating the xcodeproject
UE-6928
#ios

diff --git a/Engine/Source/Programs/UnrealBuildTool/System/XcodeProjectFileGenerator.cs b/Engine/Source/Programs/UnrealBuildTool/System/XcodeProjectFileGenerator.cs
--- a/Engine/Source/Programs/UnrealBuildTool/System/XcodeProjectFileGenerator.cs
+++ b/Engine/Source/Programs/UnrealBuildTool/System/XcodeProjectFileGenerator.cs
@@ -477,6 +477,7 @@ namespace UnrealBuildTool
 			else
 			{
 				IUEToolChain Toolchain = UEToolChain.GetPlatformToolChain(CPPTargetPlatform.IOS);
+				Toolchain.SetUpGlobalEnvironment();
 
 				Contents.Append(
 					"\t\t" + ConfigGuid + " /* " + ConfigName + " */ = {" + ProjectFileGenerator.NewLine +
@@ -506,6 +507,7 @@ namespace UnrealBuildTool
 		private void AppendIOSBuildConfig(ref StringBuilder Contents, string ConfigName, string ConfigGuid)
 		{
 			IUEToolChain Toolchain = UEToolChain.GetPlatformToolChain(CPPTargetPlatform.IOS);
+			Toolchain.SetUpGlobalEnvironment();
 
 			Contents.Append(
 				"\t\t" + ConfigGuid + " /* " + ConfigName + " */ = {" + ProjectFileGenerator.NewLine +
@@ -526,6 +528,7 @@ namespace UnrealBuildTool
 		private void AppendIOSRunConfig(ref StringBuilder Contents, string ConfigName, string ConfigGuid, string TargetName, string EngineRelative, string GamePath, bool bIsUE4Game, bool IsAGame, bool bIsUE4Client)
 		{
 			IUEToolChain Toolchain = UEToolChain.GetPlatformToolChain(CPPTargetPlatform.IOS);
+			Toolchain.SetUpGlobalEnvironment();
 
 			Contents.Append(
 				"\t\t" + ConfigGuid + " /* " + ConfigName + " */ = {" + ProjectFileGenerator.NewLine +
