fix for build spam
UE-6858
#html5

diff --git a/Engine/Source/Programs/UnrealBuildTool/HTML5/HTML5ToolChain.cs b/Engine/Source/Programs/UnrealBuildTool/HTML5/HTML5ToolChain.cs
--- a/Engine/Source/Programs/UnrealBuildTool/HTML5/HTML5ToolChain.cs
+++ b/Engine/Source/Programs/UnrealBuildTool/HTML5/HTML5ToolChain.cs
@@ -72,8 +72,8 @@ namespace UnrealBuildTool
             Result += " -s DISABLE_EXCEPTION_CATCHING=1";
             // enable checking for missing functions at link time as opposed to runtime
             Result += " -s WARN_ON_UNDEFINED_SYMBOLS=1";
-            // we want full ES2 and not use system SDL
-            Result += " -s FULL_ES2=1 -s USE_SDL=0 ";
+            // we want full ES2
+            Result += " -s FULL_ES2=1 ";
             // export console command handler. Export main func too because default exports ( e.g Main ) are overridden if we use custom exported functions. 
             Result += " -s EXPORTED_FUNCTIONS=\"['_main', '_resize_game', '_on_fatal']\" ";
 
