UnrealBuildTool: Fixed with detecting outdated includes with UBT Makefiles
- Fixed flat include cache not being refreshed after user switches between projects
- Fixed flat include cache not being rebuilt if the cache file was missing

diff --git a/Engine/Source/Programs/UnrealBuildTool/System/ActionGraph.cs b/Engine/Source/Programs/UnrealBuildTool/System/ActionGraph.cs
--- a/Engine/Source/Programs/UnrealBuildTool/System/ActionGraph.cs
+++ b/Engine/Source/Programs/UnrealBuildTool/System/ActionGraph.cs
@@ -782,7 +782,8 @@ namespace UnrealBuildTool
 				Log.WriteLineIf(BuildConfiguration.bLogDetailedActionStats && !String.IsNullOrEmpty( LatestUpdatedProducedItemName ),
 					TraceEventType.Verbose, "{0}: Oldest produced item is {1}", RootAction.StatusDescription, LatestUpdatedProducedItemName);
 
-				bool bFindCPPIncludePrerequisites = false;
+				bool bCheckIfIncludedFilesAreNewer = false;
+				bool bPerformExhaustiveIncludeSearchAndUpdateCache = false;
 				if( RootAction.ActionType == ActionType.Compile )
 				{
 					// Outdated targets don't need their headers scanned yet, because presumably they would already be out of dated based on already-cached
@@ -793,7 +794,7 @@ namespace UnrealBuildTool
 						UnrealBuildTool.IsAssemblingBuild && 
 						RootAction.ActionType == ActionType.Compile )
 					{
-						bFindCPPIncludePrerequisites = true;
+						bCheckIfIncludedFilesAreNewer = true;
 					}
 
 					// Were we asked to force an update of our cached includes BEFORE we try to build?  This may be needed if our cache can no longer
@@ -801,12 +802,13 @@ namespace UnrealBuildTool
 					if( BuildConfiguration.bUseUBTMakefiles && 
 						UnrealBuildTool.bNeedsFullCPPIncludeRescan )
 					{
-						bFindCPPIncludePrerequisites = true;
+						// This will be slow!
+						bPerformExhaustiveIncludeSearchAndUpdateCache = true;
 					}
 				}
 
 
-				if( bFindCPPIncludePrerequisites )
+				if( bCheckIfIncludedFilesAreNewer || bPerformExhaustiveIncludeSearchAndUpdateCache )
 				{
 					// Scan this file for included headers that may be out of date.  Note that it's OK if we break out early because we found
 					// the action to be outdated.  For outdated actions, we kick off a separate include scan in a background thread later on to
@@ -817,7 +819,7 @@ namespace UnrealBuildTool
 						// @todo ubtmake: Make sure we are catching RC files here too.  Anything that the toolchain would have tried it on.  Logic should match the CACHING stuff below
 						if( PrerequisiteItem.CachedCPPIncludeInfo != null )
 						{
-							var IncludedFileList = CPPEnvironment.FindAndCacheAllIncludedFiles( Target, PrerequisiteItem, BuildPlatform, PrerequisiteItem.CachedCPPIncludeInfo, bOnlyCachedDependencies:BuildConfiguration.bUseUBTMakefiles );
+							var IncludedFileList = CPPEnvironment.FindAndCacheAllIncludedFiles( Target, PrerequisiteItem, BuildPlatform, PrerequisiteItem.CachedCPPIncludeInfo, bOnlyCachedDependencies:!bPerformExhaustiveIncludeSearchAndUpdateCache );
 							if( IncludedFileList != null )
 							{ 
 								foreach( var IncludedFile in IncludedFileList )
diff --git a/Engine/Source/Programs/UnrealBuildTool/System/CPPHeaders.cs b/Engine/Source/Programs/UnrealBuildTool/System/CPPHeaders.cs
--- a/Engine/Source/Programs/UnrealBuildTool/System/CPPHeaders.cs
+++ b/Engine/Source/Programs/UnrealBuildTool/System/CPPHeaders.cs
@@ -182,11 +182,8 @@ namespace UnrealBuildTool
 				CPPIncludeInfo.IncludeFileSearchDictionary = new Dictionary<string,FileItem>();
 			}
 
-			bool bUseFlatCPPIncludeDependencyCache = 
-				( BuildConfiguration.bUseUBTMakefiles &&
-				  ( !BuildConfiguration.bUseUBTMakefiles || UnrealBuildTool.IsAssemblingBuild ) );
-
-			if( bUseFlatCPPIncludeDependencyCache && bOnlyCachedDependencies )
+			bool bUseFlatCPPIncludeDependencyCache = BuildConfiguration.bUseUBTMakefiles && UnrealBuildTool.IsAssemblingBuild;
+			if( bOnlyCachedDependencies && bUseFlatCPPIncludeDependencyCache )
 			{ 
 				Result = FlatCPPIncludeDependencyCache[Target].GetDependenciesForFile( SourceFile.AbsolutePath );
 				if( Result == null )
@@ -198,6 +195,15 @@ namespace UnrealBuildTool
 			{
 				// @todo ubtmake: HeaderParser.h is missing from the include set for Module.UnrealHeaderTool.cpp (failed to find include using:  FileItem DirectIncludeResolvedFile = CPPEnvironment.FindIncludedFile(DirectInclude.IncludeName, !BuildConfiguration.bCheckExternalHeadersForModification, IncludePathsToSearch, IncludeFileSearchDictionary );)
 
+				// If we're doing an exhaustive include scan, make sure that we have our include dependency cache loaded and ready
+				if( !bOnlyCachedDependencies )
+				{ 
+					if( !IncludeDependencyCache.ContainsKey( Target ) )
+					{
+						IncludeDependencyCache.Add( Target, DependencyCache.Create( DependencyCache.GetDependencyCachePathForTarget( Target ) ) );
+					}
+				}
+
 				Result = new List<FileItem>();
 
 				var IncludedFileList = new IncludedFilesSet();
diff --git a/Engine/Source/Programs/UnrealBuildTool/System/FlatCPPIncludeDepencencyCache.cs b/Engine/Source/Programs/UnrealBuildTool/System/FlatCPPIncludeDepencencyCache.cs
--- a/Engine/Source/Programs/UnrealBuildTool/System/FlatCPPIncludeDepencencyCache.cs
+++ b/Engine/Source/Programs/UnrealBuildTool/System/FlatCPPIncludeDepencencyCache.cs
@@ -68,6 +68,8 @@ namespace UnrealBuildTool
 			}
 
 			// Fall back to a clean cache on error or non-existence.
+			UnrealBuildTool.bNeedsFullCPPIncludeRescan = true;
+			Log.TraceInformation( "Performing full C++ include scan (no include cache file)" );
 			return new FlatCPPIncludeDependencyCache( Cache );
 		}
 
diff --git a/Engine/Source/Programs/UnrealBuildTool/System/UnrealBuildTool.cs b/Engine/Source/Programs/UnrealBuildTool/System/UnrealBuildTool.cs
--- a/Engine/Source/Programs/UnrealBuildTool/System/UnrealBuildTool.cs
+++ b/Engine/Source/Programs/UnrealBuildTool/System/UnrealBuildTool.cs
@@ -2082,18 +2082,6 @@ namespace UnrealBuildTool
                     // @todo ubtmake: This thread will access data structures that are also used on the main UBT thread, but during this time UBT
                     // is only invoking the build executor, so should not be touching this stuff.  However, we need to at some guards to make sure.
 
-
-                    // When in 'assembler only' mode, we need to load the direct include graph from a cache now.  We skipped loading this earlier because
-                    // we didn't need it unless there was actually an outdated dependency.  And we can load this on our worker thread while we're compiling.
-                    foreach( var Target in TargetToOutdatedPrerequisitesMap.Keys )
-                    {
-                        if (!CPPEnvironment.IncludeDependencyCache.ContainsKey(Target))
-                        {
-                            CPPEnvironment.IncludeDependencyCache.Add( Target, DependencyCache.Create( DependencyCache.GetDependencyCachePathForTarget(Target) ) );
-                        }
-                    }
-
-
                     foreach( var TargetAndOutdatedPrerequisites in TargetToOutdatedPrerequisitesMap )
                     {
                         var Target = TargetAndOutdatedPrerequisites.Key;
