UBT: Fixed public library path settings being ignored in build rules

diff --git a/Engine/Source/Programs/UnrealBuildTool/Configuration/UEBuildModule.cs b/Engine/Source/Programs/UnrealBuildTool/Configuration/UEBuildModule.cs
--- a/Engine/Source/Programs/UnrealBuildTool/Configuration/UEBuildModule.cs
+++ b/Engine/Source/Programs/UnrealBuildTool/Configuration/UEBuildModule.cs
@@ -1055,6 +1055,7 @@ namespace UnrealBuildTool
 			IEnumerable<FileItem> InSourceFiles,
 			IEnumerable<string> InPublicIncludePaths,
 			IEnumerable<string> InPublicSystemIncludePaths,
+            IEnumerable<string> InPublicLibraryPaths,
 			IEnumerable<string> InDefinitions,
 			IEnumerable<string> InPublicIncludePathModuleNames,
 			IEnumerable<string> InPublicDependencyModuleNames,
@@ -1090,7 +1091,7 @@ namespace UnrealBuildTool
 					InDefinitions, 
 					InPublicIncludePaths,
 					InPublicSystemIncludePaths, 
-					null, 
+					InPublicLibraryPaths, 
 					InPublicAdditionalLibraries,
 					InPublicFrameworks,
 					InPublicWeakFrameworks,
@@ -2184,7 +2185,7 @@ namespace UnrealBuildTool
 			bool bInBuildSourceFiles
 			)
 			: base(InTarget,InName,InType,InModuleDirectory,InOutputDirectory,InIsRedistributableOverride,InIntelliSenseGatherer,
-			InSourceFiles,InPublicIncludePaths,InPublicSystemIncludePaths,InDefinitions,
+			InSourceFiles,InPublicIncludePaths,InPublicSystemIncludePaths,null,InDefinitions,
 			InPublicIncludePathModuleNames,InPublicDependencyModuleNames,InPublicDelayLoadDLLs,InPublicAdditionalLibraries,InPublicFrameworks,InPublicWeakFrameworks,InPublicAdditionalFrameworks,InPublicAdditionalShadowFiles,InPublicAdditionalBundleResources,
 			InPrivateIncludePaths,InPrivateIncludePathModuleNames,InPrivateDependencyModuleNames,
             InCircularlyReferencedDependentModules, InDynamicallyLoadedModuleNames, InPlatformSpecificDynamicallyLoadedModuleNames, InOptimizeCode,
diff --git a/Engine/Source/Programs/UnrealBuildTool/Configuration/UEBuildTarget.cs b/Engine/Source/Programs/UnrealBuildTool/Configuration/UEBuildTarget.cs
--- a/Engine/Source/Programs/UnrealBuildTool/Configuration/UEBuildTarget.cs
+++ b/Engine/Source/Programs/UnrealBuildTool/Configuration/UEBuildTarget.cs
@@ -2045,6 +2045,7 @@ namespace UnrealBuildTool
 				InIntelliSenseGatherer: null,
 				InSourceFiles: SourceFiles.ToList(),
 				InPublicIncludePaths: new List<string>(),
+                InPublicLibraryPaths: new List<string>(),
 				InPublicSystemIncludePaths: new List<string>(),
 				InDefinitions: new List<string>(),
 				InPublicIncludePathModuleNames: new List<string>(),
@@ -3017,6 +3018,7 @@ namespace UnrealBuildTool
 							InIntelliSenseGatherer: IntelliSenseGatherer,
 							InSourceFiles: ModuleSourceFiles,
 							InPublicSystemIncludePaths: RulesObject.PublicSystemIncludePaths,
+                            InPublicLibraryPaths: RulesObject.PublicLibraryPaths,
 							InPublicIncludePaths: RulesObject.PublicIncludePaths,
 							InDefinitions: RulesObject.Definitions,
 							InPublicIncludePathModuleNames: RulesObject.PublicIncludePathModuleNames,
