Disabled bUseUBTMakefiles when building with RPCUtility

diff --git a/Engine/Source/Programs/UnrealBuildTool/System/UnrealBuildTool.cs b/Engine/Source/Programs/UnrealBuildTool/System/UnrealBuildTool.cs
--- a/Engine/Source/Programs/UnrealBuildTool/System/UnrealBuildTool.cs
+++ b/Engine/Source/Programs/UnrealBuildTool/System/UnrealBuildTool.cs
@@ -787,6 +787,12 @@ namespace UnrealBuildTool
                         }
                     }
 
+					// @todo ubtmake: remove this when building with RPCUtility works
+					if (CheckPlatform == UnrealTargetPlatform.Mac || CheckPlatform == UnrealTargetPlatform.IOS)
+					{
+						BuildConfiguration.bUseUBTMakefiles = false;
+					}
+
                     // If we were asked to enable fast build iteration, we want the 'gather' phase to default to off (unless it is overridden below
                     // using a command-line option.)
                     if( BuildConfiguration.bUseUBTMakefiles )
