Add a helper function to get the module directory from inside the module rules. Useful for adding ThirdParty include paths from plugins and game modules.

diff --git a/Engine/Source/Programs/UnrealBuildTool/System/RulesCompiler.cs b/Engine/Source/Programs/UnrealBuildTool/System/RulesCompiler.cs
--- a/Engine/Source/Programs/UnrealBuildTool/System/RulesCompiler.cs
+++ b/Engine/Source/Programs/UnrealBuildTool/System/RulesCompiler.cs
@@ -222,6 +222,17 @@ namespace UnrealBuildTool
         /// Extra modules this module may require at run time, that are on behalf of another platform (i.e. shader formats and the like)
         public List<string> PlatformSpecificDynamicallyLoadedModuleNames = new List<string>();
 
+		/// <summary>
+		/// Property for the directory containing this module. Useful for adding paths to third party dependencies.
+		/// </summary>
+		public string ModuleDirectory
+		{
+			get
+			{
+				return Path.GetDirectoryName(RulesCompiler.GetModuleFilename(GetType().Name));
+			}
+		}
+
 		/// <summary>
 		/// Add the given ThirdParty modules as static private dependencies
 		///	Statically linked to this module, meaning they utilize exports from the other module
