Slight tidy-up of GenerateLinkerFixupsContents so that it returns the contents instead of filling out an empty passed-in container.

diff --git a/Engine/Source/Programs/UnrealBuildTool/Configuration/UEBuildTarget.cs b/Engine/Source/Programs/UnrealBuildTool/Configuration/UEBuildTarget.cs
--- a/Engine/Source/Programs/UnrealBuildTool/Configuration/UEBuildTarget.cs
+++ b/Engine/Source/Programs/UnrealBuildTool/Configuration/UEBuildTarget.cs
@@ -1939,9 +1939,8 @@ namespace UnrealBuildTool
 				BindArtificialModuleToBinary(NewModule, ExecutableBinary);
 
 				// Create the cpp file
-				List<string> LinkerFixupsFileContents = new List<string>();
 				NewModule.bSkipDefinitionsForCompileEnvironment = false;
-				GenerateLinkerFixupsContents(ExecutableBinary, LinkerFixupsFileContents, NewModule.CreateModuleCompileEnvironment(GlobalCompileEnvironment), HeaderFilename, LinkerFixupsName, PrivateDependencyModuleNames);
+				var LinkerFixupsFileContents = GenerateLinkerFixupsContents(ExecutableBinary, NewModule.CreateModuleCompileEnvironment(GlobalCompileEnvironment), HeaderFilename, LinkerFixupsName, PrivateDependencyModuleNames);
 				NewModule.bSkipDefinitionsForCompileEnvironment = true;
 
 				// Determine if the file changed. Write it if it either doesn't exist or the contents are different.
@@ -1963,9 +1962,11 @@ namespace UnrealBuildTool
 			}
 		}
 
-		private void GenerateLinkerFixupsContents(UEBuildBinary ExecutableBinary, List<string> LinkerFixupsFileContents, CPPEnvironment CompileEnvironment, string HeaderFilename, string LinkerFixupsName, List<string> PrivateDependencyModuleNames)
+		private  List<string> GenerateLinkerFixupsContents(UEBuildBinary ExecutableBinary, CPPEnvironment CompileEnvironment, string HeaderFilename, string LinkerFixupsName, List<string> PrivateDependencyModuleNames)
 		{
-			LinkerFixupsFileContents.Add("#include \"" + HeaderFilename + "\"");
+			var Result = new List<string>();
+
+			Result.Add("#include \"" + HeaderFilename + "\"");
 
 			// To reduce the size of the command line for the compiler, we're going to put all definitions inside of the cpp file.
 			foreach (var Definition in CompileEnvironment.Config.Definitions)
@@ -1982,30 +1983,32 @@ namespace UnrealBuildTool
 				{
 					MacroName = Definition;
 				}
-				LinkerFixupsFileContents.Add("#ifndef " + MacroName);
-				LinkerFixupsFileContents.Add(String.Format("\t#define {0} {1}", MacroName, MacroValue));
-				LinkerFixupsFileContents.Add("#endif");
+				Result.Add("#ifndef " + MacroName);
+				Result.Add(String.Format("\t#define {0} {1}", MacroName, MacroValue));
+				Result.Add("#endif");
 			}
 
 			// Add a function that is not referenced by anything that invokes all the empty functions in the different static libraries
-			LinkerFixupsFileContents.Add("void " + LinkerFixupsName + "()");
-			LinkerFixupsFileContents.Add("{");
+			Result.Add("void " + LinkerFixupsName + "()");
+			Result.Add("{");
 
 			// Fill out the body of the function with the empty function calls. This is what causes the static libraries to be considered relevant
 			var DependencyModules = ExecutableBinary.GetAllDependencyModules(bIncludeDynamicallyLoaded: false, bForceCircular: false);
 			foreach(string ModuleName in DependencyModules.OfType<UEBuildModuleCPP>().Where(CPPModule => CPPModule.AutoGenerateCppInfo != null).Select(CPPModule => CPPModule.Name).Distinct())
 			{
-				LinkerFixupsFileContents.Add("    extern void EmptyLinkFunctionForGeneratedCode" + ModuleName + "();");
-				LinkerFixupsFileContents.Add("    EmptyLinkFunctionForGeneratedCode" + ModuleName + "();");
+				Result.Add("    extern void EmptyLinkFunctionForGeneratedCode" + ModuleName + "();");
+				Result.Add("    EmptyLinkFunctionForGeneratedCode" + ModuleName + "();");
 			}
 			foreach (var DependencyModuleName in PrivateDependencyModuleNames)
 			{
-				LinkerFixupsFileContents.Add("    extern void EmptyLinkFunctionForStaticInitialization" + DependencyModuleName + "();");
-				LinkerFixupsFileContents.Add("    EmptyLinkFunctionForStaticInitialization" + DependencyModuleName + "();");
+				Result.Add("    extern void EmptyLinkFunctionForStaticInitialization" + DependencyModuleName + "();");
+				Result.Add("    EmptyLinkFunctionForStaticInitialization" + DependencyModuleName + "();");
 			}
 
 			// End the function body that was started above
-			LinkerFixupsFileContents.Add("}");
+			Result.Add("}");
+
+			return Result;
 		}
 
 		/// <summary>
