Shadow third-party frameworks to the remote when compiling on Mac from Windows.

diff --git a/Engine/Source/Programs/UnrealBuildTool/Mac/MacToolChain.cs b/Engine/Source/Programs/UnrealBuildTool/Mac/MacToolChain.cs
--- a/Engine/Source/Programs/UnrealBuildTool/Mac/MacToolChain.cs
+++ b/Engine/Source/Programs/UnrealBuildTool/Mac/MacToolChain.cs
@@ -260,19 +260,19 @@ namespace UnrealBuildTool
 			return Result;
 		}
 
-		static string AddFrameworkToLinkCommand(string FrameworkName, string Arg = "-framework")
+		string AddFrameworkToLinkCommand(string FrameworkName, string Arg = "-framework")
 		{
 			string Result = "";
 			if (FrameworkName.EndsWith(".framework"))
 			{
-				Result += " -F \"" + Path.GetDirectoryName(FrameworkName) + "\"";
+				Result += " -F \"" + ConvertPath(Path.GetDirectoryName(Path.GetFullPath(FrameworkName))) + "\"";
 				FrameworkName = Path.GetFileNameWithoutExtension(FrameworkName);
 			}
 			Result += " " + Arg + " \"" + FrameworkName + "\"";
 			return Result;
 		}
 
-		static string GetLinkArguments_Global(LinkEnvironment LinkEnvironment)
+		string GetLinkArguments_Global(LinkEnvironment LinkEnvironment)
 		{
 			string Result = "";
 
@@ -681,6 +681,20 @@ namespace UnrealBuildTool
 						throw new BuildException("Couldn't find required additional file to shadow: {0}", AdditionalShadowFile);
 					}
 				}
+
+				// Add any frameworks to be shadowed to the remote
+				foreach (string FrameworkPath in LinkEnvironment.Config.Frameworks)
+				{
+					if(FrameworkPath.EndsWith(".framework"))
+					{
+						foreach(string FrameworkFile in Directory.EnumerateFiles(FrameworkPath, "*", SearchOption.AllDirectories))
+						{
+							FileItem FrameworkFileItem = FileItem.GetExistingItemByPath(FrameworkFile);
+							QueueFileForBatchUpload(FrameworkFileItem);
+							LinkAction.PrerequisiteItems.Add(FrameworkFileItem);
+						}
+					}
+				}
 			}
 
 			bool bIsBuildingAppBundle = !LinkEnvironment.Config.bIsBuildingDLL && !LinkEnvironment.Config.bIsBuildingLibrary && !LinkEnvironment.Config.bIsBuildingConsoleApplication;
