fix for stub created entirely on Mac not being useable for content-only projects
UE-6834
#ios

diff --git a/Engine/Source/Programs/UnrealBuildTool/IOS/UEDeployIOS.cs b/Engine/Source/Programs/UnrealBuildTool/IOS/UEDeployIOS.cs
--- a/Engine/Source/Programs/UnrealBuildTool/IOS/UEDeployIOS.cs
+++ b/Engine/Source/Programs/UnrealBuildTool/IOS/UEDeployIOS.cs
@@ -116,17 +116,10 @@ namespace UnrealBuildTool.IOS
 			string BuildDirectory = InProjectDirectory + "/Build/IOS";
 			string IntermediateDirectory = (bIsUE4Game ? InEngineDir : InProjectDirectory) + "/Intermediate/IOS";
 
-			// don't delete the payload directory, just update the data in it
-			//if (Directory.Exists(PayloadDirectory))
-			//{
-			//	Directory.Delete(PayloadDirectory, true);
-			//}
-
 			Directory.CreateDirectory(BinaryPath);
 			Directory.CreateDirectory(PayloadDirectory);
 			Directory.CreateDirectory(AppDirectory);
 			Directory.CreateDirectory(BuildDirectory);
-			Directory.CreateDirectory(CookedContentDirectory);
 
 			// create the entitlements file
 			WriteEntitlementsFile(Path.Combine(IntermediateDirectory, GameName + ".entitlements"));
@@ -348,38 +341,41 @@ namespace UnrealBuildTool.IOS
 			// copy the GameName binary
 			File.Copy(BinaryPath + "/" + GameExeName, AppDirectory + "/" + GameName, true);
 
-			// copy engine assets in
-			if (bSkipDefaultPNGs)
-			{
-				// we still want default icons
-				CopyFiles(InEngineDir + "/Build/IOS/Resources/Graphics", AppDirectory, "Icon*.png", true);
-			}
-			else
-			{
-				CopyFiles(InEngineDir + "/Build/IOS/Resources/Graphics", AppDirectory, "*.png", true);
-			}
-			// merge game assets on top
-			if (Directory.Exists(BuildDirectory + "/Resources/Graphics"))
+			if (!BuildConfiguration.bCreateStubIPA)
 			{
-				CopyFiles(BuildDirectory + "/Resources/Graphics", AppDirectory, "*.png", true);
-			}
-			
-			// copy additional engine framework assets in
-			string FrameworkAssetsPath = InEngineDir + "/Intermediate/IOS/FrameworkAssets";
+				// copy engine assets in
+				if (bSkipDefaultPNGs)
+				{
+					// we still want default icons
+					CopyFiles(InEngineDir + "/Build/IOS/Resources/Graphics", AppDirectory, "Icon*.png", true);
+				}
+				else
+				{
+					CopyFiles(InEngineDir + "/Build/IOS/Resources/Graphics", AppDirectory, "*.png", true);
+				}
+				// merge game assets on top
+				if (Directory.Exists(BuildDirectory + "/Resources/Graphics"))
+				{
+					CopyFiles(BuildDirectory + "/Resources/Graphics", AppDirectory, "*.png", true);
+				}
 
-			// Let project override assets if they exist
-			if ( Directory.Exists( InProjectDirectory + "/Intermediate/IOS/FrameworkAssets" ) )
-			{
-				FrameworkAssetsPath = InProjectDirectory + "/Intermediate/IOS/FrameworkAssets";
-			}
+				// copy additional engine framework assets in
+				string FrameworkAssetsPath = InEngineDir + "/Intermediate/IOS/FrameworkAssets";
 
-			if ( Directory.Exists( FrameworkAssetsPath ) )
-			{
-				CopyFolder( FrameworkAssetsPath, AppDirectory, true );
-			}
+				// Let project override assets if they exist
+				if (Directory.Exists(InProjectDirectory + "/Intermediate/IOS/FrameworkAssets"))
+				{
+					FrameworkAssetsPath = InProjectDirectory + "/Intermediate/IOS/FrameworkAssets";
+				}
+
+				if (Directory.Exists(FrameworkAssetsPath))
+				{
+					CopyFolder(FrameworkAssetsPath, AppDirectory, true);
+				}
 
-			//CopyFiles(BuildDirectory, PayloadDirectory, null, "iTunesArtwork", null);
-			CopyFolder(InEngineDir + "/Content/Stats", AppDirectory + "/cookeddata/engine/content/stats", true);
+				Directory.CreateDirectory(CookedContentDirectory);
+				CopyFolder(InEngineDir + "/Content/Stats", AppDirectory + "/cookeddata/engine/content/stats", true);
+			}
 
 			return true;
 		}
