Copy additional bundle resources using RPC utility when building Mac targets on Windows.

diff --git a/Engine/Source/Programs/UnrealBuildTool/Mac/MacToolChain.cs b/Engine/Source/Programs/UnrealBuildTool/Mac/MacToolChain.cs
--- a/Engine/Source/Programs/UnrealBuildTool/Mac/MacToolChain.cs
+++ b/Engine/Source/Programs/UnrealBuildTool/Mac/MacToolChain.cs
@@ -1148,13 +1148,18 @@ namespace UnrealBuildTool
 
 			FileItem TargetItem = LocalToRemoteFileItem(FileItem.GetItemByPath(TargetPath), false);
 
-			CopyAction.CommandArguments = string.Format("-c 'cp -f -R \"{0}\" \"{1}\"; touch -c \"{2}\"'", SourcePath, Path.GetDirectoryName(TargetPath) + "/", TargetPath);
+			CopyAction.CommandArguments = string.Format("-c 'cp -f -R \"{0}\" \"{1}\"; touch -c \"{2}\"'", ConvertPath(SourcePath), Path.GetDirectoryName(TargetPath).Replace('\\', '/') + "/", TargetPath.Replace('\\', '/'));
 			CopyAction.PrerequisiteItems.Add(Executable);
 			CopyAction.ProducedItems.Add(TargetItem);
 			CopyAction.bShouldOutputStatusDescription = Resource.bShouldLog;
 			CopyAction.StatusDescription = string.Format("Copying {0} to app bundle", Path.GetFileName(Resource.ResourcePath));
 			CopyAction.bCanExecuteRemotely = false;
 
+			if (BuildHostPlatform.Current.Platform != UnrealTargetPlatform.Mac)
+			{
+				CopyAction.ActionHandler = new Action.BlockingActionHandler(RPCUtilHelper.RPCActionHandler);
+			}
+
 			if (Directory.Exists(Resource.ResourcePath))
 			{
 				foreach (string ResourceFile in Directory.GetFiles(Resource.ResourcePath, "*", SearchOption.AllDirectories))
@@ -1338,18 +1343,13 @@ namespace UnrealBuildTool
 		{
 			var OutputFiles = base.PostBuild(Executable, BinaryLinkEnvironment);
 
-			// If building for Mac on a Mac, use actions to finalize the builds (otherwise, we use Deploy)
-			if (BuildHostPlatform.Current.Platform != UnrealTargetPlatform.Mac)
-			{
-				return OutputFiles;
-			}
-
 			foreach (UEBuildBundleResource Resource in BinaryLinkEnvironment.Config.AdditionalBundleResources)
 			{
 				OutputFiles.Add(CopyBundleResource(Resource, Executable));
 			}
 
-			if (BinaryLinkEnvironment.Config.bIsBuildingDLL || BinaryLinkEnvironment.Config.bIsBuildingLibrary)
+			// If building for Mac on a Mac, use actions to finalize the builds (otherwise, we use Deploy)
+			if (BuildHostPlatform.Current.Platform != UnrealTargetPlatform.Mac || BinaryLinkEnvironment.Config.bIsBuildingDLL || BinaryLinkEnvironment.Config.bIsBuildingLibrary)
 			{
 				return OutputFiles;
 			}
