Fix some files not being copied in the remote tool chain. The path to the executable is always a remote path; converting it again makes it invalid.

diff --git a/Engine/Source/Programs/UnrealBuildTool/Mac/MacToolChain.cs b/Engine/Source/Programs/UnrealBuildTool/Mac/MacToolChain.cs
--- a/Engine/Source/Programs/UnrealBuildTool/Mac/MacToolChain.cs
+++ b/Engine/Source/Programs/UnrealBuildTool/Mac/MacToolChain.cs
@@ -1146,7 +1146,15 @@ namespace UnrealBuildTool
 			string SourcePath = Path.Combine(CopyAction.WorkingDirectory, Resource.ResourcePath);
 			string TargetPath = Path.Combine(BundlePath, "Contents", Resource.BundleContentsSubdir, Path.GetFileName(Resource.ResourcePath));
 
-			FileItem TargetItem = LocalToRemoteFileItem(FileItem.GetItemByPath(TargetPath), false);
+			FileItem TargetItem;
+			if(BuildHostPlatform.Current.Platform == UnrealTargetPlatform.Mac)
+			{
+				TargetItem = FileItem.GetItemByPath(TargetPath);
+			}
+			else
+			{
+				TargetItem = FileItem.GetRemoteItemByPath(TargetPath, RemoteToolChainPlatform);
+			}
 
 			CopyAction.CommandArguments = string.Format("-c 'cp -f -R \"{0}\" \"{1}\"; touch -c \"{2}\"'", ConvertPath(SourcePath), Path.GetDirectoryName(TargetPath).Replace('\\', '/') + "/", TargetPath.Replace('\\', '/'));
 			CopyAction.PrerequisiteItems.Add(Executable);
@@ -1349,7 +1357,12 @@ namespace UnrealBuildTool
 			}
 
 			// If building for Mac on a Mac, use actions to finalize the builds (otherwise, we use Deploy)
-			if (BuildHostPlatform.Current.Platform != UnrealTargetPlatform.Mac || BinaryLinkEnvironment.Config.bIsBuildingDLL || BinaryLinkEnvironment.Config.bIsBuildingLibrary)
+			if (BuildHostPlatform.Current.Platform != UnrealTargetPlatform.Mac)
+			{
+				return OutputFiles;
+			}
+
+			if (BinaryLinkEnvironment.Config.bIsBuildingDLL || BinaryLinkEnvironment.Config.bIsBuildingLibrary)
 			{
 				return OutputFiles;
 			}
