XP fix, moved app name override from additional definitions to its own target descriptor property

diff --git a/Engine/Source/Programs/UnrealBuildTool/Configuration/UEBuildTarget.cs b/Engine/Source/Programs/UnrealBuildTool/Configuration/UEBuildTarget.cs
--- a/Engine/Source/Programs/UnrealBuildTool/Configuration/UEBuildTarget.cs
+++ b/Engine/Source/Programs/UnrealBuildTool/Configuration/UEBuildTarget.cs
@@ -165,6 +165,7 @@ namespace UnrealBuildTool
 		public bool bIsEditorRecompile;
 		public string RemoteRoot;
 		public List<OnlyModule> OnlyModules;
+		public string OverridenAppName;
 	}
 
 
@@ -214,6 +215,7 @@ namespace UnrealBuildTool
 			var Configuration = UnrealTargetConfiguration.Unknown;
 			string RemoteRoot = null;
 			var OnlyModules = new List<OnlyModule>();
+			string OverridenAppName = null;
 
 			// If true, the recompile was launched by the editor.
 			bool bIsEditorRecompile = false;
@@ -233,7 +235,8 @@ namespace UnrealBuildTool
 				}
 				else if (Arguments[ArgumentIndex].ToLowerInvariant().StartsWith("-overridetargetappname="))
 				{
-					AdditionalDefinitions.Add(Arguments[ArgumentIndex]);
+					string OverrideTargetAppNameSwitch = "-overridetargetappname=";
+					OverridenAppName = Arguments[ArgumentIndex].Substring(OverrideTargetAppNameSwitch.Length);
 				}
 				else
 				{
@@ -462,6 +465,7 @@ namespace UnrealBuildTool
 			string RemoteRoot = Desc.RemoteRoot;
 			List<OnlyModule> OnlyModules = Desc.OnlyModules;
 			bool bIsEditorRecompile = Desc.bIsEditorRecompile;
+			string OverridenAppName = Desc.OverridenAppName;
 
 			UEBuildTarget BuildTarget = null;
 			if( !ProjectFileGenerator.bGenerateProjectFiles )
@@ -485,7 +489,8 @@ namespace UnrealBuildTool
 				InAdditionalDefinitions:AdditionalDefinitions, 
 				InRemoteRoot:RemoteRoot, 
 				InOnlyModules:OnlyModules, 
-				bInEditorRecompile:bIsEditorRecompile);
+				bInEditorRecompile:bIsEditorRecompile,
+				InOverridenAppName:OverridenAppName);
 			if (Target == null)
 			{
 				if (UEBuildConfiguration.bCleanProject)
@@ -763,21 +768,17 @@ namespace UnrealBuildTool
 			List<string> InAdditionalDefinitions,
 			string InRemoteRoot,
 			List<OnlyModule> InOnlyModules,
-			bool bInEditorRecompile)
+			bool bInEditorRecompile,
+			string InOverridenAppName = "" )
 		{
-			string CmdlineAppName = null;
-			const string OverrideTargetAppNameSwitch = "-overridetargetappname=";
-			if ((CmdlineAppName = InAdditionalDefinitions.Find(x => x.StartsWith(OverrideTargetAppNameSwitch))) != null)
-			{
-				OverridenAppName = CmdlineAppName.Substring(OverrideTargetAppNameSwitch.Length);
-			}
-
-				AppName = InAppName;
+			AppName = InAppName;
 			GameName = InGameName;
 			Platform = InPlatform;
 			Configuration = InConfiguration;
 			Rules = InRulesObject;
 			bEditorRecompile = bInEditorRecompile;
+			OverridenAppName = InOverridenAppName;
+
 
 			{
 				bCompileMonolithic = (Rules != null) ? Rules.ShouldCompileMonolithic(InPlatform, InConfiguration) : false;
diff --git a/Engine/Source/Programs/UnrealBuildTool/System/RulesCompiler.cs b/Engine/Source/Programs/UnrealBuildTool/System/RulesCompiler.cs
--- a/Engine/Source/Programs/UnrealBuildTool/System/RulesCompiler.cs
+++ b/Engine/Source/Programs/UnrealBuildTool/System/RulesCompiler.cs
@@ -1742,7 +1742,7 @@ namespace UnrealBuildTool
 		/// <param name="Target">Information about the target associated with this target</param>
 		/// <returns>The build target object for the specified build rules source file</returns>
 		public static UEBuildTarget CreateTarget(string TargetName, TargetInfo Target,
-			List<string> InAdditionalDefinitions, string InRemoteRoot, List<OnlyModule> InOnlyModules, bool bInEditorRecompile)
+			List<string> InAdditionalDefinitions, string InRemoteRoot, List<OnlyModule> InOnlyModules, bool bInEditorRecompile, string InOverridenAppName)
 		{
 			var CreateTargetStartTime = DateTime.UtcNow;
 
@@ -1830,7 +1830,8 @@ namespace UnrealBuildTool
 							InAdditionalDefinitions:InAdditionalDefinitions, 
 							InRemoteRoot:InRemoteRoot, 
 							InOnlyModules:InOnlyModules,
-							bInEditorRecompile:bInEditorRecompile);
+							bInEditorRecompile:bInEditorRecompile,
+							InOverridenAppName: InOverridenAppName);
 					}
 					break;
 			}
