Bit of housekeeping before adding list of dependencies into UBT's manifest:

* Renamed FileManifest to BuildManifest and FileManifestItems to BuildProducts, to more accurately reflect what they are.
* Removed ref qualifier when passing BuildManifest around; it's already a reference type, and not meant to be assigned to.
* Removed UAT's redundant copy of BuildManifest, since it already references the UBT assembly.

diff --git a/Engine/Source/Programs/AutomationTool/CommandUtils.cs b/Engine/Source/Programs/AutomationTool/CommandUtils.cs
--- a/Engine/Source/Programs/AutomationTool/CommandUtils.cs
+++ b/Engine/Source/Programs/AutomationTool/CommandUtils.cs
@@ -1390,32 +1390,14 @@ namespace AutomationTool
 			return Filename;
 		}
 
-		/// <summary>
-		/// A container for a binary files (dll, exe) with its associated debug info.
-		/// </summary>
-		public class FileManifest
-		{
-			/// <summary>
-			/// Items
-			/// </summary>
-			public readonly List<string> FileManifestItems = new List<string>();
-
-			/// <summary>
-			/// Constructor
-			/// </summary>
-			public FileManifest()
-			{
-			}
-		}
-
 		/// <summary>
 		/// Reads a file manifest and returns it
 		/// </summary>
 		/// <param name="ManifestName">ManifestName</param>
 		/// <returns></returns>
-		public static FileManifest ReadManifest(string ManifestName)
+		public static UnrealBuildTool.BuildManifest ReadManifest(string ManifestName)
 		{
-			return XmlHandler.ReadXml<FileManifest>(ManifestName);
+			return XmlHandler.ReadXml<UnrealBuildTool.BuildManifest>(ManifestName);
 		}
 
 		private static void CloneDirectoryRecursiveWorker(string SourcePathBase, string TargetPathBase, List<string> ClonedFiles)
diff --git a/Engine/Source/Programs/AutomationTool/Scripts/BuildCommonTools.Automation.cs b/Engine/Source/Programs/AutomationTool/Scripts/BuildCommonTools.Automation.cs
--- a/Engine/Source/Programs/AutomationTool/Scripts/BuildCommonTools.Automation.cs
+++ b/Engine/Source/Programs/AutomationTool/Scripts/BuildCommonTools.Automation.cs
@@ -68,10 +68,10 @@ public class BuildCommonTools : BuildCommand
 		string ManifestPath = ParseParamValue("manifest");
 		if(ManifestPath != null)
 		{
-			UnrealBuildTool.FileManifest Manifest = new UnrealBuildTool.FileManifest();
+			BuildManifest Manifest = new BuildManifest();
 			foreach(string BuildProductFile in Builder.BuildProductFiles)
 			{
-				Manifest.AddFileName(BuildProductFile);
+				Manifest.AddBuildProduct(BuildProductFile);
 			}
 			UnrealBuildTool.Utils.WriteClass(Manifest, ManifestPath, "");
 		}
diff --git a/Engine/Source/Programs/AutomationTool/UE4Build.cs b/Engine/Source/Programs/AutomationTool/UE4Build.cs
--- a/Engine/Source/Programs/AutomationTool/UE4Build.cs
+++ b/Engine/Source/Programs/AutomationTool/UE4Build.cs
@@ -76,8 +76,8 @@ namespace AutomationTool
 			Log("Copied UBT manifest to {0}", OutFile);
 
 
-			FileManifest Manifest = ReadManifest(ManifestName);
-			foreach (string Item in Manifest.FileManifestItems)
+			UnrealBuildTool.BuildManifest Manifest = ReadManifest(ManifestName);
+			foreach (string Item in Manifest.BuildProducts)
 			{
 				PrepareBuildProduct(Item);
 			}
@@ -89,8 +89,8 @@ namespace AutomationTool
 			{
 				throw new AutomationException("BUILD FAILED UBT Manifest {0} does not exist.", ManifestName);
 			}
-			FileManifest Manifest = ReadManifest(ManifestName);
-			foreach (string Item in Manifest.FileManifestItems)
+			UnrealBuildTool.BuildManifest Manifest = ReadManifest(ManifestName);
+			foreach (string Item in Manifest.BuildProducts)
 			{
 				if (!FileExists_NoExceptions(Item))
 				{
@@ -130,7 +130,7 @@ namespace AutomationTool
 
 		public class XGEItem
 		{
-			public FileManifest Manifest;
+			public UnrealBuildTool.BuildManifest Manifest;
 			public string CommandLine;
 			public UnrealBuildTool.UnrealTargetPlatform Platform;
 			public string Config;
@@ -229,7 +229,7 @@ namespace AutomationTool
 			// allow the platform to perform any actions after building a target (seems almost like this should be done in UBT)
 			Platform.Platforms[Item.Platform].PostBuildTarget(this, Item.ProjectName, Item.UProjectPath, Item.Config);
 
-			foreach (string ManifestItem in Item.Manifest.FileManifestItems)
+			foreach (string ManifestItem in Item.Manifest.BuildProducts)
 			{
 				if (!FileExists_NoExceptions(ManifestItem))
 				{
@@ -239,9 +239,9 @@ namespace AutomationTool
 			}
 		}
 
-		void XGEDeleteBuildProducts(FileManifest Manifest)
+		void XGEDeleteBuildProducts(UnrealBuildTool.BuildManifest Manifest)
 		{
-			foreach (string Item in Manifest.FileManifestItems)
+			foreach (string Item in Manifest.BuildProducts)
 			{
 				DeleteFile(Item);
 			}
diff --git a/Engine/Source/Programs/UnrealBuildTool/Android/AndroidToolChain.cs b/Engine/Source/Programs/UnrealBuildTool/Android/AndroidToolChain.cs
--- a/Engine/Source/Programs/UnrealBuildTool/Android/AndroidToolChain.cs
+++ b/Engine/Source/Programs/UnrealBuildTool/Android/AndroidToolChain.cs
@@ -940,13 +940,13 @@ namespace UnrealBuildTool
 			return Outputs.ToArray();
 		}
 
-		public override void AddFilesToManifest(ref FileManifest Manifest, UEBuildBinary Binary)
+		public override void AddFilesToManifest(BuildManifest Manifest, UEBuildBinary Binary)
 		{
 			// the binary will have all of the .so's in the output files, we need to trim down to the shared apk (which is what needs to go into the manifest)
 			foreach (string BinaryPath in Binary.Config.OutputFilePaths)
 			{
 				string ApkFile = Path.ChangeExtension(BinaryPath, ".apk");
-				Manifest.AddFileName(ApkFile);
+				Manifest.AddBuildProduct(ApkFile);
 			}
 		}
 
diff --git a/Engine/Source/Programs/UnrealBuildTool/Configuration/UEBuildTarget.cs b/Engine/Source/Programs/UnrealBuildTool/Configuration/UEBuildTarget.cs
--- a/Engine/Source/Programs/UnrealBuildTool/Configuration/UEBuildTarget.cs
+++ b/Engine/Source/Programs/UnrealBuildTool/Configuration/UEBuildTarget.cs
@@ -95,30 +95,29 @@ namespace UnrealBuildTool
 	/// <summary>
 	/// A container for a binary files (dll, exe) with its associated debug info.
 	/// </summary>
-	public class FileManifest
+	public class BuildManifest
 	{
-		public readonly List<string> FileManifestItems = new List<string>();
+		public readonly List<string> BuildProducts = new List<string>();
 
-		public FileManifest()
+		public BuildManifest()
 		{
 		}
 
-		public void AddFileName( string FilePath )
+		public void AddBuildProduct(string FileName)
 		{
-			FileManifestItems.Add( Path.GetFullPath( FilePath ) );
+			string FullFileName = Path.GetFullPath(FileName);
+			if (!BuildProducts.Contains(FullFileName))
+			{
+				BuildProducts.Add(FullFileName);
+			}
 		}
 
-		public void AddBinaryNames(string OutputFilePath, string DebugInfoExtension)
+		public void AddBuildProduct(string FileName, string DebugInfoExtension)
 		{
-			// Only add unique files to the manifest. Multiple games have many shared binaries.
-			if( !FileManifestItems.Any( x => Path.GetFullPath( OutputFilePath ) == x ) )
+			AddBuildProduct(FileName);
+			if(!String.IsNullOrEmpty(DebugInfoExtension))
 			{
-				string FullPath = Path.GetFullPath(OutputFilePath);
-				FileManifestItems.Add( FullPath );
-				if (!string.IsNullOrEmpty(DebugInfoExtension))
-				{
-					FileManifestItems.Add(Path.ChangeExtension(FullPath, DebugInfoExtension));
-				}
+				AddBuildProduct(Path.ChangeExtension(FileName, DebugInfoExtension));
 			}
 		}
 	}
@@ -992,7 +991,7 @@ namespace UnrealBuildTool
 		/// <param name="Binaries">Target binaries</param>
 		/// <param name="Platform">Tareet platform</param>
 		/// <param name="Manifest">Manifest</param>
-		protected void CleanTarget(List<UEBuildBinary> Binaries, CPPTargetPlatform Platform, FileManifest Manifest)
+		protected void CleanTarget(List<UEBuildBinary> Binaries, CPPTargetPlatform Platform, BuildManifest Manifest)
 		{
 			{
 				var TargetFilename = RulesCompiler.GetTargetFilename(GameName);
@@ -1015,8 +1014,8 @@ namespace UnrealBuildTool
 
 				// Collect all files to delete.
 				var AdditionalFileExtensions = new string[] { ".lib", ".exp", ".dll.response" };
-				var AllFilesToDelete = new List<string>(Manifest.FileManifestItems);
-				foreach (var FileManifestItem in Manifest.FileManifestItems)
+				var AllFilesToDelete = new List<string>(Manifest.BuildProducts);
+				foreach (var FileManifestItem in Manifest.BuildProducts)
 				{
 					var FileExt = Path.GetExtension(FileManifestItem);
 					if (FileExt == ".dll" || FileExt == ".exe")
@@ -1316,11 +1315,11 @@ namespace UnrealBuildTool
 				ManifestPath = "../Intermediate/Build/Manifest.xml";
 			}
 
-			FileManifest Manifest = new FileManifest();
+			BuildManifest Manifest = new BuildManifest();
 			if (UEBuildConfiguration.bMergeManifests)
 			{
 				// Load in existing manifest (if any)
-				Manifest = Utils.ReadClass<FileManifest>(ManifestPath);
+				Manifest = Utils.ReadClass<BuildManifest>(ManifestPath);
 			}
 
 			UnrealTargetPlatform TargetPlatform = CPPTargetPlatformToUnrealTargetPlatform( Platform );
@@ -1347,7 +1346,7 @@ namespace UnrealBuildTool
 				// Create and add the binary and associated debug info
 				foreach (string OutputFilePath in Binary.Config.OutputFilePaths)
 				{
-					Manifest.AddBinaryNames(OutputFilePath, DebugInfoExtension);
+					Manifest.AddBuildProduct(OutputFilePath, DebugInfoExtension);
 				}
 
 				if (Binary.Config.Type == UEBuildBinaryType.Executable &&
@@ -1356,17 +1355,17 @@ namespace UnrealBuildTool
 				{
 					foreach (string OutputFilePath in Binary.Config.OutputFilePaths)
 					{
-						Manifest.AddBinaryNames(UEBuildBinary.GetAdditionalConsoleAppPath(OutputFilePath), DebugInfoExtension);
+						Manifest.AddBuildProduct(UEBuildBinary.GetAdditionalConsoleAppPath(OutputFilePath), DebugInfoExtension);
 					}
 				}
 
-                ToolChain.AddFilesToManifest(ref Manifest,Binary);
+                ToolChain.AddFilesToManifest(Manifest, Binary);
 			}
 			{
 				string DebugInfoExtension = BuildPlatform.GetDebugInfoExtension(UEBuildBinaryType.StaticLibrary);
 				foreach (var RedistLib in SpecialRocketLibFilesThatAreBuildProducts)
 				{
-					Manifest.AddBinaryNames(RedistLib, DebugInfoExtension);
+					Manifest.AddBuildProduct(RedistLib, DebugInfoExtension);
 				}
 			}
 
@@ -1377,7 +1376,7 @@ namespace UnrealBuildTool
 			}
 			if (UEBuildConfiguration.bGenerateManifest)
 			{
-				Utils.WriteClass<FileManifest>(Manifest, ManifestPath, "");
+				Utils.WriteClass<BuildManifest>(Manifest, ManifestPath, "");
 			}
 		}
 
diff --git a/Engine/Source/Programs/UnrealBuildTool/HTML5/HTML5ToolChain.cs b/Engine/Source/Programs/UnrealBuildTool/HTML5/HTML5ToolChain.cs
--- a/Engine/Source/Programs/UnrealBuildTool/HTML5/HTML5ToolChain.cs
+++ b/Engine/Source/Programs/UnrealBuildTool/HTML5/HTML5ToolChain.cs
@@ -502,12 +502,10 @@ namespace UnrealBuildTool
 			throw new BuildException("HTML5 cannot compile C# files");
 		}
 
-        public override void AddFilesToManifest(ref FileManifest Manifest, UEBuildBinary Binary)
+        public override void AddFilesToManifest(BuildManifest Manifest, UEBuildBinary Binary)
         {
             // we need to include the generated .mem file.  
-            string MemFile = Binary.Config.OutputFilePath + ".mem";
-            // make sure we don't add mem files more than once. 
-            Manifest.AddBinaryNames(MemFile, null);
+            Manifest.AddBuildProduct(Binary.Config.OutputFilePath + ".mem");
         }
 
 		public override UnrealTargetPlatform GetPlatform()
diff --git a/Engine/Source/Programs/UnrealBuildTool/IOS/IOSToolChain.cs b/Engine/Source/Programs/UnrealBuildTool/IOS/IOSToolChain.cs
--- a/Engine/Source/Programs/UnrealBuildTool/IOS/IOSToolChain.cs
+++ b/Engine/Source/Programs/UnrealBuildTool/IOS/IOSToolChain.cs
@@ -276,12 +276,12 @@ namespace UnrealBuildTool
 			}
 		}
 
-		public override void AddFilesToManifest(ref FileManifest Manifest, UEBuildBinary Binary)
+		public override void AddFilesToManifest(BuildManifest Manifest, UEBuildBinary Binary)
 		{
 			if (BuildConfiguration.bCreateStubIPA)
 			{
 				string StubFile = Path.Combine (Path.GetDirectoryName (Binary.Config.OutputFilePath), Path.GetFileNameWithoutExtension (Binary.Config.OutputFilePath) + ".stub");
-				Manifest.AddFileName (StubFile);
+				Manifest.AddBuildProduct(StubFile);
 			}
 		}
 
diff --git a/Engine/Source/Programs/UnrealBuildTool/Mac/MacToolChain.cs b/Engine/Source/Programs/UnrealBuildTool/Mac/MacToolChain.cs
--- a/Engine/Source/Programs/UnrealBuildTool/Mac/MacToolChain.cs
+++ b/Engine/Source/Programs/UnrealBuildTool/Mac/MacToolChain.cs
@@ -1228,7 +1228,7 @@ namespace UnrealBuildTool
 
 		static private string BundleContentsDirectory = "";
 
-        public override void AddFilesToManifest(ref FileManifest Manifest, UEBuildBinary Binary)
+        public override void AddFilesToManifest(BuildManifest Manifest, UEBuildBinary Binary)
 		{
 			if (Binary.Target.GlobalLinkEnvironment.Config.bIsBuildingConsoleApplication)
 			{
@@ -1258,10 +1258,7 @@ namespace UnrealBuildTool
 					if (Path.GetExtension(AdditionalLibrary) == ".dylib")
 					{
 						string Entry = BundleContentsDirectory + "MacOS/" + LibName;
-						if (!Manifest.FileManifestItems.Contains(Path.GetFullPath(Entry)))
-						{
-							Manifest.AddFileName(Entry);
-						}
+						Manifest.AddBuildProduct(Entry);
 					}
 				}
 			}
@@ -1272,25 +1269,25 @@ namespace UnrealBuildTool
 				{
 					foreach (string ResourceFile in Directory.GetFiles(Resource.ResourcePath, "*", SearchOption.AllDirectories))
 					{
-						Manifest.AddFileName(Path.Combine(BundleContentsDirectory, Resource.BundleContentsSubdir, ResourceFile.Substring(Path.GetDirectoryName(Resource.ResourcePath).Length + 1)));
+						Manifest.AddBuildProduct(Path.Combine(BundleContentsDirectory, Resource.BundleContentsSubdir, ResourceFile.Substring(Path.GetDirectoryName(Resource.ResourcePath).Length + 1)));
 					}
 				}
 				else
 				{
-					Manifest.AddFileName(Path.Combine(BundleContentsDirectory, Resource.BundleContentsSubdir, Path.GetFileName(Resource.ResourcePath)));
+					Manifest.AddBuildProduct(Path.Combine(BundleContentsDirectory, Resource.BundleContentsSubdir, Path.GetFileName(Resource.ResourcePath)));
 				}
 			}
 
 			if (Binary.Config.Type == UEBuildBinaryType.Executable)
 			{
 				// And we also need all the resources
-				Manifest.AddFileName(BundleContentsDirectory + "Info.plist");
-				Manifest.AddFileName(BundleContentsDirectory + "PkgInfo");
-				Manifest.AddFileName(BundleContentsDirectory + "Resources/UE4.icns");
+				Manifest.AddBuildProduct(BundleContentsDirectory + "Info.plist");
+				Manifest.AddBuildProduct(BundleContentsDirectory + "PkgInfo");
+				Manifest.AddBuildProduct(BundleContentsDirectory + "Resources/UE4.icns");
 
 				if (Binary.Config.TargetName.StartsWith("UE4Editor"))
 				{
-					Manifest.AddFileName(BundleContentsDirectory + "Resources/UProject.icns");
+					Manifest.AddBuildProduct(BundleContentsDirectory + "Resources/UProject.icns");
 				}
 			}
 		}
diff --git a/Engine/Source/Programs/UnrealBuildTool/System/ProjectFileGenerator.cs b/Engine/Source/Programs/UnrealBuildTool/System/ProjectFileGenerator.cs
--- a/Engine/Source/Programs/UnrealBuildTool/System/ProjectFileGenerator.cs
+++ b/Engine/Source/Programs/UnrealBuildTool/System/ProjectFileGenerator.cs
@@ -1981,17 +1981,17 @@ namespace UnrealBuildTool
 		/// </summary>
 		protected virtual bool WriteProjectFileManifest()
 		{
-			FileManifest Manifest = new FileManifest();
+			BuildManifest Manifest = new BuildManifest();
 			foreach( var CurProject in GeneratedProjectFiles )
 			{
 				foreach( var SourceFile in CurProject.SourceFiles )
 				{
-					Manifest.AddFileName( SourceFile.FilePath );
+					Manifest.AddBuildProduct( SourceFile.FilePath );
 				}
 			}
 
 			string ManifestName = Path.Combine( ProjectFileGenerator.IntermediateProjectFilesPath, "UE4SourceFiles.xml" );
-			Utils.WriteClass<FileManifest>( Manifest, ManifestName, "" );
+			Utils.WriteClass<BuildManifest>( Manifest, ManifestName, "" );
 			return true;
 		}
 
diff --git a/Engine/Source/Programs/UnrealBuildTool/ToolChain/UEToolChain.cs b/Engine/Source/Programs/UnrealBuildTool/ToolChain/UEToolChain.cs
--- a/Engine/Source/Programs/UnrealBuildTool/ToolChain/UEToolChain.cs
+++ b/Engine/Source/Programs/UnrealBuildTool/ToolChain/UEToolChain.cs
@@ -38,7 +38,7 @@ namespace UnrealBuildTool
 
 		void SetUpGlobalEnvironment();
 
-        void AddFilesToManifest(ref FileManifest manifest, UEBuildBinary Binary );
+        void AddFilesToManifest(BuildManifest Manifest, UEBuildBinary Binary);
 
 		void SetupBundleDependencies(List<UEBuildBinary> Binaries, string GameName);
 
@@ -253,7 +253,7 @@ namespace UnrealBuildTool
 
 		}
 
-        public virtual void AddFilesToManifest(ref FileManifest manifest, UEBuildBinary Binary )
+        public virtual void AddFilesToManifest(BuildManifest Manifest, UEBuildBinary Binary)
         {
 
         }
diff --git a/Engine/Source/Programs/UnrealBuildTool/Windows/VCToolChain.cs b/Engine/Source/Programs/UnrealBuildTool/Windows/VCToolChain.cs
--- a/Engine/Source/Programs/UnrealBuildTool/Windows/VCToolChain.cs
+++ b/Engine/Source/Programs/UnrealBuildTool/Windows/VCToolChain.cs
@@ -1656,7 +1656,7 @@ namespace UnrealBuildTool
 			return IncludePaths;
 		}
 
-        public override void AddFilesToManifest(ref FileManifest Manifest, UEBuildBinary Binary)
+        public override void AddFilesToManifest(BuildManifest Manifest, UEBuildBinary Binary)
         {
             // ok, this is pretty awful, we want the import libraries that go with the editor, only on the PC
             if (UnrealBuildTool.BuildingRocket() &&
@@ -1665,7 +1665,7 @@ namespace UnrealBuildTool
                 Binary.Config.Type == UEBuildBinaryType.DynamicLinkLibrary)
             {
                 // ok, this is pretty awful, we want the import libraries that go with the editor, only on the PC
-                Manifest.AddBinaryNames(Path.Combine(Binary.Config.IntermediateDirectory, Path.GetFileNameWithoutExtension(Binary.Config.OutputFilePath) + ".lib"), "");
+                Manifest.AddBuildProduct(Path.Combine(Binary.Config.IntermediateDirectory, Path.GetFileNameWithoutExtension(Binary.Config.OutputFilePath) + ".lib"), "");
             }
         }
 
