#HTML5

   Force ThirdParty libs to be linked in always last. Linker may refuse to link in archives if they are linked in before the linker encounters dependent symbols. (fixes unresolved symbols in html5 rocket build )

diff --git a/Engine/Source/Programs/UnrealBuildTool/HTML5/HTML5ToolChain.cs b/Engine/Source/Programs/UnrealBuildTool/HTML5/HTML5ToolChain.cs
--- a/Engine/Source/Programs/UnrealBuildTool/HTML5/HTML5ToolChain.cs
+++ b/Engine/Source/Programs/UnrealBuildTool/HTML5/HTML5ToolChain.cs
@@ -465,22 +465,30 @@ namespace UnrealBuildTool
 			}
 			if (!LinkEnvironment.Config.bIsBuildingLibrary)
 			{
-				foreach (string InputFile in LinkEnvironment.Config.AdditionalLibraries)
-				{
-					FileItem Item = FileItem.GetItemByPath(InputFile);
-
-					if (Item.AbsolutePath.Contains(".lib"))
-						continue;
-
-					if (Item != null)
-					{
-						if (Item.ToString().Contains(".js"))
-							LinkAction.CommandArguments += string.Format(" --js-library \"{0}\"", Item.AbsolutePath);
-						else
-							LinkAction.CommandArguments += string.Format(" \"{0}\"", Item.AbsolutePath);
-						LinkAction.PrerequisiteItems.Add(Item);
-					}
-				}
+                    // Make sure ThirdParty libs are at the end. 
+                    List<string> ThirdParty = (from Lib in LinkEnvironment.Config.AdditionalLibraries
+                                     where Lib.Contains("ThirdParty")
+                                     select Lib).ToList();
+
+                    LinkEnvironment.Config.AdditionalLibraries.RemoveAll(Element => Element.Contains("ThirdParty"));
+                    LinkEnvironment.Config.AdditionalLibraries.AddRange(ThirdParty);
+
+                    foreach (string InputFile in LinkEnvironment.Config.AdditionalLibraries)
+                    {
+                        FileItem Item = FileItem.GetItemByPath(InputFile);
+
+                        if (Item.AbsolutePath.Contains(".lib"))
+                            continue;
+
+                        if (Item != null)
+                        {
+                            if (Item.ToString().Contains(".js"))
+                                LinkAction.CommandArguments += string.Format(" --js-library \"{0}\"", Item.AbsolutePath);
+                            else
+                                LinkAction.CommandArguments += string.Format(" \"{0}\"", Item.AbsolutePath);
+                            LinkAction.PrerequisiteItems.Add(Item);
+                        }
+                    }
 			}
 			// make the file we will create
 			OutputFile = FileItem.GetItemByPath(LinkEnvironment.Config.OutputFilePath);
