Fixed command-line spew when generating projects with UBTMakefiles enabled

diff --git a/Engine/Source/Programs/UnrealBuildTool/System/FlatCPPIncludeDepencencyCache.cs b/Engine/Source/Programs/UnrealBuildTool/System/FlatCPPIncludeDepencencyCache.cs
--- a/Engine/Source/Programs/UnrealBuildTool/System/FlatCPPIncludeDepencencyCache.cs
+++ b/Engine/Source/Programs/UnrealBuildTool/System/FlatCPPIncludeDepencencyCache.cs
@@ -67,9 +67,13 @@ namespace UnrealBuildTool
 				}
 			}
 
+			if( !ProjectFileGenerator.bGenerateProjectFiles )
+			{ 
+				UnrealBuildTool.bNeedsFullCPPIncludeRescan = true;
+				Log.TraceInformation( "Performing full C++ include scan (no include cache file)" );
+			}
+
 			// Fall back to a clean cache on error or non-existence.
-			UnrealBuildTool.bNeedsFullCPPIncludeRescan = true;
-			Log.TraceInformation( "Performing full C++ include scan (no include cache file)" );
 			return new FlatCPPIncludeDependencyCache( Cache );
 		}
 
