Fixed UnrealBuildTool not able to hot reload a newly-created editor module when UBT Makefiles were enabled
- UBT was expecting a random-numbered suffix on all hot reloaded module file names, but this was not valid for newly-created modules in the editor

diff --git a/Engine/Source/Programs/UnrealBuildTool/System/UnrealBuildTool.cs b/Engine/Source/Programs/UnrealBuildTool/System/UnrealBuildTool.cs
--- a/Engine/Source/Programs/UnrealBuildTool/System/UnrealBuildTool.cs
+++ b/Engine/Source/Programs/UnrealBuildTool/System/UnrealBuildTool.cs
@@ -2441,9 +2441,25 @@ namespace UnrealBuildTool
 					// Assume that the first produced item (with no extension) is our output file name
 					var OriginalFileNameWithoutExtension = Utils.GetFilenameWithoutAnyExtensions( Action.ProducedItems[0].AbsolutePath );
 
-					// Remove the numbered suffix from the end of the file
-					var OriginalNumberSuffix = OriginalFileNameWithoutExtension.Substring( OriginalFileNameWithoutExtension.LastIndexOf( '-' ) );
-					var OriginalFileNameWithoutNumberSuffix = OriginalFileNameWithoutExtension.Substring( 0, OriginalFileNameWithoutExtension.Length - OriginalNumberSuffix.Length );	// Remove "-####"
+					// Figure out if we have a numbered suffix at the end of the output file name, which will be present for already-existing
+					// module that was previously hot reloaded.  However, if the editor asked UnrealBuildTool to compile a brand new module
+					// and hot reload it, the output file will not have a numbered suffixed extension.
+					string OriginalFileNameWithoutNumberSuffix = OriginalFileNameWithoutExtension;
+					{
+						int IndexOfNumberedSuffix = OriginalFileNameWithoutExtension.LastIndexOf( '-' );
+						if( IndexOfNumberedSuffix != -1 )
+						{ 
+							var OriginalNumberSuffix = OriginalFileNameWithoutExtension.Substring( IndexOfNumberedSuffix );
+
+							// Make sure the number we pulled off the end of the string was actually a number (e.g. "-3141")
+							int TestParsedInt;
+							if( int.TryParse( OriginalNumberSuffix, out TestParsedInt ) )
+							{ 
+								// Remove the numbered suffix from the end of the file
+								OriginalFileNameWithoutNumberSuffix = OriginalFileNameWithoutExtension.Substring( 0, OriginalFileNameWithoutExtension.Length - OriginalNumberSuffix.Length );	// Remove "-####"
+							}
+						}
+					}
 
 					// Figure out which suffix to use
 					string UniqueSuffix = null;
