Include engine source in Rocket projects (but don't build anything).

diff --git a/Engine/Source/Programs/UnrealBuildTool/System/ProjectFileGenerator.cs b/Engine/Source/Programs/UnrealBuildTool/System/ProjectFileGenerator.cs
--- a/Engine/Source/Programs/UnrealBuildTool/System/ProjectFileGenerator.cs
+++ b/Engine/Source/Programs/UnrealBuildTool/System/ProjectFileGenerator.cs
@@ -116,10 +116,6 @@ namespace UnrealBuildTool
 		/// Optional list of platforms to generate projects for
 		readonly List<UnrealTargetPlatform> ProjectPlatforms = new List<UnrealTargetPlatform>();
 
-		/// Whether engine module projects should be included in single-game projects, or Rocket projects.  This is useful if you want to be able to
-		/// quickly locate engine header files using VAX, for example.  But it clutters up the solution file.
-		protected bool bAlwaysIncludeEngineModules = false;
-
 		/// When bGeneratingGameProjectFiles=true, this is the game name we're generating projects for
 		protected string GameProjectName = null;
 
@@ -145,6 +141,9 @@ namespace UnrealBuildTool
 		/// True if all documentation languages should be included in generated projects, otherwise only "INT" will be included
 		bool bAllDocumentationLanguages = false;
 
+		/// True if we should include engine source in the generated solution
+		protected bool IncludeEngineSource = true;
+
 		/// True if shader source files should be included in the generated projects
 		bool IncludeShaderSource = true;
 
@@ -160,6 +159,9 @@ namespace UnrealBuildTool
         /// True if we should include template files (.template files) in the generated project
         bool IncludeTemplateFiles = true;
 
+		/// True if we should include program projects in the generated solution
+		protected bool IncludeEnginePrograms = true;
+
 		/// True if we should reflect "Source" sub-directories on disk in the master project as master project directories.
 		/// This arguably adds some visual clutter to the master project, but is truer to the on-disk file organization.
 		bool KeepSourceSubDirectories = true;
@@ -290,12 +292,6 @@ namespace UnrealBuildTool
 					}
 				}
 				IntermediateProjectFilesPath = Path.Combine(MasterProjectRelativePath, "Intermediate", "ProjectFiles");
-
-				IncludeDocumentation = false;
-				IncludeBuildSystemFiles = false;
-				IncludeShaderSource = true;
-				IncludeTemplateFiles = false;
-				IncludeConfigFiles = true;
 			}
 			else if( bGeneratingRocketProjectFiles )
 			{
@@ -316,12 +312,6 @@ namespace UnrealBuildTool
 				{
 					throw new BuildException("Directory '{0}' is missing 'Source' folder.", MasterProjectRelativePath);
 				}
-
-				IncludeDocumentation = false;
-				IncludeBuildSystemFiles = false;
-				IncludeShaderSource = false;
-				IncludeTemplateFiles = false;
-				IncludeConfigFiles = true;
 			}
 
 			// Modify the name if specific platforms were given
@@ -377,12 +367,6 @@ namespace UnrealBuildTool
 			}
 			RulesCompiler.SetAssemblyNameAndGameFolders( AssemblyName, AssemblyGameFolders );
 
-			MasterProjectFolder ProgramsFolder = RootFolder;
-			if( ( !bGeneratingGameProjectFiles && !bGeneratingRocketProjectFiles ) || bAlwaysIncludeEngineModules )
-			{
-				ProgramsFolder = RootFolder.AddSubFolder( "Programs" );
-			}
-
 			ProjectFile EngineProject = null;
 			Dictionary<string, ProjectFile> GameProjects = null;
 			Dictionary<string, ProjectFile> ProgramProjects = null;
@@ -401,7 +385,7 @@ namespace UnrealBuildTool
 				}
 
 				// Place projects into root level solution folders
-				if( ( !bGeneratingGameProjectFiles && !bGeneratingRocketProjectFiles ) || bAlwaysIncludeEngineModules )
+				if( IncludeEngineSource )
 				{
 					if( EngineProject != null )
 					{
@@ -411,8 +395,11 @@ namespace UnrealBuildTool
 						if( IncludeConfigFiles )
 						{
 							AddEngineConfigFiles( EngineProject );
-							AddUnrealHeaderToolConfigFiles(EngineProject);
-							AddUBTConfigFilesToEngineProject(EngineProject);
+							if( IncludeEnginePrograms )
+							{
+								AddUnrealHeaderToolConfigFiles(EngineProject);
+								AddUBTConfigFilesToEngineProject(EngineProject);
+							}
 						}
 
 						// Engine localization files
@@ -463,7 +450,7 @@ namespace UnrealBuildTool
 
 					foreach( var CurProgramProject in ProgramProjects.Values )
 					{
-						ProgramsFolder.ChildProjects.Add( CurProgramProject );
+						RootFolder.AddSubFolder( "Programs" ).ChildProjects.Add( CurProgramProject );
 					}
 
 					// Add all of the config files for generated program targets
@@ -479,14 +466,13 @@ namespace UnrealBuildTool
 			AddProjectsForAllModules(AllGameProjects, ProgramProjects, AllModuleFiles, bGatherThirdPartySource);
 
 			{
-				if( !bGeneratingRocketProjectFiles )
+				if( IncludeEnginePrograms )
 				{
+					MasterProjectFolder ProgramsFolder = RootFolder.AddSubFolder( "Programs" );
+
 					// Add UnrealBuildTool to the master project
 					AddUnrealBuildToolProject( ProgramsFolder );
-				}
 
-				if( ( !bGeneratingGameProjectFiles && !bGeneratingRocketProjectFiles ) || bAlwaysIncludeEngineModules )
-				{
 					// Add AutomationTool to the master project
 					ProgramsFolder.ChildProjects.Add(AddSimpleCSharpProject("AutomationTool", bShouldBuildForAllSolutionTargets: true, bForceDevelopmentConfiguration: true));
 
@@ -531,7 +517,7 @@ namespace UnrealBuildTool
 					var IntelliSenseTargetFiles = new List<Tuple<ProjectFile, string>>();
 					{
 						// Engine targets
-						if( EngineProject != null )
+						if( EngineProject != null && !bGeneratingRocketProjectFiles )
 						{
 							foreach( var ProjectTarget in EngineProject.ProjectTargets )
 							{
@@ -643,6 +629,7 @@ namespace UnrealBuildTool
 		/// <param name="IncludeAllPlatforms">True if all platforms should be included</param>
 		protected virtual void ConfigureProjectFileGeneration( String[] Arguments, ref bool IncludeAllPlatforms )
 		{
+			bool bAlwaysIncludeEngineModules = false;
 			foreach( var CurArgument in Arguments )
 			{
 				if( CurArgument.StartsWith( "-" ) )
@@ -706,13 +693,11 @@ namespace UnrealBuildTool
 						
 						case "-GAME":
 							// Generates project files for a single game
-							// @todo rocketprojects: Needs documentation
 							bGeneratingGameProjectFiles = true;
 							break;
 
 						case "-ENGINE":
 							// Forces engine modules and targets to be included in game-specific project files
-							// @todo rocketprojects: Needs documentation
 							bAlwaysIncludeEngineModules = true;
 							break;
 
@@ -756,10 +741,18 @@ namespace UnrealBuildTool
 			if( bGeneratingGameProjectFiles )
 			{
 				GameProjectName = Path.GetFileNameWithoutExtension(UnrealBuildTool.GetUProjectFile());
-					if (String.IsNullOrEmpty(GameProjectName))
-					{
+				if (String.IsNullOrEmpty(GameProjectName))
+				{
 					throw new BuildException("A valid game project was not found in the specified location (" + UnrealBuildTool.GetUProjectPath() + ")");
-					}
+				}
+
+				IncludeEngineSource = bAlwaysIncludeEngineModules;
+				IncludeDocumentation = false;
+				IncludeBuildSystemFiles = false;
+				IncludeShaderSource = true;
+				IncludeTemplateFiles = false;
+				IncludeConfigFiles = true;
+				IncludeEnginePrograms = bAlwaysIncludeEngineModules;
 			}
 			else if( bGeneratingRocketProjectFiles )
 			{
@@ -775,6 +768,14 @@ namespace UnrealBuildTool
 				{
 					throw new BuildException("A valid Rocket game project was not found in the specified location (" + UnrealBuildTool.GetUProjectPath() + ")");
 				}
+
+				IncludeEngineSource = true;
+				IncludeDocumentation = false;
+				IncludeBuildSystemFiles = false;
+				IncludeShaderSource = false;
+				IncludeTemplateFiles = false;
+				IncludeConfigFiles = true;
+				IncludeEnginePrograms = false;
 			}
 			else
 			{
@@ -1472,13 +1473,10 @@ namespace UnrealBuildTool
 				{
 					if( IsEngineModule )
 					{
-						if( bAlwaysIncludeEngineModules )
+						if( IncludeEngineSource )
 						{
-							if( bGeneratingRocketProjectFiles )
-							{
-								// This is an engine module, so strip out private code from the Rocket projects.  Rocket users only need public headers to compile
-								IncludePrivateSourceCode = false;
-							}
+							// This is an engine module, so strip out private code from the Rocket projects.  Rocket users only need public headers to compile
+							IncludePrivateSourceCode &= !bGeneratingRocketProjectFiles;
 						}
 						else
 						{
@@ -1611,16 +1609,13 @@ namespace UnrealBuildTool
 				}
 
 				bool WantProjectFileForTarget = true;
-				if( bGeneratingGameProjectFiles || bGeneratingRocketProjectFiles )
+				if(TargetFileRelativeToEngineDirectory.StartsWith(Path.Combine("Source", "Programs"), StringComparison.InvariantCultureIgnoreCase))
 				{
-					if( IsEngineTarget )
-					{
-						if( !bAlwaysIncludeEngineModules )
-						{
-							// We were asked to exclude engine modules from the generated projects
-							WantProjectFileForTarget = false;
-						}
-					}
+					WantProjectFileForTarget = IncludeEnginePrograms;
+				}
+				else if(TargetFileRelativeToEngineDirectory.StartsWith(Path.Combine("Source"), StringComparison.InvariantCultureIgnoreCase))
+				{
+					WantProjectFileForTarget = IncludeEngineSource;
 				}
 
 				if (WantProjectFileForTarget)
diff --git a/Engine/Source/Programs/UnrealBuildTool/System/VCProject.cs b/Engine/Source/Programs/UnrealBuildTool/System/VCProject.cs
--- a/Engine/Source/Programs/UnrealBuildTool/System/VCProject.cs
+++ b/Engine/Source/Programs/UnrealBuildTool/System/VCProject.cs
@@ -940,7 +940,8 @@ namespace UnrealBuildTool
 					bShouldCompileMonolithic |= TargetRulesObject.ShouldCompileMonolithic(Platform, Configuration);
 
 					// Get the output directory
-					string RootDirectory = Path.GetFullPath(ProjectFileGenerator.EngineRelativePath);
+					string EngineRootDirectory = Path.GetFullPath(ProjectFileGenerator.EngineRelativePath);
+					string RootDirectory = EngineRootDirectory;
                     if ((TargetRules.IsAGame(TargetRulesObject.Type) || TargetRulesObject.Type == TargetRules.TargetType.Server) && bShouldCompileMonolithic && !TargetRulesObject.bOutputToEngineBinaries)
 					{
 						if (UnrealBuildTool.HasUProjectFile() && Utils.IsFileUnderDirectory(TargetFilePath, UnrealBuildTool.GetUProjectPath()))
@@ -1012,39 +1013,50 @@ namespace UnrealBuildTool
 						UEPlatformProjectGenerator.GenerateGamePlatformSpecificProperties(Platform, Configuration, TargetRulesObject.Type, VCProjectFileContent, RootDirectory, TargetFilePath);
 					}
 
-					// This is the standard UE4 based project NMake build line:
-					//	..\..\Build\BatchFiles\Build.bat <TARGETNAME> <PLATFORM> <CONFIGURATION>
-					//	ie ..\..\Build\BatchFiles\Build.bat BlankProgram Win64 Debug
-
-					string ProjectPlatformConfiguration = " " + TargetName + " " + UBTPlatformName + " " + UBTConfigurationName;
-					string BatchFilesDirectoryName = Path.Combine(ProjectFileGenerator.EngineRelativePath, "Build", "BatchFiles");
-
-					// NMake Build command line
-					VCProjectFileContent.Append("		<NMakeBuildCommandLine>");
-					VCProjectFileContent.Append(EscapePath(NormalizeProjectPath(Path.Combine(BatchFilesDirectoryName, "Build.bat"))) + ProjectPlatformConfiguration.ToString());
-					if (IsForeignProject)
+					if(UnrealBuildTool.RunningRocket() && Utils.IsFileUnderDirectory(TargetFilePath, EngineRootDirectory))
 					{
-						VCProjectFileContent.Append(" " + UProjectPath + (UnrealBuildTool.RunningRocket() ? " -rocket" : ""));
+						// Don't build any engine targets for Rocket. We can still start the program through the debugger though...
+						VCProjectFileContent.AppendFormat("\t\t<NMakeBuildCommandLine>@echo {0} is up to date.</NMakeBuildCommandLine>" + ProjectFileGenerator.NewLine, TargetName);
+						VCProjectFileContent.AppendFormat("\t\t<NMakeReBuildCommandLine>@echo {0} is up to date.</NMakeReBuildCommandLine>" + ProjectFileGenerator.NewLine, TargetName);
+						VCProjectFileContent.AppendFormat("\t\t<NMakeCleanCommandLine>@echo Ignoring clean for {0}.</NMakeCleanCommandLine>" + ProjectFileGenerator.NewLine, TargetName);
 					}
-					VCProjectFileContent.Append("</NMakeBuildCommandLine>" + ProjectFileGenerator.NewLine);
-
-					// NMake ReBuild command line
-					VCProjectFileContent.Append("		<NMakeReBuildCommandLine>");
-					VCProjectFileContent.Append(EscapePath(NormalizeProjectPath(Path.Combine(BatchFilesDirectoryName, "Rebuild.bat"))) + ProjectPlatformConfiguration.ToString());
-					if (IsForeignProject)
+					else
 					{
-						VCProjectFileContent.Append(" " + UProjectPath + (UnrealBuildTool.RunningRocket() ? " -rocket" : ""));
-					}
-					VCProjectFileContent.Append("</NMakeReBuildCommandLine>" + ProjectFileGenerator.NewLine);
+						// This is the standard UE4 based project NMake build line:
+						//	..\..\Build\BatchFiles\Build.bat <TARGETNAME> <PLATFORM> <CONFIGURATION>
+						//	ie ..\..\Build\BatchFiles\Build.bat BlankProgram Win64 Debug
 
-					// NMake Clean command line
-					VCProjectFileContent.Append("		<NMakeCleanCommandLine>");
-					VCProjectFileContent.Append(EscapePath(NormalizeProjectPath(Path.Combine(BatchFilesDirectoryName, "Clean.bat"))) + ProjectPlatformConfiguration.ToString());
-					if (IsForeignProject)
-					{
-						VCProjectFileContent.Append(" " + UProjectPath + (UnrealBuildTool.RunningRocket() ? " -rocket" : ""));
+						string ProjectPlatformConfiguration = " " + TargetName + " " + UBTPlatformName + " " + UBTConfigurationName;
+						string BatchFilesDirectoryName = Path.Combine(ProjectFileGenerator.EngineRelativePath, "Build", "BatchFiles");
+
+						// NMake Build command line
+						VCProjectFileContent.Append("		<NMakeBuildCommandLine>");
+						VCProjectFileContent.Append(EscapePath(NormalizeProjectPath(Path.Combine(BatchFilesDirectoryName, "Build.bat"))) + ProjectPlatformConfiguration.ToString());
+						if (IsForeignProject)
+						{
+							VCProjectFileContent.Append(" " + UProjectPath + (UnrealBuildTool.RunningRocket() ? " -rocket" : ""));
+						}
+						VCProjectFileContent.Append("</NMakeBuildCommandLine>" + ProjectFileGenerator.NewLine);
+
+						// NMake ReBuild command line
+						VCProjectFileContent.Append("		<NMakeReBuildCommandLine>");
+						VCProjectFileContent.Append(EscapePath(NormalizeProjectPath(Path.Combine(BatchFilesDirectoryName, "Rebuild.bat"))) + ProjectPlatformConfiguration.ToString());
+						if (IsForeignProject)
+						{
+							VCProjectFileContent.Append(" " + UProjectPath + (UnrealBuildTool.RunningRocket() ? " -rocket" : ""));
+						}
+						VCProjectFileContent.Append("</NMakeReBuildCommandLine>" + ProjectFileGenerator.NewLine);
+
+						// NMake Clean command line
+						VCProjectFileContent.Append("		<NMakeCleanCommandLine>");
+						VCProjectFileContent.Append(EscapePath(NormalizeProjectPath(Path.Combine(BatchFilesDirectoryName, "Clean.bat"))) + ProjectPlatformConfiguration.ToString());
+						if (IsForeignProject)
+						{
+							VCProjectFileContent.Append(" " + UProjectPath + (UnrealBuildTool.RunningRocket() ? " -rocket" : ""));
+						}
+						VCProjectFileContent.Append("</NMakeCleanCommandLine>" + ProjectFileGenerator.NewLine);
 					}
-					VCProjectFileContent.Append("</NMakeCleanCommandLine>" + ProjectFileGenerator.NewLine);
+
 					VCProjectFileContent.Append("		<NMakeOutput>");
 					VCProjectFileContent.Append(NormalizeProjectPath(NMakePath));
 					VCProjectFileContent.Append("</NMakeOutput>" + ProjectFileGenerator.NewLine);
diff --git a/Engine/Source/Programs/UnrealBuildTool/System/VCProjectFileGenerator.cs b/Engine/Source/Programs/UnrealBuildTool/System/VCProjectFileGenerator.cs
--- a/Engine/Source/Programs/UnrealBuildTool/System/VCProjectFileGenerator.cs
+++ b/Engine/Source/Programs/UnrealBuildTool/System/VCProjectFileGenerator.cs
@@ -799,7 +799,7 @@ namespace UnrealBuildTool
 							ExplorerState.OpenProjects.Add(new Tuple<string, string[]>(ProjectName, new string[]{ }));
 						}
 					}
-					if((!bGeneratingGameProjectFiles && !bGeneratingRocketProjectFiles) || bAlwaysIncludeEngineModules)
+					if( IncludeEnginePrograms )
 					{
 						ExplorerState.OpenProjects.Add(new Tuple<string, string[]>("Automation", new string[0]));
 					}
diff --git a/Engine/Source/Programs/UnrealBuildTool/System/XcodeProjectFileGenerator.cs b/Engine/Source/Programs/UnrealBuildTool/System/XcodeProjectFileGenerator.cs
--- a/Engine/Source/Programs/UnrealBuildTool/System/XcodeProjectFileGenerator.cs
+++ b/Engine/Source/Programs/UnrealBuildTool/System/XcodeProjectFileGenerator.cs
@@ -1084,7 +1084,7 @@ namespace UnrealBuildTool
 
 					if (IsEngineTarget)
 					{
-						if (!bAlwaysIncludeEngineModules)
+						if (!IncludeEngineSource)
 						{
 							// We were asked to exclude engine modules from the generated projects
 							WantProjectFileForTarget = false;
