UE-7304: Force UBT to return success back to calling process for both compilation succeeded (existing behvaiour) and target up-to-date (newly added)

diff --git a/Engine/Source/Programs/UnrealBuildTool/System/UnrealBuildTool.cs b/Engine/Source/Programs/UnrealBuildTool/System/UnrealBuildTool.cs
--- a/Engine/Source/Programs/UnrealBuildTool/System/UnrealBuildTool.cs
+++ b/Engine/Source/Programs/UnrealBuildTool/System/UnrealBuildTool.cs
@@ -1231,7 +1231,9 @@ namespace UnrealBuildTool
             // Print some performance info
             Log.TraceVerbose("Execution time: {0}", (DateTime.UtcNow - StartTime - MutexWaitTime).TotalSeconds);
 
-            return (int) Result;
+            int RetVal = (Result == ECompilationResult.Succeeded || Result == ECompilationResult.UpToDate) ? 0 : (int)Result;
+
+            return RetVal; 
         }
 
 
