Workaround for crash in Mono 3.12 when checking exit status for current process (which we do when we enumerate all processes)

diff --git a/Engine/Source/Programs/UnrealBuildTool/System/BuildHostPlatform.cs b/Engine/Source/Programs/UnrealBuildTool/System/BuildHostPlatform.cs
--- a/Engine/Source/Programs/UnrealBuildTool/System/BuildHostPlatform.cs
+++ b/Engine/Source/Programs/UnrealBuildTool/System/BuildHostPlatform.cs
@@ -334,7 +334,7 @@ namespace UnrealBuildTool
 						try
 						{
 							var ExistingProc = Process.GetProcessById(Pid);
-							if (ExistingProc != null && ExistingProc.HasExited == false)
+							if (ExistingProc != null && Pid != Process.GetCurrentProcess().Id && ExistingProc.HasExited == false)
 							{
 								var ProcInfo = new ProcessInfo(ExistingProc);
 								ProcInfo.Name = Path.GetFileName(Filename);
