Include source and symbols for the Rocket editor.

diff --git a/Engine/Source/Programs/UnrealBuildTool/Configuration/UEBuildTarget.cs b/Engine/Source/Programs/UnrealBuildTool/Configuration/UEBuildTarget.cs
--- a/Engine/Source/Programs/UnrealBuildTool/Configuration/UEBuildTarget.cs
+++ b/Engine/Source/Programs/UnrealBuildTool/Configuration/UEBuildTarget.cs
@@ -1349,6 +1349,16 @@ namespace UnrealBuildTool
 					Manifest.AddBuildProduct(OutputFilePath, DebugInfoExtension);
 				}
 
+				// Add all the stripped debug symbols
+				if(UnrealBuildTool.BuildingRocket() && (Platform == CPPTargetPlatform.Win32 || Platform == CPPTargetPlatform.Win64))
+				{
+					foreach(string OutputFilePath in Binary.Config.OutputFilePaths)
+					{
+						string StrippedPath = Path.Combine(Path.GetDirectoryName(OutputFilePath), Path.GetFileNameWithoutExtension(OutputFilePath) + "-Stripped" + DebugInfoExtension);
+						Manifest.AddBuildProduct(StrippedPath);
+					}
+				}
+
 				if (Binary.Config.Type == UEBuildBinaryType.Executable &&
 					  GlobalLinkEnvironment.Config.CanProduceAdditionalConsoleApp &&
 					  UEBuildConfiguration.bBuildEditor)
diff --git a/Engine/Source/Programs/UnrealBuildTool/Windows/VCToolChain.cs b/Engine/Source/Programs/UnrealBuildTool/Windows/VCToolChain.cs
--- a/Engine/Source/Programs/UnrealBuildTool/Windows/VCToolChain.cs
+++ b/Engine/Source/Programs/UnrealBuildTool/Windows/VCToolChain.cs
@@ -1531,6 +1531,15 @@ namespace UnrealBuildTool
 						ProducedItems.Add(PDBFile);
 					}
 
+					// Write a stripped PDB file for Rocket
+					if(UnrealBuildTool.BuildingRocket())
+					{
+						string StrippedPDBFilePath = Path.Combine(LinkEnvironment.Config.OutputDirectory, Path.GetFileNameWithoutExtension(OutputFile.AbsolutePath) + "-Stripped.pdb");
+						FileItem StrippedPDBFile = FileItem.GetItemByPath(StrippedPDBFilePath);
+						Arguments.AppendFormat(" /PDBSTRIPPED:\"{0}\"", StrippedPDBFilePath);
+						ProducedItems.Add(StrippedPDBFile);
+					}
+
 					// Write the MAP file to the output directory.			
 #if false					
 					if (true)
