Don't print iOS SDK info when generating Xcode project and use better text when generating data for project indexing (Binding IntellliSense data was misleading)

diff --git a/Engine/Source/Programs/UnrealBuildTool/IOS/IOSToolChain.cs b/Engine/Source/Programs/UnrealBuildTool/IOS/IOSToolChain.cs
--- a/Engine/Source/Programs/UnrealBuildTool/IOS/IOSToolChain.cs
+++ b/Engine/Source/Programs/UnrealBuildTool/IOS/IOSToolChain.cs
@@ -203,7 +203,10 @@ namespace UnrealBuildTool
 					{
 						// on the Mac, we can just get the directory name
 						SubDirs = System.IO.Directory.GetDirectories(BaseSDKDir);
-						Log.TraceInformation(String.Format("Directories : {0} {1}", SubDirs, SubDirs[0]));
+						if (!ProjectFileGenerator.bGenerateProjectFiles)
+						{
+							Log.TraceInformation(String.Format("Directories : {0} {1}", SubDirs, SubDirs[0]));
+						}
 					}
 					else
 					{
@@ -266,13 +269,16 @@ namespace UnrealBuildTool
 
 			IOSSDKVersionFloat = float.Parse(IOSSDKVersion, System.Globalization.CultureInfo.InvariantCulture);
 
-			if (BuildHostPlatform.Current.Platform != UnrealTargetPlatform.Mac)
-			{
-				Log.TraceInformation("Compiling with IOS SDK {0} on Mac {1}", IOSSDKVersionFloat, RemoteServerName);
-			}
-			else
+			if (!ProjectFileGenerator.bGenerateProjectFiles)
 			{
-				Log.TraceInformation("Compiling with IOS SDK {0}", IOSSDKVersionFloat);
+				if (BuildHostPlatform.Current.Platform != UnrealTargetPlatform.Mac)
+				{
+					Log.TraceInformation("Compiling with IOS SDK {0} on Mac {1}", IOSSDKVersionFloat, RemoteServerName);
+				}
+				else
+				{
+					Log.TraceInformation("Compiling with IOS SDK {0}", IOSSDKVersionFloat);
+				}
 			}
 		}
 
diff --git a/Engine/Source/Programs/UnrealBuildTool/System/ProjectFileGenerator.cs b/Engine/Source/Programs/UnrealBuildTool/System/ProjectFileGenerator.cs
--- a/Engine/Source/Programs/UnrealBuildTool/System/ProjectFileGenerator.cs
+++ b/Engine/Source/Programs/UnrealBuildTool/System/ProjectFileGenerator.cs
@@ -1304,7 +1304,8 @@ namespace UnrealBuildTool
 			var bSuccess = true;
 			if( ShouldGenerateIntelliSenseData() && TargetFiles.Count > 0 )
 			{
-				using(ProgressWriter Progress = new ProgressWriter("Binding IntelliSense data...", true))
+				string ProgressInfoText = Utils.IsRunningOnMono ? "Generating data for project indexing..." : "Binding IntelliSense data...";
+				using(ProgressWriter Progress = new ProgressWriter(ProgressInfoText, true))
 				{
 					for(int TargetIndex = 0; TargetIndex < TargetFiles.Count; ++TargetIndex)
 					{
