- Generate AndroidManifest.xml from ProjectSettings
- Moved some Android settings to ProjectSettings, re-enabled SDK settings
- Removed SigningConfig.xml, and moved those settings into project settings
- Added concept of NotForLicensees and NoRedist engine and project config settings
- Removed BaseInternalGame.ini, replaced with NotForLicensees/BaseGame.ini
- Moved User*.ini to end of .ini hierarchy
- Added support for CLASS_GlobalUserConfig, so their settings will be saved to <AppData>/.../User*.ini (useful for SDK paths, etc)
- Enabled AndroidPlatformEditor module on Mac
- Changed Mac Build.sh to allow for Android on the commandline (just pass through if it's not an Xcode platform name)
- Iterative Android packaging now looks at just the important .ini sections, NOT entire .ini files

diff --git a/Engine/Build/Android/Java/AndroidManifest.xml b/Engine/Build/Android/Java/AndroidManifest.xml
deleted file mode 100644
diff --git a/Engine/Build/Android/Java/SigningConfig.xml b/Engine/Build/Android/Java/SigningConfig.xml
deleted file mode 100644
diff --git a/Engine/Build/BatchFiles/Mac/Build.sh b/Engine/Build/BatchFiles/Mac/Build.sh
--- a/Engine/Build/BatchFiles/Mac/Build.sh
+++ b/Engine/Build/BatchFiles/Mac/Build.sh
@@ -42,12 +42,14 @@ case $ACTION in
 		        	Platform="IOS"
 		         	AdditionalFlags+=" -deploy -simulator -nocreatestub"
 			;;
-
-			"HTML5") 
+			"macosx")
+					Platform="Mac"
+			;;
+			"HTML5")
 				Platform="HTML5"
 			;;
-			*)	
-				Platform="Mac"
+			*)
+				Platform="$2"
 				AdditionalFlags+=" -deploy " 
 			;; 
 		esac
@@ -74,11 +76,14 @@ case $ACTION in
 		        	AdditionalFlags+=" -simulator"
 				AdditionalFlags+=" -nocreatestub"
 			;; 
+			"macosx")
+					Platform="Mac"
+			;;
 			"HTML5")
 				Platform="HTML5"
 			;;
 			*) 
-        			Platform="Mac"
+        			Platform="$3"
 			;;
 		
 		esac
diff --git a/Engine/Config/BaseEngine.ini b/Engine/Config/BaseEngine.ini
--- a/Engine/Config/BaseEngine.ini
+++ b/Engine/Config/BaseEngine.ini
@@ -2122,6 +2122,12 @@ bBuildForES2=true
 bBuildForES31=false
 bBuildForGL4=false
 bSplitIntoSeparateApks=false
+Orientation=landscape
+DepthBufferPreference=0
+PackageName=com.YourCompany.YourProject
+StoreVersion=1
+VersionDisplayName=1.0
+MinSDKVersion=9
 
 [/Script/UnrealEd.CookerSettings]
 DefaultPVRTCQuality=1
diff --git a/Engine/Config/BaseInternalGame.ini b/Engine/Config/BaseInternalGame.ini
deleted file mode 100644
diff --git a/Engine/Source/Developer/Android/AndroidPlatformEditor/Private/AndroidPlatformEditorModule.cpp b/Engine/Source/Developer/Android/AndroidPlatformEditor/Private/AndroidPlatformEditorModule.cpp
--- a/Engine/Source/Developer/Android/AndroidPlatformEditor/Private/AndroidPlatformEditorModule.cpp
+++ b/Engine/Source/Developer/Android/AndroidPlatformEditor/Private/AndroidPlatformEditorModule.cpp
@@ -42,16 +42,15 @@ class FAndroidPlatformEditorModule
 		{
 			SettingsModule->RegisterSettings("Project", "Platforms", "Android",
 				LOCTEXT("RuntimeSettingsName", "Android"),
-				LOCTEXT("RuntimeSettingsDescription", "Settings and resources for Android platforms"),
+				LOCTEXT("RuntimeSettingsDescription", "Project settings for Android apps"),
 				GetMutableDefault<UAndroidRuntimeSettings>()
 			);
 
-			// This will need to be added back in once we know where we are going to be setting the settings from
-// 			SettingsModule->RegisterSettings("Project", "User Settings", "Android",
-// 				LOCTEXT("RuntimeSettingsName", "Android"),
-// 				LOCTEXT("RuntimeSettingsDescription", "Settings for Android SDK"),
-// 				GetMutableDefault<UAndroidSDKSettings>()
-//			);
+ 			SettingsModule->RegisterSettings("Project", "Platforms", "AndroidSDK",
+ 				LOCTEXT("SDKSettingsName", "Android SDK"),
+ 				LOCTEXT("SDKSettingsDescription", "Settings for Android SDK (for all projects)"),
+ 				GetMutableDefault<UAndroidSDKSettings>()
+			);
 		}
 
 		// Force the SDK settings into a sane state initially so we can make use of them
@@ -71,7 +70,7 @@ class FAndroidPlatformEditorModule
 		if (SettingsModule != nullptr)
 		{
 			SettingsModule->UnregisterSettings("Project", "Platforms", "Android");
-			// SettingsModule->UnregisterSettings("Project", "User Settings", "Android");
+			SettingsModule->UnregisterSettings("Project", "Platforms", "AndroidSDK");
 		}
 	}
 };
diff --git a/Engine/Source/Developer/Android/AndroidPlatformEditor/Private/AndroidSDKSettings.cpp b/Engine/Source/Developer/Android/AndroidPlatformEditor/Private/AndroidSDKSettings.cpp
--- a/Engine/Source/Developer/Android/AndroidPlatformEditor/Private/AndroidSDKSettings.cpp
+++ b/Engine/Source/Developer/Android/AndroidPlatformEditor/Private/AndroidSDKSettings.cpp
@@ -37,14 +37,11 @@ void UAndroidSDKSettings::SetupInitialTargetPaths()
 	FString TempPath;
 	if (GConfig->GetString(TEXT("/Script/AndroidPlatformEditor.AndroidSDKSettings"), TEXT("SDKPath"), TempPath, UserIni))
 	{
-		UE_LOG(AndroidSDKSettings, Warning, TEXT("SDK Path = %s"), *TempPath);
-
 		if (TempPath.Contains(TEXT("Path=")))
 		{
 			SDKPath.Path = TempPath.Mid(7, TempPath.Len() - 9); // Pull the value from the string Path="<payload>"
 		}
 	}
-//	UE_LOG(AndroidSDKSettings, Warning, TEXT("SDK Path = %s"), *SDKPath.Path);
 
 	if (GConfig->GetString(TEXT("/Script/AndroidPlatformEditor.AndroidSDKSettings"), TEXT("NDKPath"), TempPath, UserIni))
 	{
@@ -73,7 +70,6 @@ void UAndroidSDKSettings::SetupInitialTargetPaths()
 			SDKPath.Path = AndroidSDKPath;
 		}
 	}
-	UE_LOG(AndroidSDKSettings, Warning, TEXT("SDK Path = %s"), *SDKPath.Path);
 
 	if (NDKPath.Path.IsEmpty())
 	{
@@ -149,7 +145,6 @@ void UAndroidSDKSettings::UpdateTargetModulePaths(bool bForceUpdate)
 	{
 		Keys.Add(TEXT("ANDROID_HOME"));
 		Values.Add(SDKPath.Path);
-		UE_LOG(AndroidSDKSettings, Log, TEXT("%s"), *SDKPath.Path);
 	}
 	
 	if (bForceUpdate || !NDKPath.Path.IsEmpty())
diff --git a/Engine/Source/Developer/Android/AndroidPlatformEditor/Private/AndroidSDKSettings.h b/Engine/Source/Developer/Android/AndroidPlatformEditor/Private/AndroidSDKSettings.h
--- a/Engine/Source/Developer/Android/AndroidPlatformEditor/Private/AndroidSDKSettings.h
+++ b/Engine/Source/Developer/Android/AndroidPlatformEditor/Private/AndroidSDKSettings.h
@@ -10,24 +10,20 @@
 /**
  * Implements the settings for the Android SDK setup.
  */
-//UCLASS(config=Engine, defaultconfig)
-UCLASS()
+UCLASS(config=Engine, globaluserconfig)
 class ANDROIDPLATFORMEDITOR_API UAndroidSDKSettings : public UObject
 {
 public:
 	GENERATED_UCLASS_BODY()
 
-//		UPROPERTY(GlobalConfig, EditDefaultsOnly, Category = SDKConfig)
-		UPROPERTY()
-		FDirectoryPath SDKPath;
+	UPROPERTY(GlobalConfig, EditAnywhere, Category = SDKConfig, Meta = (DisplayName = "Location of Android SDK (the directory usually contains 'android-sdk-')"))
+	FDirectoryPath SDKPath;
 
-//	UPROPERTY(GlobalConfig, EditDefaultsOnly, Category = SDKConfig)
-	UPROPERTY()
-		FDirectoryPath NDKPath;
+	UPROPERTY(GlobalConfig, EditAnywhere, Category = SDKConfig, Meta = (DisplayName = "Location of Android NDK (the directory usually contains 'android-ndk-')"))
+	FDirectoryPath NDKPath;
 
-//	UPROPERTY(GlobalConfig, EditDefaultsOnly, Category = SDKConfig)
-	UPROPERTY()
-		FDirectoryPath ANTPath;
+	UPROPERTY(GlobalConfig, EditAnywhere, Category = SDKConfig, Meta = (DisplayName = "Location of ANT (the directory usually contains 'apache-ant-')"))
+	FDirectoryPath ANTPath;
 
 #if WITH_EDITOR
 	// UObject interface
diff --git a/Engine/Source/Developer/Android/AndroidPlatformEditor/Private/AndroidTargetSettingsCustomization.cpp b/Engine/Source/Developer/Android/AndroidPlatformEditor/Private/AndroidTargetSettingsCustomization.cpp
--- a/Engine/Source/Developer/Android/AndroidPlatformEditor/Private/AndroidTargetSettingsCustomization.cpp
+++ b/Engine/Source/Developer/Android/AndroidPlatformEditor/Private/AndroidTargetSettingsCustomization.cpp
@@ -34,12 +34,8 @@ FAndroidTargetSettingsCustomization::FAndroidTargetSettingsCustomization()
 	: AndroidRelativePath(TEXT(""))
 	, EngineAndroidPath(FPaths::EngineDir() + TEXT("Build/Android/Java"))
 	, GameAndroidPath(FPaths::GameDir() + TEXT("Build/Android"))
-	, EngineManifestPath(EngineAndroidPath / TEXT("AndroidManifest.xml"))
-	, GameManifestPath(GameAndroidPath / TEXT("AndroidManifest.xml"))
 	, EngineGooglePlayAppIDPath(EngineAndroidPath / TEXT("res") / TEXT("values") / TEXT("GooglePlayAppID.xml"))
 	, GameGooglePlayAppIDPath(GameAndroidPath / TEXT("res") / TEXT("values") / TEXT("GooglePlayAppID.xml"))
-	, EngineSigningConfigPath(EngineAndroidPath / TEXT("SigningConfig.xml"))
-	, GameSigningConfigPath(GameAndroidPath / TEXT("SigningConfig.xml"))
 	, EngineProguardPath(EngineAndroidPath / TEXT("proguard-project.txt"))
 	, GameProguardPath(GameAndroidPath / TEXT("proguard-project.txt"))
 	, EngineProjectPropertiesPath(EngineAndroidPath / TEXT("project.properties"))
@@ -63,34 +59,23 @@ void FAndroidTargetSettingsCustomization::CustomizeDetails(IDetailLayoutBuilder&
 void FAndroidTargetSettingsCustomization::BuildAppManifestSection(IDetailLayoutBuilder& DetailLayout)
 {
 	// Cache some categories
-	IDetailCategoryBuilder& AppManifestCategory = DetailLayout.EditCategory(TEXT("AppManifest"));
+	IDetailCategoryBuilder& APKPackagingCategory = DetailLayout.EditCategory(TEXT("APKPackaging"));
 	IDetailCategoryBuilder& BuildCategory = DetailLayout.EditCategory(TEXT("Build"));
+	IDetailCategoryBuilder& SigningCategory = DetailLayout.EditCategory(TEXT("DistributionSigning"));
 
-	TSharedRef<SPlatformSetupMessage> PlatformSetupMessage = SNew(SPlatformSetupMessage, GameManifestPath)
+	TSharedRef<SPlatformSetupMessage> PlatformSetupMessage = SNew(SPlatformSetupMessage, GameProjectPropertiesPath)
 		.PlatformName(LOCTEXT("AndroidPlatformName", "Android"))
 		.OnSetupClicked(this, &FAndroidTargetSettingsCustomization::CopySetupFilesIntoProject);
 
 	SetupForPlatformAttribute = PlatformSetupMessage->GetReadyToGoAttribute();
 
-	AppManifestCategory.AddCustomRow(LOCTEXT("Warning", "Warning"), false)
+	APKPackagingCategory.AddCustomRow(LOCTEXT("Warning", "Warning"), false)
 		.WholeRowWidget
 		[
 			PlatformSetupMessage
 		];
 
-	AppManifestCategory.AddCustomRow(LOCTEXT("AppManifestHyperlink", "App Manifest Hyperlink"), false)
-		.WholeRowWidget
-		[
-			SNew(SBox)
-			.HAlign(HAlign_Center)
-			[
-				SNew(SHyperlinkLaunchURL, TEXT("http://developer.android.com/guide/topics/manifest/manifest-intro.html"))
-				.Text(LOCTEXT("AndroidDeveloperManifestPage", "Android Developer Page on the App Manifest"))
-				.ToolTipText(LOCTEXT("AndroidDeveloperManifestPageTooltip", "Opens a page that discusses the App Manifest"))
-			]
-		];
-
-	AppManifestCategory.AddCustomRow(LOCTEXT("AppManifestLabel", "App Manifest"), false)
+	APKPackagingCategory.AddCustomRow(LOCTEXT("BuildFolderLabel", "Build Folder"), false)
 		.IsEnabled(SetupForPlatformAttribute)
 		.NameContent()
 		[
@@ -100,7 +85,7 @@ void FAndroidTargetSettingsCustomization::BuildAppManifestSection(IDetailLayoutB
 			.FillWidth(1.0f)
 			[
 				SNew(STextBlock)
-				.Text(LOCTEXT("AppManifestLabel", "App Manifest"))
+				.Text(LOCTEXT("BuildFolderLabel", "Build Folder"))
 				.Font(DetailLayout.GetDetailFont())
 			]
 		]
@@ -111,22 +96,24 @@ void FAndroidTargetSettingsCustomization::BuildAppManifestSection(IDetailLayoutB
 			.AutoWidth()
 			[
 				SNew(SButton)
-				.Text(LOCTEXT("OpenManifestFolderButton", "Open Manifest Folder"))
-				.ToolTipText(LOCTEXT("OpenManifestFolderButton_Tooltip", "Opens the folder containing the manifest (AndroidManifest.xml) in Explorer or Finder"))
-				.OnClicked(this, &FAndroidTargetSettingsCustomization::OpenManifestFolder)
+				.Text(LOCTEXT("OpenBuildFolderButton", "Open Build Folder"))
+				.ToolTipText(LOCTEXT("OpenManifestFolderButton_Tooltip", "Opens the folder containing the build files in Explorer or Finder (it's recommended you check these in to source control to share with your team)"))
+				.OnClicked(this, &FAndroidTargetSettingsCustomization::OpenBuildFolder)
 			]
 		];
 
-	// Show properties that are gated by the manifest file being present and writable
-	TSharedRef<IPropertyHandle> OrientationProperty = DetailLayout.GetProperty(GET_MEMBER_NAME_CHECKED(UAndroidRuntimeSettings, Orientation));
-	OrientationProperty->SetOnPropertyValueChanged(FSimpleDelegate::CreateRaw(this, &FAndroidTargetSettingsCustomization::OnOrientationModified));
-	AppManifestCategory.AddProperty(OrientationProperty)
-		.EditCondition(SetupForPlatformAttribute, NULL);
-
-	TSharedRef<IPropertyHandle> DepthBufferPreferenceProperty = DetailLayout.GetProperty(GET_MEMBER_NAME_CHECKED(UAndroidRuntimeSettings, DepthBufferPreference));
-	DepthBufferPreferenceProperty->SetOnPropertyValueChanged(FSimpleDelegate::CreateRaw(this, &FAndroidTargetSettingsCustomization::OnDepthBufferPreferenceModified));
-	AppManifestCategory.AddProperty(DepthBufferPreferenceProperty)
-		.EditCondition(SetupForPlatformAttribute, NULL);
+	// Signing category
+	SigningCategory.AddCustomRow(LOCTEXT("SigningHyperlink", "Signing Hyperlink"), false)
+		.WholeRowWidget
+		[
+			SNew(SBox)
+			.HAlign(HAlign_Center)
+			[
+				SNew(SHyperlinkLaunchURL, TEXT("http://developer.android.com/tools/publishing/app-signing.html#releasemode"))
+				.Text(LOCTEXT("AndroidDeveloperSigningPage", "Android Developer page on Signing for Distribution"))
+				.ToolTipText(LOCTEXT("AndroidDeveloperSigningPageTooltip", "Opens a page that discusses the signing using keytool"))
+			]
+		];
 
 	// Google Play category
 	IDetailCategoryBuilder& GooglePlayCategory = DetailLayout.EditCategory(TEXT("GooglePlayServices"));
@@ -242,10 +229,10 @@ void FAndroidTargetSettingsCustomization::BuildIconSection(IDetailLayoutBuilder&
 	}
 }
 
-FReply FAndroidTargetSettingsCustomization::OpenManifestFolder()
+FReply FAndroidTargetSettingsCustomization::OpenBuildFolder()
 {
-	const FString EditAppManifestFolder = FPaths::ConvertRelativePathToFull(FPaths::GetPath(GameManifestPath));
-	FPlatformProcess::ExploreFolder(*EditAppManifestFolder);
+	const FString BuildFolder = FPaths::ConvertRelativePathToFull(FPaths::GetPath(GameProjectPropertiesPath));
+	FPlatformProcess::ExploreFolder(*BuildFolder);
 
 	return FReply::Handled();
 }
@@ -254,7 +241,7 @@ void FAndroidTargetSettingsCustomization::CopySetupFilesIntoProject()
 {
 	// First copy the manifest, it must get copied
 	FText ErrorMessage;
-	if (!SourceControlHelpers::CopyFileUnderSourceControl(GameManifestPath, EngineManifestPath, LOCTEXT("AppManifest", "App Manifest"), /*out*/ ErrorMessage))
+	if (!SourceControlHelpers::CopyFileUnderSourceControl(GameProjectPropertiesPath, EngineProjectPropertiesPath, LOCTEXT("ProjectProperties", "Project Properties"), /*out*/ ErrorMessage))
 	{
 		FNotificationInfo Info(ErrorMessage);
 		Info.ExpireDuration = 3.0f;
@@ -275,9 +262,7 @@ void FAndroidTargetSettingsCustomization::CopySetupFilesIntoProject()
 		}
 
 		// and copy the other files (aren't required)
-		SourceControlHelpers::CopyFileUnderSourceControl(GameSigningConfigPath, EngineSigningConfigPath, LOCTEXT("SigningConfig", "Distribution Signing Config"), /*out*/ ErrorMessage);
 		SourceControlHelpers::CopyFileUnderSourceControl(GameProguardPath, EngineProguardPath, LOCTEXT("Proguard", "Proguard Settings"), /*out*/ ErrorMessage);
-		SourceControlHelpers::CopyFileUnderSourceControl(GameProjectPropertiesPath, EngineProjectPropertiesPath, LOCTEXT("ProjectProperties", "Project Properties"), /*out*/ ErrorMessage);
 	}
 
 	SavedLayoutBuilder->ForceRefreshDetails();
@@ -296,67 +281,6 @@ void FAndroidTargetSettingsCustomization::CopyGooglePlayAppIDFileIntoProject()
 	SavedLayoutBuilder->ForceRefreshDetails();
 }
 
-void FAndroidTargetSettingsCustomization::OnOrientationModified()
-{
-	check(SetupForPlatformAttribute.Get() == true);
-
-
-	FManifestUpdateHelper Updater(GameManifestPath);
-
-	const FString OrientationTag(TEXT("android:screenOrientation=\""));
-	const FString ClosingQuote(TEXT("\""));
-	const FString NewOrientationString = OrientationToString(GetDefault<UAndroidRuntimeSettings>()->Orientation);
-	Updater.ReplaceKey(OrientationTag, ClosingQuote, NewOrientationString);
-
-	Updater.Finalize(GameManifestPath);
-}
-
-FString FAndroidTargetSettingsCustomization::OrientationToString(const EAndroidScreenOrientation::Type Orientation)
-{
-	extern ANDROIDRUNTIMESETTINGS_API class UEnum* Z_Construct_UEnum_AndroidRuntimeSettings_EAndroidScreenOrientation();
-	
-	UEnum* Enum = Z_Construct_UEnum_AndroidRuntimeSettings_EAndroidScreenOrientation();
-	check(Enum != nullptr);
-	
-	return Enum->GetMetaData(TEXT("ManifestValue"), (int32)Orientation);
-}
-
-void FAndroidTargetSettingsCustomization::OnDepthBufferPreferenceModified()
-{
-	check(SetupForPlatformAttribute.Get() == true);
-
-
-	FManifestUpdateHelper Updater(GameManifestPath);
-
-	const FString PrefixTag(TEXT("meta-data android:name=\"com.epicgames.ue4.GameActivity.DepthBufferPreference\" android:value=\""));
-	const FString ClosingQuote(TEXT("\""));
-	const FString NewDepthBufferPreferenceString = DepthBufferPreferenceToString(GetDefault<UAndroidRuntimeSettings>()->DepthBufferPreference);
-
-	// First check if the key needs to be created
-	if (!Updater.HasKey(PrefixTag))
-	{
-		// Add the key for value replacement
-		const FString CloseActivityTag(TEXT("</activity"));
-		const FString CloseBracket(TEXT(">"));
-		const FString NewKeyString(TEXT(">\n\t\t<meta-data android:name=\"com.epicgames.ue4.GameActivity.DepthBufferPreference\" android:value=\"0\"/"));
-		Updater.ReplaceKey(CloseActivityTag, CloseBracket, NewKeyString);
-	}
-
-	Updater.ReplaceKey(PrefixTag, ClosingQuote, NewDepthBufferPreferenceString);
-
-	Updater.Finalize(GameManifestPath);
-}
-
-FString FAndroidTargetSettingsCustomization::DepthBufferPreferenceToString(const EAndroidDepthBufferPreference::Type DepthBufferPreference)
-{
-	extern ANDROIDRUNTIMESETTINGS_API class UEnum* Z_Construct_UEnum_AndroidRuntimeSettings_EAndroidDepthBufferPreference();
-
-	UEnum* Enum = Z_Construct_UEnum_AndroidRuntimeSettings_EAndroidDepthBufferPreference();
-	check(Enum != nullptr);
-
-	return Enum->GetMetaData(TEXT("ManifestValue"), (int32)DepthBufferPreference);
-}
-
 void FAndroidTargetSettingsCustomization::OnAppIDModified()
 {
 	check(SetupForPlatformAttribute.Get() == true);
diff --git a/Engine/Source/Developer/Android/AndroidPlatformEditor/Private/AndroidTargetSettingsCustomization.h b/Engine/Source/Developer/Android/AndroidPlatformEditor/Private/AndroidTargetSettingsCustomization.h
--- a/Engine/Source/Developer/Android/AndroidPlatformEditor/Private/AndroidTargetSettingsCustomization.h
+++ b/Engine/Source/Developer/Android/AndroidPlatformEditor/Private/AndroidTargetSettingsCustomization.h
@@ -25,8 +25,8 @@ private:
 	void BuildAppManifestSection(IDetailLayoutBuilder& DetailLayout);
 	void BuildIconSection(IDetailLayoutBuilder& DetailLayout);
 
-	// Navigates to the manifest in the explorer or finder
-	FReply OpenManifestFolder();
+	// Navigates to the build files in the explorer or finder
+	FReply OpenBuildFolder();
 
 	// Copies the setup files for the platform into the project
 	void CopySetupFilesIntoProject();
@@ -34,12 +34,6 @@ private:
 	// Copies the strings.xml file for the platform into the project
 	void CopyGooglePlayAppIDFileIntoProject();
 
-	// Called when the orientation is modified
-	void OnOrientationModified();
-
-	// Called when the depth preference is modified
-	void OnDepthBufferPreferenceModified();
-
 	// Called when the app id is modified
 	void OnAppIDModified();
 
@@ -49,15 +43,9 @@ private:
 	const FString EngineAndroidPath;
 	const FString GameAndroidPath;
 
-	const FString EngineManifestPath;
-	const FString GameManifestPath;
-	
 	const FString EngineGooglePlayAppIDPath;
 	const FString GameGooglePlayAppIDPath;
 
-	const FString EngineSigningConfigPath;
-	const FString GameSigningConfigPath;
-
 	const FString EngineProguardPath;
 	const FString GameProguardPath;
 
diff --git a/Engine/Source/Developer/Settings/Private/SettingsSection.cpp b/Engine/Source/Developer/Settings/Private/SettingsSection.cpp
--- a/Engine/Source/Developer/Settings/Private/SettingsSection.cpp
+++ b/Engine/Source/Developer/Settings/Private/SettingsSection.cpp
@@ -52,7 +52,7 @@ bool FSettingsSection::CanImport() const
 
 bool FSettingsSection::CanResetDefaults() const
 {
-	return (ResetDefaultsDelegate.IsBound() || (SettingsObject.IsValid() && SettingsObject->GetClass()->HasAnyClassFlags(CLASS_Config) && !SettingsObject->GetClass()->HasAnyClassFlags(CLASS_DefaultConfig)));
+	return (ResetDefaultsDelegate.IsBound() || (SettingsObject.IsValid() && SettingsObject->GetClass()->HasAnyClassFlags(CLASS_Config) && !SettingsObject->GetClass()->HasAnyClassFlags(CLASS_DefaultConfig | CLASS_GlobalUserConfig)));
 }
 
 
@@ -64,7 +64,7 @@ bool FSettingsSection::CanSave() const
 
 bool FSettingsSection::CanSaveDefaults() const
 {
-	return (SaveDefaultsDelegate.IsBound() || (SettingsObject.IsValid() && SettingsObject->GetClass()->HasAnyClassFlags(CLASS_Config) && !SettingsObject->GetClass()->HasAnyClassFlags(CLASS_DefaultConfig)));
+	return (SaveDefaultsDelegate.IsBound() || (SettingsObject.IsValid() && SettingsObject->GetClass()->HasAnyClassFlags(CLASS_Config) && !SettingsObject->GetClass()->HasAnyClassFlags(CLASS_DefaultConfig | CLASS_GlobalUserConfig)));
 }
 
 
@@ -140,6 +140,7 @@ bool FSettingsSection::HasDefaultSettingsObject()
 		return false;
 	}
 
+	// @todo userconfig: Should we add GlobalUserConfig here?
 	return SettingsObject->GetClass()->HasAnyClassFlags(CLASS_DefaultConfig);
 }
 
@@ -169,7 +170,7 @@ bool FSettingsSection::ResetDefaults()
 		return ResetDefaultsDelegate.Execute();
 	}
 
-	if (SettingsObject.IsValid() && SettingsObject->GetClass()->HasAnyClassFlags(CLASS_Config) && !SettingsObject->GetClass()->HasAnyClassFlags(CLASS_DefaultConfig))
+	if (SettingsObject.IsValid() && SettingsObject->GetClass()->HasAnyClassFlags(CLASS_Config) && !SettingsObject->GetClass()->HasAnyClassFlags(CLASS_DefaultConfig | CLASS_GlobalUserConfig))
 	{
 		FString ConfigName = SettingsObject->GetClass()->GetConfigName();
 
@@ -205,6 +206,10 @@ bool FSettingsSection::Save()
 		{
 			SettingsObject->UpdateDefaultConfigFile();
 		}
+		else if (SettingsObject->GetClass()->HasAnyClassFlags(CLASS_GlobalUserConfig))
+		{
+			SettingsObject->UpdateGlobalUserConfigFile();
+		}
 		else
 		{
 			SettingsObject->SaveConfig();
diff --git a/Engine/Source/Developer/SettingsEditor/Private/Widgets/SSettingsEditor.cpp b/Engine/Source/Developer/SettingsEditor/Private/Widgets/SSettingsEditor.cpp
--- a/Engine/Source/Developer/SettingsEditor/Private/Widgets/SSettingsEditor.cpp
+++ b/Engine/Source/Developer/SettingsEditor/Private/Widgets/SSettingsEditor.cpp
@@ -454,7 +454,7 @@ void SSettingsEditor::RecordPreferenceChangedAnalytics( ISettingsSectionPtr Sele
 {
 	UProperty* ChangedProperty = PropertyChangedEvent.MemberProperty;
 	// submit analytics data
-	if(FEngineAnalytics::IsAvailable() && ChangedProperty != nullptr)
+	if(FEngineAnalytics::IsAvailable() && ChangedProperty != nullptr && ChangedProperty->GetOwnerClass() != nullptr)
 	{
 		TArray<FAnalyticsEventAttribute> EventAttributes;
 		EventAttributes.Add(FAnalyticsEventAttribute(TEXT("PropertySection"), SelectedSection->GetDisplayName().ToString()));
diff --git a/Engine/Source/Editor/DetailCustomizations/Private/HardwareTargetingSettingsDetails.cpp b/Engine/Source/Editor/DetailCustomizations/Private/HardwareTargetingSettingsDetails.cpp
--- a/Engine/Source/Editor/DetailCustomizations/Private/HardwareTargetingSettingsDetails.cpp
+++ b/Engine/Source/Editor/DetailCustomizations/Private/HardwareTargetingSettingsDetails.cpp
@@ -174,7 +174,8 @@ public:
 
 			SettingRegions.FindChecked(Settings.SettingsObject)->SetText(Settings.Description);
 
-			if (!Settings.SettingsObject->GetClass()->HasAnyClassFlags(CLASS_Config | CLASS_DefaultConfig))
+			// @todo userconfig: Should we check for GlobalUserConfig here? It's not ever going to be checked in...
+			if (!Settings.SettingsObject->GetClass()->HasAnyClassFlags(CLASS_Config | CLASS_DefaultConfig /*| CLASS_GlobalUserConfig*/))
 			{
 				continue;
 			}
diff --git a/Engine/Source/Editor/UnrealEd/UnrealEd.Build.cs b/Engine/Source/Editor/UnrealEd/UnrealEd.Build.cs
--- a/Engine/Source/Editor/UnrealEd/UnrealEd.Build.cs
+++ b/Engine/Source/Editor/UnrealEd/UnrealEd.Build.cs
@@ -194,7 +194,7 @@ public class UnrealEd : ModuleRules
 			}
 		);
 
-		if (Target.Platform == UnrealTargetPlatform.Win64 || Target.Platform == UnrealTargetPlatform.Win32)
+		if (Target.Platform == UnrealTargetPlatform.Win64 || Target.Platform == UnrealTargetPlatform.Win32 || Target.Platform == UnrealTargetPlatform.Mac)
 		{
 			DynamicallyLoadedModuleNames.Add("AndroidPlatformEditor");
 		}
diff --git a/Engine/Source/Programs/UnrealBuildTool/Android/UEDeployAndroid.cs b/Engine/Source/Programs/UnrealBuildTool/Android/UEDeployAndroid.cs
--- a/Engine/Source/Programs/UnrealBuildTool/Android/UEDeployAndroid.cs
+++ b/Engine/Source/Programs/UnrealBuildTool/Android/UEDeployAndroid.cs
@@ -462,14 +462,42 @@ namespace UnrealBuildTool.Android
 		{
 			// make the settings string - this will be char by char compared against last time
 			StringBuilder CurrentSettings = new StringBuilder();
-			CurrentSettings.AppendFormat("NDKROOT={0}{1}", Environment.GetEnvironmentVariable("NDKROOT"), Environment.NewLine);
-			CurrentSettings.AppendFormat("ANDROID_HOME={0}{1}", Environment.GetEnvironmentVariable("ANDROID_HOME"), Environment.NewLine);
-			CurrentSettings.AppendFormat("ANT_HOME={0}{1}", Environment.GetEnvironmentVariable("ANT_HOME"), Environment.NewLine);
-			CurrentSettings.AppendFormat("JAVA_HOME={0}{1}", Environment.GetEnvironmentVariable("JAVA_HOME"), Environment.NewLine);
-			CurrentSettings.AppendFormat("SDKVersion={0}{1}", GetSdkApiLevel(), Environment.NewLine);
-			CurrentSettings.AppendFormat("bForDistribution={0}{1}", bForDistribution, Environment.NewLine);
-			CurrentSettings.AppendFormat("bMakeSeparateApks={0}{1}", bMakeSeparateApks, Environment.NewLine);
-			CurrentSettings.AppendFormat("bOBBinApk={0}{1}", bOBBinApk, Environment.NewLine);
+			CurrentSettings.AppendLine(string.Format("NDKROOT={0}", Environment.GetEnvironmentVariable("NDKROOT")));
+			CurrentSettings.AppendLine(string.Format("ANDROID_HOME={0}", Environment.GetEnvironmentVariable("ANDROID_HOME")));
+			CurrentSettings.AppendLine(string.Format("ANT_HOME={0}", Environment.GetEnvironmentVariable("ANT_HOME")));
+			CurrentSettings.AppendLine(string.Format("JAVA_HOME={0}", Environment.GetEnvironmentVariable("JAVA_HOME")));
+			CurrentSettings.AppendLine(string.Format("SDKVersion={0}", GetSdkApiLevel()));
+			CurrentSettings.AppendLine(string.Format("bForDistribution={0}", bForDistribution));
+			CurrentSettings.AppendLine(string.Format("bMakeSeparateApks={0}", bMakeSeparateApks));
+			CurrentSettings.AppendLine(string.Format("bOBBinApk={0}", bOBBinApk));
+
+			// all AndroidRuntimeSettings ini settings in here
+			ConfigCacheIni Ini = new ConfigCacheIni(UnrealTargetPlatform.Android, "Engine", UnrealBuildTool.GetUProjectPath());
+			ConfigCacheIni.IniSection Section = Ini.FindSection("/Script/AndroidRuntimeSettings.AndroidRuntimeSettings");
+			if (Section != null)
+			{
+				foreach (string Key in Section.Keys)
+				{
+					List<string> Values = Section[Key];
+					foreach (string Value in Values)
+					{
+						CurrentSettings.AppendLine(string.Format("{0}={1}", Key, Value));
+					}
+				}
+			}
+
+			if (Section != null)
+			{
+				Section = Ini.FindSection("/Script/AndroidPlatformEditor.AndroidSDKSettings");
+				foreach (string Key in Section.Keys)
+				{
+					List<string> Values = Section[Key];
+					foreach (string Value in Values)
+					{
+						CurrentSettings.AppendLine(string.Format("{0}={1}", Key, Value));
+					}
+				}
+			}
 
 			string[] Arches = AndroidToolChain.GetAllArchitectures();
 			foreach (string Arch in Arches)
@@ -517,11 +545,6 @@ namespace UnrealBuildTool.Android
 						InputFiles.AddRange(Directory.EnumerateFiles(GameBuildFilesPath, "*.*", SearchOption.AllDirectories));
 					}
 
-					// rebuild if .ini files change
-					// @todo android: programmatically determine if any .ini setting changed?
-					InputFiles.Add(Path.Combine(EngineDirectory, "Config\\BaseEngine.ini"));
-					InputFiles.Add(Path.Combine(ProjectDirectory, "Config\\DefaultEngine.ini"));
-
 					// make sure changed java settings will rebuild apk
 					InputFiles.Add(JavaSettingsFile);
 
@@ -565,6 +588,151 @@ namespace UnrealBuildTool.Android
 			return bAllInputsCurrent;
 		}
 
+		private int ConvertDepthBufferIniValue(string IniValue)
+		{
+			switch (IniValue.ToLower())
+			{
+				case "bits16":
+					return 16;
+				case "bits24":
+					return 24;
+				case "bits32":
+					return 32;
+				default:
+					return 0;
+			}
+		}
+
+		private string ConvertOrientationIniValue(string IniValue)
+		{
+			switch (IniValue.ToLower())
+			{
+				case "portrait":
+					return "portrait";
+				case "reverseportrait":
+					return "reversePortrait";
+				case "sensorportrait":
+					return "sensorPortrait";
+				case "landscape":
+					return "landscape";
+				case "reversepandscape":
+					return "reverseLandscape";
+				case "sensorpandscape":
+					return "sensorLandscape";
+				case "sensor":
+					return "sensor";
+				case "fullsensor":
+					return "fullSensor";
+				default:
+					return "landscape";
+			}
+		}
+
+		private string GenerateManifest(bool bIsForDistribution)
+		{
+			// ini file to get settings from
+			ConfigCacheIni Ini = new ConfigCacheIni(UnrealTargetPlatform.Android, "Engine", UnrealBuildTool.GetUProjectPath());
+			string PackageName;
+			Ini.GetString("/Script/AndroidRuntimeSettings.AndroidRuntimeSettings", "PackageName", out PackageName);
+			bool bEnableGooglePlaySupport;
+			Ini.GetBool("/Script/AndroidRuntimeSettings.AndroidRuntimeSettings", "bEnableGooglePlaySupport", out bEnableGooglePlaySupport);
+			string DepthBufferPreference;
+			Ini.GetString("/Script/AndroidRuntimeSettings.AndroidRuntimeSettings", "DepthBufferPreference", out DepthBufferPreference);
+			int MinSDKVersion;
+			Ini.GetInt32("/Script/AndroidRuntimeSettings.AndroidRuntimeSettings", "MinSDKVersion", out MinSDKVersion);
+			int StoreVersion;
+			Ini.GetInt32("/Script/AndroidRuntimeSettings.AndroidRuntimeSettings", "StoreVersion", out StoreVersion);
+			string VersionDisplayName;
+			Ini.GetString("/Script/AndroidRuntimeSettings.AndroidRuntimeSettings", "VersionDisplayName", out VersionDisplayName);
+			string Orientation;
+			Ini.GetString("/Script/AndroidRuntimeSettings.AndroidRuntimeSettings", "Orientation", out Orientation);
+			string ExtraActivitySettings;
+			Ini.GetString("/Script/AndroidRuntimeSettings.AndroidRuntimeSettings", "ExtraActivitySettings", out ExtraActivitySettings);
+			string ExtraApplicationSettings;
+			Ini.GetString("/Script/AndroidRuntimeSettings.AndroidRuntimeSettings", "ExtraApplicationSettings", out ExtraApplicationSettings);
+			List<string> ExtraPermissions;
+			Ini.GetArray("/Script/AndroidRuntimeSettings.AndroidRuntimeSettings", "ExtraPermissions", out ExtraPermissions);
+
+			StringBuilder Text = new StringBuilder();
+			Text.AppendLine("<?xml version=\"1.0\" encoding=\"utf-8\"?>");
+			Text.AppendLine("<manifest xmlns:android=\"http://schemas.android.com/apk/res/android\"");
+			Text.AppendLine(string.Format("          package=\"{0}\"", PackageName));
+			Text.AppendLine(string.Format("          android:versionCode=\"{0}\"", StoreVersion));
+			Text.AppendLine(string.Format("          android:versionName=\"{0}\">", VersionDisplayName));
+			Text.AppendLine("");
+
+			Text.AppendLine("\t<!-- Application Definition -->");
+			Text.AppendLine("\t<application android:label=\"@string/app_name\"");
+			Text.AppendLine("\t             android:icon=\"@drawable/icon\"");
+			Text.AppendLine("\t             android:hasCode=\"true\">");
+			Text.AppendLine("\t\t<activity android:name=\"com.epicgames.ue4.GameActivity\"");
+			Text.AppendLine("\t\t          android:label=\"@string/app_name\"");
+			Text.AppendLine("\t\t          android:theme=\"@android:style/Theme.NoTitleBar.Fullscreen\"");
+			Text.AppendLine("\t\t          android:launchMode=\"singleTask\"");
+			Text.AppendLine(string.Format("\t\t          android:screenOrientation=\"{0}\"", ConvertOrientationIniValue(Orientation)));
+			Text.AppendLine(string.Format("\t\t          android:debuggable=\"{0}\"", bIsForDistribution ? "false" : "true"));
+			Text.AppendLine("\t\t          android:configChanges=\"orientation|keyboardHidden\">");
+			Text.AppendLine("\t\t\t<meta-data android:name=\"android.app.lib_name\" android:value=\"UE4\"/>");
+			Text.AppendLine("\t\t\t<intent-filter>");
+			Text.AppendLine("\t\t\t\t<action android:name=\"android.intent.action.MAIN\" />");
+			Text.AppendLine("\t\t\t\t<category android:name=\"android.intent.category.LAUNCHER\" />");
+			Text.AppendLine("\t\t\t</intent-filter>");
+			if (!string.IsNullOrEmpty(ExtraActivitySettings))
+			{
+				ExtraActivitySettings = ExtraActivitySettings.Replace("\\n", "\n");
+				foreach (string Line in ExtraActivitySettings.Split("\r\n".ToCharArray()))
+				{
+					Text.AppendLine("\t\t\t" + Line);
+				}
+			}
+			Text.AppendLine("\t\t</activity>");
+			Text.AppendLine(string.Format("\t\t<meta-data android:name=\"com.epicgames.ue4.GameActivity.DepthBufferPreference\" android:value=\"{0}\"/>", ConvertDepthBufferIniValue(DepthBufferPreference)));
+			if (bEnableGooglePlaySupport)
+			{
+				Text.AppendLine("\t\t<meta-data android:name=\"com.google.android.gms.games.APP_ID\"");
+				Text.AppendLine("\t\t           android:value=\"@string/app_id\" />");
+				Text.AppendLine("\t\t<meta-data android:name=\"com.google.android.gms.version\"");
+				Text.AppendLine("\t\t           android:value=\"@integer/google_play_services_version\" />");
+				Text.AppendLine("\t\t<activity android:name=\"com.google.android.gms.ads.AdActivity\"");
+				Text.AppendLine("\t\t          android:configChanges=\"keyboard|keyboardHidden|orientation|screenLayout|uiMode|screenSize|smallestScreenSize\"/>");
+			}
+			if (!string.IsNullOrEmpty(ExtraApplicationSettings))
+			{
+				ExtraApplicationSettings = ExtraApplicationSettings.Replace("\\n", "\n");
+				foreach (string Line in ExtraApplicationSettings.Split("\r\n".ToCharArray()))
+				{
+					Text.AppendLine("\t\t" + Line);
+				}
+			}
+			Text.AppendLine("\t</application>");
+
+			Text.AppendLine("");
+			Text.AppendLine("\t<!-- Requirements -->");
+			// need just the number part of the sdk
+			Text.AppendLine(string.Format("\t<uses-sdk android:minSdkVersion=\"{0}\"/>", MinSDKVersion));
+			Text.AppendLine("\t<uses-feature android:glEsVersion=\"0x00020000\" android:required=\"true\" />");
+			Text.AppendLine("\t<uses-permission android:name=\"android.permission.INTERNET\"/>");
+			Text.AppendLine("\t<uses-permission android:name=\"android.permission.WRITE_EXTERNAL_STORAGE\"/>");
+			Text.AppendLine("\t<uses-permission android:name=\"android.permission.ACCESS_NETWORK_STATE\"/>");
+			Text.AppendLine("\t<uses-permission android:name=\"android.permission.WAKE_LOCK\"/>");
+			Text.AppendLine("\t<uses-permission android:name=\"android.permission.READ_PHONE_STATE\"/>");
+			Text.AppendLine("\t<uses-permission android:name=\"com.android.vending.CHECK_LICENSE\"/>");
+			Text.AppendLine("\t<uses-permission android:name=\"android.permission.ACCESS_WIFI_STATE\"/>");
+			Text.AppendLine("\t<uses-permission android:name=\"android.permission.MODIFY_AUDIO_SETTINGS\"/>");
+			Text.AppendLine("\t<uses-permission android:name=\"android.permission.GET_ACCOUNTS\"/>");
+			Text.AppendLine("\t<uses-permission android:name=\"com.android.vending.BILLING\"/>");
+			Text.AppendLine("\t<uses-permission android:name=\"android.permission.DISABLE_KEYGUARD\"/>");
+			if (ExtraPermissions != null)
+			{
+				foreach (string Permission in ExtraPermissions)
+				{
+					Text.AppendLine(string.Format("\t<uses-permission android:name=\"{0}\"/>", Permission));
+				}
+			}
+			Text.AppendLine("</manifest>");
+			return Text.ToString();
+		}
+
 		private void MakeApk(string ProjectName, string ProjectDirectory, string OutputPath, string EngineDirectory, bool bForDistribution, string CookFlavor, bool bMakeSeparateApks, bool bIncrementalPackage)
 		{
 			Log.TraceInformation("\n===={0}====PREPARING TO MAKE APK=================================================================", DateTime.Now.ToString());
@@ -593,16 +761,33 @@ namespace UnrealBuildTool.Android
 
 			// do we match previous build settings?
 			bool bBuildSettingsMatch = false;
-			if (File.Exists(BuildSettingsCacheFile))
+
+			string ManifestFile = Path.Combine(UE4BuildPath, "AndroidManifest.xml");
+			string NewManifest = GenerateManifest(bForDistribution);
+			string OldManifest = File.Exists(ManifestFile) ? File.ReadAllText(ManifestFile) : "";
+			if (NewManifest == OldManifest) 
 			{
-				string PreviousBuildSettings = File.ReadAllText(BuildSettingsCacheFile);
-				if (PreviousBuildSettings == CurrentBuildSettings)
-				{
-					bBuildSettingsMatch = true;
-				}
-				else
+				bBuildSettingsMatch = true;
+			}
+			else
+			{
+				Log.TraceInformation("AndroidManifest.xml that was generated is different than last build, forcing repackage.");
+			}
+
+			// if the manifest matches, look at other settings stored in a file
+			if (bBuildSettingsMatch)
+			{
+				if (File.Exists(BuildSettingsCacheFile))
 				{
-					Log.TraceInformation("Previous .apk file(s) were made with different build settings, forcing repackage.");
+					string PreviousBuildSettings = File.ReadAllText(BuildSettingsCacheFile);
+					if (PreviousBuildSettings == CurrentBuildSettings)
+					{
+						bBuildSettingsMatch = true;
+					}
+					else
+					{
+						Log.TraceInformation("Previous .apk file(s) were made with different build settings, forcing repackage.");
+					}
 				}
 			}
 
@@ -676,6 +861,9 @@ namespace UnrealBuildTool.Android
 			CopyFileDirectory(GameBuildFilesPath + "/NotForLicensees", UE4BuildPath, Replacements);
 			CopyFileDirectory(GameBuildFilesPath + "/NoRedist", UE4BuildPath, Replacements);
 
+			// in case the game had an AndroidManifest.xml file, we overwrite it now with the generated one
+			File.WriteAllText(ManifestFile, NewManifest);
+
 			// update metadata files (like project.properties, build.xml) if we are missing a build.xml or if we just overwrote project.properties with a bad version in it (from game/engine dir)
 			UpdateProjectProperties(UE4BuildPath, ProjectName);
 
@@ -814,60 +1002,24 @@ namespace UnrealBuildTool.Android
 
 		private void PrepareToSignApk(string BuildPath)
 		{
-			string ConfigFilePath = Path.Combine(BuildPath, "SigningConfig.xml");
-			if (!File.Exists(ConfigFilePath))
+			// ini file to get settings from
+			ConfigCacheIni Ini = new ConfigCacheIni(UnrealTargetPlatform.Android, "Engine", UnrealBuildTool.GetUProjectPath());
+			string KeyAlias, KeyStore, KeyStorePassword, KeyPassword;
+			Ini.GetString("/Script/AndroidRuntimeSettings.AndroidRuntimeSettings", "KeyAlias", out KeyAlias);
+			Ini.GetString("/Script/AndroidRuntimeSettings.AndroidRuntimeSettings", "KeyStore", out KeyStore);
+			Ini.GetString("/Script/AndroidRuntimeSettings.AndroidRuntimeSettings", "KeyStorePassword", out KeyStorePassword);
+			Ini.GetString("/Script/AndroidRuntimeSettings.AndroidRuntimeSettings", "KeyPassword", out KeyPassword);
+
+			if (string.IsNullOrEmpty(KeyAlias) || string.IsNullOrEmpty(KeyStore) || string.IsNullOrEmpty(KeyStorePassword))
 			{
-				throw new BuildException("Unable to sign for Shipping without signing config file: '{0}", ConfigFilePath);
+				throw new BuildException("DistributionSigning settings are not all set. Check the DistributionSettings section in the Andriod tab of Project Settings");
 			}
 
-			// open an Xml parser for the config file
-			string ConfigFile = File.ReadAllText(ConfigFilePath);
-			XmlReader Xml = XmlReader.Create(new StringReader(ConfigFile));
-
-			string Alias = "UESigningKey";
-			string KeystorePassword = "";
-			string KeyPassword = "_sameaskeystore_";
-			string Keystore = "UE.keystore";
-
-			Xml.ReadToFollowing("ue-signing-config");
-			bool bFinishedSection = false;
-			while (Xml.Read() && !bFinishedSection)
-			{
-				switch (Xml.NodeType)
-				{
-					case XmlNodeType.Element:
-						if (Xml.Name == "keyalias")
-						{
-							Alias = Xml.ReadElementContentAsString();
-						}
-						else if (Xml.Name == "keystorepassword")
-						{
-							KeystorePassword = Xml.ReadElementContentAsString();
-						}
-						else if (Xml.Name == "keypassword")
-						{
-							KeyPassword = Xml.ReadElementContentAsString();
-						}
-						else if (Xml.Name == "keystore")
-						{
-							Keystore = Xml.ReadElementContentAsString();
-						}
-						break;
-					case XmlNodeType.EndElement:
-						if (Xml.Name == "ue-signing-config")
-						{
-							bFinishedSection = true;
-						}
-						break;
-				}
-			}
-
-
 			string[] AntPropertiesLines = new string[4];
-			AntPropertiesLines[0] = "key.store=" + Keystore;
-			AntPropertiesLines[1] = "key.alias=" + Alias;
-			AntPropertiesLines[2] = "key.store.password=" + KeystorePassword;
-			AntPropertiesLines[3] = "key.alias.password=" + ((KeyPassword == "_sameaskeystore_") ? KeystorePassword : KeyPassword);
+			AntPropertiesLines[0] = "key.store=" + KeyStore;
+			AntPropertiesLines[1] = "key.alias=" + KeyAlias;
+			AntPropertiesLines[2] = "key.store.password=" + KeyStorePassword;
+			AntPropertiesLines[3] = "key.alias.password=" + ((string.IsNullOrEmpty(KeyPassword) || KeyPassword == "_sameaskeystore_") ? KeyStorePassword : KeyPassword);
 
 			// now write out the properties
 			File.WriteAllLines(Path.Combine(BuildPath, "ant.properties"), AntPropertiesLines);
diff --git a/Engine/Source/Programs/UnrealBuildTool/Configuration/EngineConfiguration.cs b/Engine/Source/Programs/UnrealBuildTool/Configuration/EngineConfiguration.cs
--- a/Engine/Source/Programs/UnrealBuildTool/Configuration/EngineConfiguration.cs
+++ b/Engine/Source/Programs/UnrealBuildTool/Configuration/EngineConfiguration.cs
@@ -86,24 +86,13 @@ namespace UnrealBuildTool
 				EngineDirectory = BuildConfiguration.RelativeEnginePath;
 			}
 
-			foreach (var IniFileName in EnumerateCrossPlatformIniFileNames(ProjectDirectory, EngineDirectory, BaseIniName))
+			foreach (var IniFileName in EnumerateCrossPlatformIniFileNames(ProjectDirectory, EngineDirectory, Platform, BaseIniName))
 			{
 				if (File.Exists(IniFileName))
 				{
 					ParseIniFile(IniFileName);
 				}
 			}
-
-			if (Platform != UnrealTargetPlatform.Unknown)
-			{
-				foreach (var IniFileName in EnumeratePlatformSpecificIniFileNames(ProjectDirectory, EngineDirectory, Platform, BaseIniName))
-				{
-					if (File.Exists(IniFileName))
-					{
-						ParseIniFile(IniFileName);
-					}
-				}
-			}
 		}
 
 		/// <summary>
@@ -516,7 +505,7 @@ namespace UnrealBuildTool
 		/// <summary>
 		/// Returns a list of INI filenames for the given project
 		/// </summary>
-		private static IEnumerable<string> EnumerateCrossPlatformIniFileNames(string ProjectDirectory, string EngineDirectory, string BaseIniName)
+		private static IEnumerable<string> EnumerateCrossPlatformIniFileNames(string ProjectDirectory, string EngineDirectory, UnrealTargetPlatform Platform, string BaseIniName)
 		{
 			// Engine/Config/Base.ini (included in every ini type, required)
 			yield return Path.Combine(EngineDirectory, "Config", "Base.ini");
@@ -524,15 +513,40 @@ namespace UnrealBuildTool
 			// Engine/Config/Base* ini
 			yield return Path.Combine(EngineDirectory, "Config", "Base" + BaseIniName + ".ini");
 
-			// Game/Config/Default* ini
-            if(!String.IsNullOrEmpty(ProjectDirectory))
+			// Engine/Config/NotForLicensees/Base* ini
+			yield return Path.Combine(EngineDirectory, "Config", "NotForLicensees", "Base" + BaseIniName + ".ini");
+
+			// Engine/Config/NoRedist/Base* ini
+			yield return Path.Combine(EngineDirectory, "Config", "NoRedist", "Base" + BaseIniName + ".ini");
+
+			if(!String.IsNullOrEmpty(ProjectDirectory))
 			{
-                yield return Path.Combine(ProjectDirectory, "Config", "Default" + BaseIniName + ".ini");
-            }
+				// Game/Config/Default* ini
+				yield return Path.Combine(ProjectDirectory, "Config", "Default" + BaseIniName + ".ini");
+			
+				// Game/Config/NotForLicensees/Default* ini
+				yield return Path.Combine(ProjectDirectory, "Config", "NotForLicensees", "Default" + BaseIniName + ".ini");
+
+				// Game/Config/NoRedist/Default* ini
+				yield return Path.Combine(ProjectDirectory, "Config", "NoRedist", "Default" + BaseIniName + ".ini");
+			}
 			
+			string PlatformName = GetIniPlatformName(Platform);			
+			if (Platform != UnrealTargetPlatform.Unknown)
+			{
+				// Engine/Config/Platform/Platform* ini
+				yield return Path.Combine(EngineDirectory, "Config", PlatformName, PlatformName + BaseIniName + ".ini");
+
+				if(!String.IsNullOrEmpty(ProjectDirectory))
+				{
+					// Game/Config/Platform/Platform* ini
+					yield return Path.Combine(ProjectDirectory, "Config", PlatformName, PlatformName + BaseIniName + ".ini");
+				}
+			}
+
 			string UserSettingsFolder = null; // Match FPlatformProcess::UserSettingsDir()
 			string PersonalFolder = null; // Match FPlatformProcess::UserDir()
-			if (Environment.OSVersion.Platform == PlatformID.MacOSX)
+			if (BuildHostPlatform.Current.Platform == UnrealTargetPlatform.Mac)
 			{
 				UserSettingsFolder = Path.Combine(Environment.GetFolderPath(Environment.SpecialFolder.Personal), "Library", "Application Support", "Epic");
 				PersonalFolder = Path.Combine(Environment.GetFolderPath(Environment.SpecialFolder.Personal), "Documents");
@@ -556,27 +570,10 @@ namespace UnrealBuildTool
 			// Game/Config/User* ini
             if (!String.IsNullOrEmpty(ProjectDirectory))
             {
-                yield return Path.Combine(ProjectDirectory, "Config", "User" + BaseIniName + ".ini");
-            }
+				yield return Path.Combine(ProjectDirectory, "Config", "User" + BaseIniName + ".ini");
+			}
 		}
 
-		/// <summary>
-		/// Returns a list of INI filenames for the given project
-		/// </summary>
-		private static IEnumerable<string> EnumeratePlatformSpecificIniFileNames(string ProjectDirectory, string EngineDirectory, UnrealTargetPlatform Platform, string BaseIniName)
-		{
-			string PlatformName = GetIniPlatformName(Platform);			
-
-			// Engine/Config/Platform/Platform* ini
-			yield return Path.Combine(EngineDirectory, "Config", PlatformName, PlatformName + BaseIniName + ".ini");
-
-			// Game/Config/Platform/Platform* ini
-            if (!String.IsNullOrEmpty(ProjectDirectory))
-            {
-                yield return Path.Combine(ProjectDirectory, "Config", PlatformName, PlatformName + BaseIniName + ".ini");
-            }
-		}		
-
 		/// <summary>
 		/// Returns the platform name to use as part of platform-specific config files
 		/// </summary>
diff --git a/Engine/Source/Programs/UnrealHeaderTool/Private/CodeGenerator.cpp b/Engine/Source/Programs/UnrealHeaderTool/Private/CodeGenerator.cpp
--- a/Engine/Source/Programs/UnrealHeaderTool/Private/CodeGenerator.cpp
+++ b/Engine/Source/Programs/UnrealHeaderTool/Private/CodeGenerator.cpp
@@ -2163,6 +2163,10 @@ FString FNativeClassHeaderGenerator::GetClassFlagExportText( UClass* Class )
 	{
 		StaticClassFlagText += TEXT(" | CLASS_DefaultConfig");
 	}
+	if( Class->HasAnyClassFlags(CLASS_GlobalUserConfig) )
+	{
+		StaticClassFlagText += TEXT(" | CLASS_GlobalUserConfig");
+	}
 	if( Class->HasAnyClassFlags(CLASS_Config) )
 	{
 		StaticClassFlagText += TEXT(" | CLASS_Config");
diff --git a/Engine/Source/Programs/UnrealHeaderTool/Private/HeaderParser.cpp b/Engine/Source/Programs/UnrealHeaderTool/Private/HeaderParser.cpp
--- a/Engine/Source/Programs/UnrealHeaderTool/Private/HeaderParser.cpp
+++ b/Engine/Source/Programs/UnrealHeaderTool/Private/HeaderParser.cpp
@@ -4327,6 +4327,11 @@ void FHeaderParser::CompileClassDeclaration(FClasses& AllClasses)
 			// Save object config only to Default INIs, never to local INIs.
 			Class->ClassFlags |= CLASS_DefaultConfig;
 		}
+		else if (Specifier == TEXT("globaluserconfig"))
+		{
+			// Save object config only to global user overrides, never to local INIs
+			Class->ClassFlags |= CLASS_GlobalUserConfig;
+		}
 		else if (Specifier == TEXT("showCategories"))
 		{
 			RequireSpecifierValue(PropSpecifier);
diff --git a/Engine/Source/Runtime/Android/AndroidRuntimeSettings/Classes/AndroidRuntimeSettings.h b/Engine/Source/Runtime/Android/AndroidRuntimeSettings/Classes/AndroidRuntimeSettings.h
--- a/Engine/Source/Runtime/Android/AndroidRuntimeSettings/Classes/AndroidRuntimeSettings.h
+++ b/Engine/Source/Runtime/Android/AndroidRuntimeSettings/Classes/AndroidRuntimeSettings.h
@@ -7,6 +7,7 @@
 UENUM()
 namespace EAndroidScreenOrientation
 {
+	// IF THIS CHANGES, MAKE SURE TO UPDATE UEDeployAndroid.cs!
 	enum Type
 	{
 		// Portrait orientation (the display is taller than it is wide)
@@ -39,6 +40,7 @@ namespace EAndroidScreenOrientation
 UENUM()
 namespace EAndroidDepthBufferPreference
 {
+	// IF THIS CHANGES, MAKE SURE TO UPDATE UEDeployAndroid.cs!
 	enum Type
 	{
 		Default = 0 UMETA(DisplayName = "Default", ManifestValue = "0"),
@@ -102,14 +104,58 @@ class ANDROIDRUNTIMESETTINGS_API UAndroidRuntimeSettings : public UObject
 public:
 	GENERATED_UCLASS_BODY()
 
+	// The official name of the product (same as the name you use on the Play Store web site)
+	UPROPERTY(GlobalConfig, EditAnywhere, Category = APKPackaging, Meta = (DisplayName = "Android Package Name (usually in the 'com.Company.Project' format)"))
+	FString PackageName;
+
+	// The version number used to indicate newer versions in the Store
+	UPROPERTY(GlobalConfig, EditAnywhere, Category = APKPackaging, Meta = (DisplayName = "Store Version (1-65535)", ClampMin="1", ClampMax="65535"))
+	int32 StoreVersion;
+
+	// The visual version displayed for end users
+	UPROPERTY(GlobalConfig, EditAnywhere, Category = APKPackaging, Meta = (DisplayName = "Version Display Name (usually x.y)"))
+	FString VersionDisplayName;
+
+	// What OS version the app is allowed to be installed on
+	UPROPERTY(GlobalConfig, EditAnywhere, Category = APKPackaging, Meta = (DisplayName = "Minimum SDK Version (9=Gingerbread, 14=Ice Cream Sandwich, 21=Lollipop)"))
+	int32 MinSDKVersion;
+	
 	// The permitted orientation or orientations of the application on the device
-	UPROPERTY(GlobalConfig, EditAnywhere, Category = AppManifest)
+	UPROPERTY(GlobalConfig, EditAnywhere, Category = APKPackaging)
 	TEnumAsByte<EAndroidScreenOrientation::Type> Orientation;
 
 	// The preferred depth buffer bitcount for Android
-	UPROPERTY(GlobalConfig, EditAnywhere, Category = AppManifest)
+	UPROPERTY(GlobalConfig, EditAnywhere, Category = APKPackaging, Meta = (DisplayName = "Preferred Depth Buffer format"))
 	TEnumAsByte<EAndroidDepthBufferPreference::Type> DepthBufferPreference;
 
+	// Any extra settings for the <application> section
+	UPROPERTY(GlobalConfig, EditAnywhere, Category = AdvancedAPKPackaging, Meta = (DisplayName = "Extra Settings for <application> section (\\n to separate lines)"))
+	FString ExtraApplicationSettings;
+
+	// Any extra settings for the main <activity> section
+	UPROPERTY(GlobalConfig, EditAnywhere, Category = AdvancedAPKPackaging, Meta = (DisplayName = "Extra Settings for <activity> section (\\n to separate lines)"))
+	FString ExtraActivitySettings;
+
+	// Any extra permissions your app needs
+	UPROPERTY(GlobalConfig, EditAnywhere, Category = AdvancedAPKPackaging, Meta = (DisplayName = "Extra Permissions (e.g. 'android.permission.INTERNET')"))
+	TArray<FString> ExtraPermissions;
+	
+	// This is the file that keytool outputs, specified with the -keystore parameter (absolute path or relative to <Project>/Build/Android)
+	UPROPERTY(GlobalConfig, EditAnywhere, Category = DistributionSigning, Meta = (DisplayName = "Key Store (output of keytool)"))
+	FString KeyStore;
+	
+	// This is the name of the key that you specified with the -alias parameter to keytool
+	UPROPERTY(GlobalConfig, EditAnywhere, Category = DistributionSigning, Meta = (DisplayName = "Key Alias (-alias parameter to keytool)"))
+	FString KeyAlias;
+	
+	// This is the password that you specified FOR THE KEYSTORE NOT THE KEY, when running keytool (either with -storepass or by typing it in).
+	UPROPERTY(GlobalConfig, EditAnywhere, Category = DistributionSigning, Meta = (DisplayName = "Key Store Password (-storepass parameter to keytool)"))
+	FString KeyStorePassword;
+	
+	// This is the password for the key that you may have specified with keytool, if it's different from the keystore password. Leave blank to use same as Keystore
+	UPROPERTY(GlobalConfig, EditAnywhere, Category = DistributionSigning, Meta = (DisplayName = "Key Password (leave blank to use Key Store Password)"))
+	FString KeyPassword;
+
 	// Enable ArmV7 support? (this will be used if all type are unchecked)
 	UPROPERTY(GlobalConfig, EditAnywhere, Category = Build, meta = (DisplayName = "Support armv7 [aka armeabi-v7a]"))
 	bool bBuildForArmV7;
@@ -168,7 +214,8 @@ public:
 	/** Android Audio encoding options */
 	UPROPERTY(GlobalConfig, EditAnywhere, Category = DataCooker, meta = (DisplayName = "Audio encoding"))
 	TEnumAsByte<EAndroidAudio::Type> AndroidAudio;
-
+	
+	
 #if WITH_EDITOR
 	// UObject interface
 	virtual void PostEditChangeProperty(struct FPropertyChangedEvent& PropertyChangedEvent) override;
diff --git a/Engine/Source/Runtime/Core/Private/Misc/ConfigCacheIni.cpp b/Engine/Source/Runtime/Core/Private/Misc/ConfigCacheIni.cpp
--- a/Engine/Source/Runtime/Core/Private/Misc/ConfigCacheIni.cpp
+++ b/Engine/Source/Runtime/Core/Private/Misc/ConfigCacheIni.cpp
@@ -10,19 +10,19 @@
 #include "EngineBuildSettings.h"
 #include "Paths.h"
 
-#ifdef WITH_EDITOR
-#	define INI_CACHE (1)
+#if WITH_EDITOR
+	#define INI_CACHE 1
 #else
-#	define INI_CACHE (0)
+	#define INI_CACHE 0
 #endif
 
 DEFINE_LOG_CATEGORY(LogConfig);
 namespace 
 {
-	FName GenerateHierarchyCacheKey(const TArray<FIniFilename>& IniHierarchy, const FString& IniPath) 
+	FString GenerateHierarchyCacheKey(const TArray<FIniFilename>& IniHierarchy, const FString& IniPath)
 	{
 #if !INI_CACHE
-		return FName();
+		return TEXT("");
 #else
 		// A Hierarchy Key is a combined list of all ini file paths that affect that inis data set.
 		FString HierKey;
@@ -39,11 +39,11 @@ namespace
 		}
 		HierKey += IniPath;
 
-		return FName(*HierKey);
+		return HierKey;
 #endif
 	}
 #if INI_CACHE
-	TMap<FName, FConfigFile> HierarchyCache;
+	TMap<FString, FConfigFile> HierarchyCache;
 #endif
 }
 
@@ -718,7 +718,7 @@ static bool LoadIniFileHierarchy(const TArray<FIniFilename>& HierarchyToLoad, FC
 		const FString& IniFileName = IniToLoad.Filename;
 		bool bDoProcess = true;
 #if INI_CACHE
-		if (!IniToLoad.CacheKey.IsNone())
+		if (IniToLoad.CacheKey.Len() > 0)
 		{
 			auto* CachedConfigFile = HierarchyCache.Find(IniToLoad.CacheKey);
 			if (CachedConfigFile) 
@@ -730,23 +730,23 @@ static bool LoadIniFileHierarchy(const TArray<FIniFilename>& HierarchyToLoad, FC
 		}
 		else
 		{
-			ConfigFile.CacheKey = FName();
+			ConfigFile.CacheKey = TEXT("");
 		}
 #endif
 		if (bDoProcess)
 		{
-			// Spit out friendly error if there was a problem locating .inis (e.g. bad command line parameter or missing folder, ...).
-			if (IsUsingLocalIniFile(*IniFileName, NULL) && (IFileManager::Get().FileSize(*IniFileName) < 0))
+		// Spit out friendly error if there was a problem locating .inis (e.g. bad command line parameter or missing folder, ...).
+		if (IsUsingLocalIniFile(*IniFileName, NULL) && (IFileManager::Get().FileSize(*IniFileName) < 0))
+		{
+			if (IniToLoad.bRequired)
+			{
+				//UE_LOG(LogConfig, Error, TEXT("Couldn't locate '%s' which is required to run '%s'"), *IniToLoad.Filename, FApp::GetGameName() );
+				return false;
+			}
+			else
 			{
-				if (IniToLoad.bRequired)
-				{
-					//UE_LOG(LogConfig, Error, TEXT("Couldn't locate '%s' which is required to run '%s'"), *IniToLoad.Filename, FApp::GetGameName() );
-					return false;
-				}
-				else
-				{
 #if INI_CACHE
-					if (!IniToLoad.CacheKey.IsNone())
+					if (IniToLoad.CacheKey.Len() > 0)
 					{
 						if (bDoProcess)
 						{
@@ -754,19 +754,19 @@ static bool LoadIniFileHierarchy(const TArray<FIniFilename>& HierarchyToLoad, FC
 						}
 					}
 #endif
-					continue;
-				}
+				continue;
 			}
-
-			bool bDoEmptyConfig = false;
-			bool bDoCombine = (IniIndex != 0);
-			bool bDoWrite = false;
-			//UE_LOG(LogConfig, Log,  TEXT( "Combining configFile: %s" ), *IniList(IniIndex) );
-			ProcessIniContents(*HierarchyToLoad.Last().Filename, *IniFileName, &ConfigFile, bDoEmptyConfig, bDoCombine, bDoWrite);
 		}
 
+		bool bDoEmptyConfig = false;
+		bool bDoCombine = (IniIndex != 0);
+		bool bDoWrite = false;
+		//UE_LOG(LogConfig, Log,  TEXT( "Combining configFile: %s" ), *IniList(IniIndex) );
+		ProcessIniContents(*HierarchyToLoad.Last().Filename, *IniFileName, &ConfigFile, bDoEmptyConfig, bDoCombine, bDoWrite);
+	}
+
 #if INI_CACHE
-		if (!IniToLoad.CacheKey.IsNone())
+		if (IniToLoad.CacheKey.Len() > 0)
 		{
 			if (bDoProcess)
 			{
@@ -987,7 +987,7 @@ bool FConfigFile::Write( const FString& Filename, bool bDoRemoteWrite/* = true*/
 
 #if INI_CACHE
 	// if we wrote the config successfully
-	if (bResult && !CacheKey.IsNone()) 
+	if (bResult && CacheKey.Len() > 0)
 	{
 		HierarchyCache.Add(CacheKey, *this);
 	}
@@ -2699,37 +2699,55 @@ static FString GetSourceIniFilename(const TCHAR* BaseIniName, const TCHAR* Platf
  */
 static void GetSourceIniHierarchyFilenames(const TCHAR* InBaseIniName, const TCHAR* InPlatformName, const TCHAR* InGameName, const TCHAR* EngineConfigDir, const TCHAR* SourceConfigDir, TArray<FIniFilename>& OutHierarchy, bool bRequireDefaultIni)
 {
+	
+	/**************************************************
+	 **** CRITICAL NOTES
+	 **** If you change this function, you need to also change EnumerateCrossPlatformIniFileNames() in EngineConfiguration.cs!!!
+	 **************************************************/
+	
 	const FString PlatformName(InPlatformName ? InPlatformName : ANSI_TO_TCHAR(FPlatformProperties::IniPlatformName()));
-	const FString BuildPurposeName(FEngineBuildSettings::IsInternalBuild() ? "Internal" : "External");
 
+	// [[[[ ENGINE DEFAULTS ]]]]
 	// Engine/Config/Base.ini (included in every ini type, required)
 	OutHierarchy.Add( FIniFilename(FString::Printf(TEXT("%sBase.ini"), EngineConfigDir), true) );
 	// Engine/Config/Base* ini
 	OutHierarchy.Add( FIniFilename(FString::Printf(TEXT("%sBase%s.ini"), EngineConfigDir, InBaseIniName), false) );
-	// Engine/Config/Base[Internal/External]* ini
-	OutHierarchy.Add( FIniFilename(FString::Printf(TEXT("%sBase%s%s.ini"), EngineConfigDir, *BuildPurposeName, InBaseIniName), false) );
-
-	// <AppData>/UE4/EngineConfig/User* ini
-    OutHierarchy.Add(FIniFilename(FPaths::Combine(FPlatformProcess::UserSettingsDir(), *FString::Printf(TEXT("Unreal Engine/Engine/Config/User%s.ini"), InBaseIniName)), false));
-	// <Documents>/UE4/EngineConfig/User* ini
-    OutHierarchy.Add(FIniFilename(FPaths::Combine(FPlatformProcess::UserDir(), *FString::Printf(TEXT("Unreal Engine/Engine/Config/User%s.ini"), InBaseIniName)), false));
+	// Engine/Config/NotForLicensees/Base* ini
+	OutHierarchy.Add( FIniFilename(FString::Printf(TEXT("%sNotForLicensees/Base%s.ini"), EngineConfigDir, InBaseIniName), false) );
+	// Engine/Config/NoRedist/Base* ini
+	OutHierarchy.Add( FIniFilename(FString::Printf(TEXT("%sNoRedist/Base%s.ini"), EngineConfigDir, InBaseIniName), false) );
 
+	// [[[[ PROJECT SETTINGS ]]]]
 	// Game/Config/Default* ini
 	OutHierarchy.Add(FIniFilename(FString::Printf(TEXT("%sDefault%s.ini"), SourceConfigDir, InBaseIniName), bRequireDefaultIni));
-	// Game/Config/DedicatedServer* ini
+	// Game/Config/DedicatedServer* ini (knowingly NOT in EngineConfiguration.cs because this is a runtime only check)
 	if (IsRunningDedicatedServer())
 	{
 		OutHierarchy.Add(FIniFilename(FString::Printf(TEXT("%s/DedicatedServer%s.ini"), SourceConfigDir, InBaseIniName), false));
 	}
-	// Game/Config/User* ini
-	FString FullIniPath = FString::Printf(TEXT("%s/User%s.ini"), SourceConfigDir, InBaseIniName);
-	OutHierarchy.Add(FIniFilename(FullIniPath, false, GenerateHierarchyCacheKey(OutHierarchy, FullIniPath)));
+	// Game/Config/NotForLicensees/Default* ini
+	OutHierarchy.Add(FIniFilename(FString::Printf(TEXT("%sNotForLicensees/Default%s.ini"), SourceConfigDir, InBaseIniName), false));
+	// Game/Config/NoRedist/Default* ini (Checkpointed here before the platform ini divergence)
+	FString HierarchyCheckpointPath = FString::Printf(TEXT("%sNoRedist/Default%s.ini"), SourceConfigDir, InBaseIniName);
+	OutHierarchy.Add(FIniFilename(HierarchyCheckpointPath, false, GenerateHierarchyCacheKey(OutHierarchy, HierarchyCheckpointPath)));
 	
+	// [[[[ PLATFORM DEFAULTS AND PROJECT SETTINGS ]]]]
 	// Engine/Config/Platform/Platform* ini
 	OutHierarchy.Add( FIniFilename(FString::Printf(TEXT("%s%s/%s%s.ini"), EngineConfigDir, *PlatformName, *PlatformName, InBaseIniName), false) );
 	// Game/Config/Platform/Platform* ini
-	FullIniPath = FString::Printf(TEXT("%s%s/%s%s.ini"), SourceConfigDir, *PlatformName, *PlatformName, InBaseIniName);
-	OutHierarchy.Add(FIniFilename(FullIniPath, false, GenerateHierarchyCacheKey(OutHierarchy, FullIniPath)) );
+	OutHierarchy.Add( FIniFilename(FString::Printf(TEXT("%s%s/%s%s.ini"), SourceConfigDir, *PlatformName, *PlatformName, InBaseIniName), false) );
+
+	// [[[[ GLOBAL USER OVERRIDES ]]]]
+	// <AppData>/UE4/EngineConfig/User* ini
+	OutHierarchy.Add(FIniFilename(FPaths::Combine(FPlatformProcess::UserSettingsDir(), *FString::Printf(TEXT("Unreal Engine/Engine/Config/User%s.ini"), InBaseIniName)), false));
+	// <Documents>/UE4/EngineConfig/User* ini
+	OutHierarchy.Add(FIniFilename(FPaths::Combine(FPlatformProcess::UserDir(), *FString::Printf(TEXT("Unreal Engine/Engine/Config/User%s.ini"), InBaseIniName)), false));
+
+	// [[[[ PROJECT USER OVERRIDES ]]]]
+	// Game/Config/User* ini (Checkpointed here at the end)
+	HierarchyCheckpointPath = FString::Printf(TEXT("%s/User%s.ini"), SourceConfigDir, InBaseIniName);
+	OutHierarchy.Add(FIniFilename(HierarchyCheckpointPath, false, GenerateHierarchyCacheKey(OutHierarchy, HierarchyCheckpointPath)));
+	
 }
 
 /**
@@ -2767,7 +2785,7 @@ void FConfigCacheIni::InitializeConfigSystem()
 {
 	// create GConfig
 	GConfig = new FConfigCacheIni;
-	
+
 	// load the main .ini files (unless we're running a program or a gameless UE4Editor.exe, DefaultEngine.ini is required).
 	const bool bIsGamelessExe = !FApp::HasGameName();
 	const bool bDefaultEngineIniRequired = !bIsGamelessExe && (GIsGameAgnosticExe || FApp::IsGameNameEmpty());
@@ -2860,7 +2878,7 @@ bool FConfigCacheIni::LoadGlobalIniFile(FString& FinalIniFilename, const TCHAR*
 
 	// Keep a record of the original settings
 	NewConfigFile.SourceConfigFile = new FConfigFile();
-	LoadIniFileHierarchy(NewConfigFile.SourceIniHierarchy, *NewConfigFile.SourceConfigFile);
+	LoadIniFileHierarchy(NewConfigFile.SourceIniHierarchy, *NewConfigFile.SourceConfigFile );
 
 	// now generate and make sure it's up to date
 	bool bResult = GenerateDestIniFile(NewConfigFile, FinalIniFilename, NewConfigFile.SourceIniHierarchy, bAllowGeneratedIniWhenCooked);
@@ -2902,7 +2920,7 @@ void FConfigCacheIni::LoadLocalIniFile(FConfigFile& ConfigFile, const TCHAR* Ini
 		
 		// Keep a record of the original settings
 		ConfigFile.SourceConfigFile = new FConfigFile();
-		LoadIniFileHierarchy(ConfigFile.SourceIniHierarchy, *ConfigFile.SourceConfigFile);
+		LoadIniFileHierarchy(ConfigFile.SourceIniHierarchy, *ConfigFile.SourceConfigFile );
 
 		// now generate and make sure it's up to date (using IniName as a Base for an ini filename)
 		const bool bAllowGeneratedINIs = true;
@@ -2928,7 +2946,7 @@ void FConfigCacheIni::LoadExternalIniFile(FConfigFile& ConfigFile, const TCHAR*
 
 		// Keep a record of the original settings
 		ConfigFile.SourceConfigFile = new FConfigFile();
-		LoadIniFileHierarchy(ConfigFile.SourceIniHierarchy, *ConfigFile.SourceConfigFile);
+		LoadIniFileHierarchy(ConfigFile.SourceIniHierarchy, *ConfigFile.SourceConfigFile );
 
 		// now generate and make sure it's up to date (using IniName as a Base for an ini filename)
 		const bool bAllowGeneratedINIs = true;
diff --git a/Engine/Source/Runtime/Core/Public/Misc/ConfigCacheIni.h b/Engine/Source/Runtime/Core/Public/Misc/ConfigCacheIni.h
--- a/Engine/Source/Runtime/Core/Public/Misc/ConfigCacheIni.h
+++ b/Engine/Source/Runtime/Core/Public/Misc/ConfigCacheIni.h
@@ -32,9 +32,9 @@ struct FIniFilename
 	/** If true this ini file is required to generate the output ini. */
 	bool bRequired;
 	/** Used as ID for looking up an INI Hierarchy */
-	FName CacheKey;
+	FString CacheKey;
 
-	FIniFilename(const FString& InFilename, bool InIsRequired, FName InCacheKey=FName())
+	FIniFilename(const FString& InFilename, bool InIsRequired, FString InCacheKey=FString(TEXT("")))
 		: Filename(InFilename)
 		, bRequired(InIsRequired) 
 		, CacheKey(InCacheKey)
@@ -67,7 +67,7 @@ public:
 	FConfigFile* SourceConfigFile;
 
 	/** Key to the cache to speed up ini parsing */
-	FName CacheKey;
+	FString CacheKey;
 
 #if !UE_BUILD_SHIPPING
 	/** The collection of overrides which stemmed from the commandline */
diff --git a/Engine/Source/Runtime/CoreUObject/Private/UObject/Obj.cpp b/Engine/Source/Runtime/CoreUObject/Private/UObject/Obj.cpp
--- a/Engine/Source/Runtime/CoreUObject/Private/UObject/Obj.cpp
+++ b/Engine/Source/Runtime/CoreUObject/Private/UObject/Obj.cpp
@@ -1774,29 +1774,43 @@ FString UObject::GetDefaultConfigFilename() const
 	return FString::Printf(TEXT("%sDefault%s.ini"), *FPaths::SourceConfigDir(), *GetClass()->ClassConfigName.ToString());
 }
 
+FString UObject::GetGlobalUserConfigFilename() const
+{
+	return FString::Printf(TEXT("%sUnreal Engine/Engine/Config/User%s.ini"), FPlatformProcess::UserSettingsDir(), *GetClass()->ClassConfigName.ToString());
+}
+
 // @todo ini: Verify per object config objects
-void UObject::UpdateDefaultConfigFile()
+void UObject::UpdateSingleSectionOfConfigFile(const FString& ConfigIniName)
 {
 	// create a sandbox FConfigCache
 	FConfigCacheIni Config;
 
 	// add an empty file to the config so it doesn't read in the original file (see FConfigCacheIni.Find())
-	FString DefaultIniName = GetDefaultConfigFilename();
-	FConfigFile& NewFile = Config.Add(DefaultIniName, FConfigFile());
+	FConfigFile& NewFile = Config.Add(ConfigIniName, FConfigFile());
 
 	// save the object properties to this file
-	SaveConfig(CPF_Config, *DefaultIniName, &Config);
+	SaveConfig(CPF_Config, *ConfigIniName, &Config);
 
 	ensureMsgf(Config.Num() == 1, TEXT("UObject::UpdateDefaultConfig() caused more files than expected in the Sandbox config cache!"));
 
 	// make sure SaveConfig wrote only to the file we expected
-	NewFile.UpdateSections(*DefaultIniName, *GetClass()->ClassConfigName.ToString());
+	NewFile.UpdateSections(*ConfigIniName, *GetClass()->ClassConfigName.ToString());
 
 	// reload the file, so that it refresh the cache internally.
 	FString FinalIniFileName;
 	GConfig->LoadGlobalIniFile(FinalIniFileName, *GetClass()->ClassConfigName.ToString(), NULL, NULL, true);
 }
 
+void UObject::UpdateDefaultConfigFile()
+{
+	UpdateSingleSectionOfConfigFile(GetDefaultConfigFilename());
+}
+
+void UObject::UpdateGlobalUserConfigFile()
+{
+	UpdateSingleSectionOfConfigFile(GetGlobalUserConfigFilename());
+}
+
 
 void UObject::InstanceSubobjectTemplates( FObjectInstancingGraph* InstanceGraph )
 {
diff --git a/Engine/Source/Runtime/CoreUObject/Public/UObject/ObjectBase.h b/Engine/Source/Runtime/CoreUObject/Public/UObject/ObjectBase.h
--- a/Engine/Source/Runtime/CoreUObject/Public/UObject/ObjectBase.h
+++ b/Engine/Source/Runtime/CoreUObject/Public/UObject/ObjectBase.h
@@ -172,6 +172,8 @@ enum EClassFlags
 	CLASS_Deprecated		  = 0x02000000,
 	/** Class not shown in editor drop down for class selection */
 	CLASS_HideDropDown		  = 0x04000000,
+	/** Class settings are saved to <AppData>/..../Blah.ini (as opposed to CLASS_DefaultConfig) */
+	CLASS_GlobalUserConfig	  = 0x08000000,
 	/** Class was declared directly in C++ and has no boilerplate generated by UnrealHeaderTool */
 	CLASS_Intrinsic			  = 0x10000000,
 	/** Class has already been constructed (maybe in a previous DLL version before hot-reload). */
@@ -188,7 +190,7 @@ enum EClassFlags
 	//@{
 	CLASS_Inherit           = CLASS_Transient | CLASS_DefaultConfig | CLASS_Config | CLASS_Localized | CLASS_PerObjectConfig | CLASS_ConfigDoNotCheckDefaults | CLASS_NotPlaceable
 							| CLASS_PointersDefaultToAutoWeak | CLASS_PointersDefaultToWeak | CLASS_Const
-							| CLASS_HasInstancedReference | CLASS_Deprecated | CLASS_DefaultToInstanced,
+							| CLASS_HasInstancedReference | CLASS_Deprecated | CLASS_DefaultToInstanced | CLASS_GlobalUserConfig,
 
 	/** these flags will be cleared by the compiler when the class is parsed during script compilation */
 	CLASS_RecompilerClear   = CLASS_Inherit | CLASS_Abstract | CLASS_NoExport | CLASS_Native | CLASS_Intrinsic | CLASS_TokenStreamAssembled,
@@ -204,6 +206,7 @@ enum EClassFlags
 	CLASS_SaveInCompiledInClasses = 
 		CLASS_Abstract | 
 		CLASS_DefaultConfig |
+		CLASS_GlobalUserConfig |
 		CLASS_Config |
 		CLASS_Transient |
 		CLASS_Localized |
diff --git a/Engine/Source/Runtime/CoreUObject/Public/UObject/UObject.h b/Engine/Source/Runtime/CoreUObject/Public/UObject/UObject.h
--- a/Engine/Source/Runtime/CoreUObject/Public/UObject/UObject.h
+++ b/Engine/Source/Runtime/CoreUObject/Public/UObject/UObject.h
@@ -746,11 +746,28 @@ public:
 	 */
 	void UpdateDefaultConfigFile();
 
+	/**
+	 * Saves just the section(s) for this class into the global user ini file for the class (with just the changes from base)
+	 */
+	void UpdateGlobalUserConfigFile();
+
+private:
+	/**
+	 * Saves just the section(s) for this class into the given ini file for the class (with just the changes from base)
+	 */
+	void UpdateSingleSectionOfConfigFile(const FString& ConfigIniName);
+public:
+	
 	/**
 	 * Get the default config filename for the specified UObject
 	 */
 	FString GetDefaultConfigFilename() const;
 
+	/**
+	 * Get the global user override config filename for the specified UObject
+	 */
+	FString GetGlobalUserConfigFilename() const;
+
 	/**
 	 * Imports property values from an .ini file.
 	 *
