- Moved SDK and NDK versions from BuildConfig.xml to ProjectSettings
- Fixed a logic error in iterative packaging dependency checking

diff --git a/Engine/Config/BaseEngine.ini b/Engine/Config/BaseEngine.ini
--- a/Engine/Config/BaseEngine.ini
+++ b/Engine/Config/BaseEngine.ini
@@ -2129,6 +2129,11 @@ StoreVersion=1
 VersionDisplayName=1.0
 MinSDKVersion=9
 
+[/Script/AndroidPlatformEditor.AndroidSDKSettings]
+SDKAPILevel=latest
+NDKAPILevel=android-19
+
+
 [/Script/UnrealEd.CookerSettings]
 DefaultPVRTCQuality=1
 DefaultASTCQualityBySpeed=1
diff --git a/Engine/Source/Programs/UnrealBuildTool/Android/AndroidToolChain.cs b/Engine/Source/Programs/UnrealBuildTool/Android/AndroidToolChain.cs
--- a/Engine/Source/Programs/UnrealBuildTool/Android/AndroidToolChain.cs
+++ b/Engine/Source/Programs/UnrealBuildTool/Android/AndroidToolChain.cs
@@ -128,10 +128,12 @@ namespace UnrealBuildTool
 
 		static public string GetNdkApiLevel()
 		{
-			// default to looking on disk for latest API level
-			string Target = AndroidPlatform.AndroidNdkApiTarget;
+			// ask the .ini system for what version to use
+			ConfigCacheIni Ini = new ConfigCacheIni(UnrealTargetPlatform.Android, "Engine", UnrealBuildTool.GetUProjectPath());
+			string NDKLevel;
+			Ini.GetString("/Script/AndroidPlatformEditor.AndroidSDKSettings", "NDKAPILevel", out NDKLevel);
 
-			if (Target == "latest")
+			if (NDKLevel == "latest")
 			{
 				// get a list of NDK platforms
 				string PlatformsDir = Environment.ExpandEnvironmentVariables("%NDKROOT%/platforms");
@@ -141,10 +143,10 @@ namespace UnrealBuildTool
 				}
 
 				// return the largest of them
-				Target = GetLargestApiLevel(Directory.GetDirectories(PlatformsDir));
+				NDKLevel = GetLargestApiLevel(Directory.GetDirectories(PlatformsDir));
 			}
 
-			return Target;
+			return NDKLevel;
 		}
 
 		static public string GetLargestApiLevel(string[] ApiLevels)
diff --git a/Engine/Source/Programs/UnrealBuildTool/Android/UEBuildAndroid.cs b/Engine/Source/Programs/UnrealBuildTool/Android/UEBuildAndroid.cs
--- a/Engine/Source/Programs/UnrealBuildTool/Android/UEBuildAndroid.cs
+++ b/Engine/Source/Programs/UnrealBuildTool/Android/UEBuildAndroid.cs
@@ -12,14 +12,6 @@ namespace UnrealBuildTool
 {
 	class AndroidPlatform : UEBuildPlatform
 	{
-		/// <summary>
-		/// Android settings.
-		/// </summary>
-		[XmlConfig]
-		public static string AndroidNdkApiTarget = "android-19";
-		[XmlConfig]
-		public static string AndroidSdkApiTarget = "latest";
-
 		// The current architecture - affects everything about how UBT operates on Android
 		public override string GetActiveArchitecture()
 		{
diff --git a/Engine/Source/Programs/UnrealBuildTool/Android/UEDeployAndroid.cs b/Engine/Source/Programs/UnrealBuildTool/Android/UEDeployAndroid.cs
--- a/Engine/Source/Programs/UnrealBuildTool/Android/UEDeployAndroid.cs
+++ b/Engine/Source/Programs/UnrealBuildTool/Android/UEDeployAndroid.cs
@@ -59,17 +59,19 @@ namespace UnrealBuildTool.Android
 		{
 			if (CachedSDKLevel == null)
 			{
-				// default to looking on disk for latest API level
-				string Target = AndroidPlatform.AndroidSdkApiTarget;
+				// ask the .ini system for what version to use
+				ConfigCacheIni Ini = new ConfigCacheIni(UnrealTargetPlatform.Android, "Engine", UnrealBuildTool.GetUProjectPath());
+				string SDKLevel;
+				Ini.GetString("/Script/AndroidPlatformEditor.AndroidSDKSettings", "SDKAPILevel", out SDKLevel);
 
 				// if we want to use whatever version the ndk uses, then use that
-				if (Target == "matchndk")
+				if (SDKLevel == "matchndk")
 				{
-					Target = AndroidToolChain.GetNdkApiLevel();
+					SDKLevel = AndroidToolChain.GetNdkApiLevel();
 				}
 
 				// run a command and capture output
-				if (Target == "latest")
+				if (SDKLevel == "latest")
 				{
 					// we expect there to be one, so use the first one
 					string AndroidCommandPath = Environment.ExpandEnvironmentVariables("%ANDROID_HOME%/tools/android" + (Utils.IsRunningOnMono ? "" : ".bat"));
@@ -87,16 +89,16 @@ namespace UnrealBuildTool.Android
 
 					if (PossibleApiLevels != null && PossibleApiLevels.Count > 0)
 					{
-						Target = AndroidToolChain.GetLargestApiLevel(PossibleApiLevels.ToArray());
+						SDKLevel = AndroidToolChain.GetLargestApiLevel(PossibleApiLevels.ToArray());
 					}
 					else
 					{
-						throw new BuildException("Can't make an APK an API installed (see \"android.bat list targets\")");
+						throw new BuildException("Can't make an APK without an API installed (see \"android.bat list targets\")");
 					}
 				}
 
-				Console.WriteLine("Building Java with SDK API '{0}'", Target);
-				CachedSDKLevel = Target;
+				Console.WriteLine("Building Java with SDK API level '{0}'", SDKLevel);
+				CachedSDKLevel = SDKLevel;
 			}
 
 			return CachedSDKLevel;
@@ -780,12 +782,9 @@ namespace UnrealBuildTool.Android
 				if (File.Exists(BuildSettingsCacheFile))
 				{
 					string PreviousBuildSettings = File.ReadAllText(BuildSettingsCacheFile);
-					if (PreviousBuildSettings == CurrentBuildSettings)
-					{
-						bBuildSettingsMatch = true;
-					}
-					else
+					if (PreviousBuildSettings != CurrentBuildSettings)
 					{
+						bBuildSettingsMatch = false;
 						Log.TraceInformation("Previous .apk file(s) were made with different build settings, forcing repackage.");
 					}
 				}
