- Fixed the default package naming for android projects

diff --git a/Engine/Config/BaseEngine.ini b/Engine/Config/BaseEngine.ini
--- a/Engine/Config/BaseEngine.ini
+++ b/Engine/Config/BaseEngine.ini
@@ -2124,7 +2124,7 @@ bBuildForGL4=false
 bSplitIntoSeparateApks=false
 Orientation=landscape
 DepthBufferPreference=0
-PackageName=com.YourCompany.YourProject
+PackageName=com.YourCompany.{PROJECT_NAME}
 StoreVersion=1
 VersionDisplayName=1.0
 MinSDKVersion=9
diff --git a/Engine/Source/Programs/UnrealBuildTool/Android/UEDeployAndroid.cs b/Engine/Source/Programs/UnrealBuildTool/Android/UEDeployAndroid.cs
--- a/Engine/Source/Programs/UnrealBuildTool/Android/UEDeployAndroid.cs
+++ b/Engine/Source/Programs/UnrealBuildTool/Android/UEDeployAndroid.cs
@@ -630,7 +630,7 @@ namespace UnrealBuildTool.Android
 			}
 		}
 
-		private string GenerateManifest(bool bIsForDistribution)
+		private string GenerateManifest(string ProjectName, bool bIsForDistribution)
 		{
 			// ini file to get settings from
 			ConfigCacheIni Ini = new ConfigCacheIni(UnrealTargetPlatform.Android, "Engine", UnrealBuildTool.GetUProjectPath());
@@ -655,6 +655,9 @@ namespace UnrealBuildTool.Android
 			List<string> ExtraPermissions;
 			Ini.GetArray("/Script/AndroidRuntimeSettings.AndroidRuntimeSettings", "ExtraPermissions", out ExtraPermissions);
 
+			// replace some variables
+			PackageName = PackageName.Replace("{PROJECT_NAME}", ProjectName);
+
 			StringBuilder Text = new StringBuilder();
 			Text.AppendLine("<?xml version=\"1.0\" encoding=\"utf-8\"?>");
 			Text.AppendLine("<manifest xmlns:android=\"http://schemas.android.com/apk/res/android\"");
@@ -765,7 +768,7 @@ namespace UnrealBuildTool.Android
 			bool bBuildSettingsMatch = false;
 
 			string ManifestFile = Path.Combine(UE4BuildPath, "AndroidManifest.xml");
-			string NewManifest = GenerateManifest(bForDistribution);
+			string NewManifest = GenerateManifest(ProjectName, bForDistribution);
 			string OldManifest = File.Exists(ManifestFile) ? File.ReadAllText(ManifestFile) : "";
 			if (NewManifest == OldManifest) 
 			{
