updates to generate plist as part of a build which moves all of the plist data in to the ini files
UEPLAT-114
#ios

diff --git a/Engine/Build/IOS/UE4Game-Info.plist b/Engine/Build/IOS/UE4Game-Info.plist
deleted file mode 100644
diff --git a/Engine/Config/BaseEngine.ini b/Engine/Config/BaseEngine.ini
--- a/Engine/Config/BaseEngine.ini
+++ b/Engine/Config/BaseEngine.ini
@@ -2099,6 +2099,7 @@ MenuStartY=10.0
 DebugInfoStartX=20.0
 DebugInfoStartY=20.0
 
+; These are the defaults for iOS settings, and they need to be in the .ini since UBT reads the .ini settings, without instantiating the class
 ; Note: These should be in IOSEngine, but are placed here because the Settings editor UI cannot write to the Game IOSEngine, meaning settings in Engine/IOSEngine are non-overridable
 [/Script/IOSRuntimeSettings.IOSRuntimeSettings]
 bEnableGameCenterSupport=False
@@ -2110,6 +2111,23 @@ bSupportsMetal=True
 bSupportsOpenGLES2=True
 bCookPVRTCTextures=True
 bCookASTCTextures=False
+bSupportsMetalMRT=False
+bDevForArmV7=True
+bDevForArm64=False
+bDevForArmV7S=False
+bShipForArmV7=True
+bShipForArm64=True
+bShipForArmV7S=False
+bUseRSync=False
+BundleDisplayName={PROJECT_NAME}
+BundleName={PROJECT_NAME}
+BundleIdentifier=com.YourCompany.{PROJECT_NAME}
+VersionInfo=1.0
+FrameRateLock=PUFRL_30
+MinimumiOSVersion=IOS_61
+bSupportsIPad=True
+bSupportsIPhone=True
+AdditionalPlistData=
 
 ; These are the defaults for Android settings, and they need to be in the .ini since UBT reads the .ini settings, without instantiating the class
 [/Script/AndroidRuntimeSettings.AndroidRuntimeSettings]
diff --git a/Engine/Source/Developer/iOS/IOSPlatformEditor/Private/IOSTargetSettingsCustomization.cpp b/Engine/Source/Developer/iOS/IOSPlatformEditor/Private/IOSTargetSettingsCustomization.cpp
--- a/Engine/Source/Developer/iOS/IOSPlatformEditor/Private/IOSTargetSettingsCustomization.cpp
+++ b/Engine/Source/Developer/iOS/IOSPlatformEditor/Private/IOSTargetSettingsCustomization.cpp
@@ -116,6 +116,7 @@ void FIOSTargetSettingsCustomization::BuildPListSection(IDetailLayoutBuilder& De
 	IDetailCategoryBuilder& OSInfoCategory = DetailLayout.EditCategory(TEXT("OS Info"));
 	IDetailCategoryBuilder& DeviceCategory = DetailLayout.EditCategory(TEXT("Devices"));
 	IDetailCategoryBuilder& BuildCategory = DetailLayout.EditCategory(TEXT("Build"));
+	IDetailCategoryBuilder& ExtraCategory = DetailLayout.EditCategory(TEXT("Extra PList Data"));
 
 	TSharedRef<SPlatformSetupMessage> PlatformSetupMessage = SNew(SPlatformSetupMessage, GameInfoPath)
 		.PlatformName(LOCTEXT("iOSPlatformName", "iOS"))
@@ -316,6 +317,8 @@ void FIOSTargetSettingsCustomization::BuildPListSection(IDetailLayoutBuilder& De
 
 	SETUP_PLIST_PROP(MinimumiOSVersion, OSInfoCategory, LOCTEXT("MinimumiOSVersionToolTip", "Minimum iOS version this game supports"));
 
+	SETUP_PLIST_PROP(AdditionalPlistData, ExtraCategory, LOCTEXT("AdditionalPlistDataToolTip", "Any additional plist key/value data utilizing \\n for a new line"));
+
 	SETUP_NONROCKET_PROP(bDevForArmV7, BuildCategory, LOCTEXT("DevForArmV7ToolTip", "Enable ArmV7 support? (this will be used if all type are unchecked)"), FIOSTargetSettingsCustomizationConstants::DisabledTip);
 	SETUP_NONROCKET_PROP(bDevForArm64, BuildCategory, LOCTEXT("DevForArm64ToolTip", "Enable Arm64 support?"), FIOSTargetSettingsCustomizationConstants::DisabledTip);
 	SETUP_NONROCKET_PROP(bDevForArmV7S, BuildCategory, LOCTEXT("DevForArmV7SToolTip", "Enable ArmV7s support?"), FIOSTargetSettingsCustomizationConstants::DisabledTip);
@@ -448,8 +451,8 @@ void FIOSTargetSettingsCustomization::CopySetupFilesIntoProject()
 
 void FIOSTargetSettingsCustomization::OnPlistPropertyModified()
 {
-	check(SetupForPlatformAttribute.Get() == true);
-	const UIOSRuntimeSettings& Settings = *GetDefault<UIOSRuntimeSettings>();
+// 	check(SetupForPlatformAttribute.Get() == true);
+/*	const UIOSRuntimeSettings& Settings = *GetDefault<UIOSRuntimeSettings>();
 
 	FManifestUpdateHelper Updater(GameInfoPath);
 
@@ -537,8 +540,11 @@ void FIOSTargetSettingsCustomization::OnPlistPropertyModified()
 	}
 	Updater.ReplaceKey(MiniOSVersionKey, ClosingString, iOSVersionBody);
 
+	const FString ClosingDict(TEXT("</dict>"));
+	Updater.ReplaceKey(iOSVersionBody + ClosingString, ClosingDict, Settings.AdditionalPlistData.Replace(TEXT("\\n"), TEXT("\n")));
+
 	// Write out the updated .plist
-	Updater.Finalize(GameInfoPath, true, FFileHelper::EEncodingOptions::ForceUTF8);
+	Updater.Finalize(GameInfoPath, true, FFileHelper::EEncodingOptions::ForceUTF8);*/
 }
 
 void FIOSTargetSettingsCustomization::BuildImageRow(IDetailLayoutBuilder& DetailLayout, IDetailCategoryBuilder& Category, const FPlatformIconInfo& Info, const FVector2D& MaxDisplaySize)
diff --git a/Engine/Source/Programs/AutomationTool/IOS/IOSPlatform.Automation.cs b/Engine/Source/Programs/AutomationTool/IOS/IOSPlatform.Automation.cs
--- a/Engine/Source/Programs/AutomationTool/IOS/IOSPlatform.Automation.cs
+++ b/Engine/Source/Programs/AutomationTool/IOS/IOSPlatform.Automation.cs
@@ -379,143 +379,6 @@ public class IOSPlatform : Platform
 			}
 		}
 
-		// copy the appropriate plist file over
-		string SourcePListFile = CombinePaths(LocalRoot, "Engine", "Build", "IOS", "UE4Game-Info.plist");
-		if (File.Exists(ProjectRoot + "/Build/IOS/Info.plist"))
-		{
-			SourcePListFile = CombinePaths(ProjectRoot, "Build", "IOS", "Info.plist");
-		}
-		else if (File.Exists(ProjectRoot + "/Build/IOS/" + ShortProjectName + "-Info.plist"))
-		{
-			SourcePListFile = CombinePaths(ProjectRoot, "Build", "IOS", ShortProjectName + "-Info.plist");
-		}
-		else if (Directory.Exists(ProjectRoot + "/Build/IOS"))
-		{
-			// look for any plist file
-			string[] Plists = Directory.GetFiles(ProjectRoot + "/Build/IOS", "*.plist");
-			if (Plists.Length > 0)
-			{
-				SourcePListFile = Plists[0];
-			}
-		}
-
-		//@TODO: This is writing to the engine directory!
-		string SourcePath = CombinePaths((IsCodeBasedProject ? ProjectRoot : LocalRoot + "\\Engine"), "Intermediate", "IOS");
-		string TargetPListFile = Path.Combine(SourcePath, (IsCodeBasedProject ? ShortProjectName : "UE4Game") + "-Info.plist");
-
-		Dictionary<string, string> Replacements = new Dictionary<string, string>();
-		Replacements.Add("${EXECUTABLE_NAME}", (IsCodeBasedProject ? ShortProjectName : "UE4Game"));
-		Replacements.Add("${BUNDLE_IDENTIFIER}", ShortProjectName.Replace("_", ""));
-		CopyFileWithReplacements(SourcePListFile, TargetPListFile, Replacements);
-
-		// Now do the .mobileprovision
-		// install the provision
-		FileInfo DestFileInfo;
-		string InEngineDir = CombinePaths(LocalRoot, "Engine");
-		string ProvisionWithPrefix = CombinePaths(InEngineDir, "Build", "IOS","UE4Game.mobileprovision");
-		string BuildDirectory = CombinePaths(ProjectRoot, "Build", "IOS");
-		if (File.Exists(BuildDirectory + "/" + ShortProjectName + ".mobileprovision"))
-		{
-			ProvisionWithPrefix = BuildDirectory + "/" + ShortProjectName + ".mobileprovision";
-		}
-		else
-		{
-			if (File.Exists(BuildDirectory + "/NotForLicensees/" + ShortProjectName + ".mobileprovision"))
-			{
-				ProvisionWithPrefix = BuildDirectory + "/NotForLicensees/" + ShortProjectName + ".mobileprovision";
-			}
-			else if (!File.Exists(ProvisionWithPrefix))
-			{
-				ProvisionWithPrefix = InEngineDir + "/Build/IOS/NotForLicensees/UE4Game.mobileprovision";
-			}
-		}
-		if (File.Exists(ProvisionWithPrefix))
-		{
-			Directory.CreateDirectory(Environment.GetEnvironmentVariable("HOME") + "/Library/MobileDevice/Provisioning Profiles/");
-			if (File.Exists(Environment.GetEnvironmentVariable("HOME") + "/Library/MobileDevice/Provisioning Profiles/" + ShortProjectName + ".mobileprovision"))
-			{
-				DestFileInfo = new FileInfo(Environment.GetEnvironmentVariable("HOME") + "/Library/MobileDevice/Provisioning Profiles/" + ShortProjectName + ".mobileprovision");
-				DestFileInfo.Attributes = DestFileInfo.Attributes & ~FileAttributes.ReadOnly;
-			}
-			File.Copy(ProvisionWithPrefix, Environment.GetEnvironmentVariable("HOME") + "/Library/MobileDevice/Provisioning Profiles/" + ShortProjectName + ".mobileprovision", true);
-			DestFileInfo = new FileInfo(Environment.GetEnvironmentVariable("HOME") + "/Library/MobileDevice/Provisioning Profiles/" + ShortProjectName + ".mobileprovision");
-			DestFileInfo.Attributes = DestFileInfo.Attributes & ~FileAttributes.ReadOnly;
-		}
-		else
-		{
-			// copy all provisions from the game directory, the engine directory, and the notforlicensees directory
-			// copy all of the provisions from the game directory to the library
-			{
-				if (Directory.Exists(BuildDirectory))
-				{
-					foreach (string Provision in Directory.EnumerateFiles(BuildDirectory, "*.mobileprovision", SearchOption.AllDirectories))
-					{
-						if (!File.Exists(Environment.GetEnvironmentVariable("HOME") + "/Library/MobileDevice/Provisioning Profiles/" + Path.GetFileName(Provision)) || File.GetLastWriteTime(Environment.GetEnvironmentVariable("HOME") + "/Library/MobileDevice/Provisioning Profiles/" + Path.GetFileName(Provision)) < File.GetLastWriteTime(Provision))
-						{
-							if (File.Exists(Environment.GetEnvironmentVariable("HOME") + "/Library/MobileDevice/Provisioning Profiles/" + Path.GetFileName(Provision)))
-							{
-								DestFileInfo = new FileInfo(Environment.GetEnvironmentVariable("HOME") + "/Library/MobileDevice/Provisioning Profiles/" + Path.GetFileName(Provision));
-								DestFileInfo.Attributes = DestFileInfo.Attributes & ~FileAttributes.ReadOnly;
-							}
-							File.Copy(Provision, Environment.GetEnvironmentVariable("HOME") + "/Library/MobileDevice/Provisioning Profiles/" + Path.GetFileName(Provision), true);
-							DestFileInfo = new FileInfo(Environment.GetEnvironmentVariable("HOME") + "/Library/MobileDevice/Provisioning Profiles/" + Path.GetFileName(Provision));
-							DestFileInfo.Attributes = DestFileInfo.Attributes & ~FileAttributes.ReadOnly;
-						}
-					}
-				}
-			}
-
-			// copy all of the provisions from the engine directory to the library
-			{
-				if (Directory.Exists(InEngineDir + "/Build/IOS"))
-				{
-					foreach (string Provision in Directory.EnumerateFiles(InEngineDir + "/Build/IOS", "*.mobileprovision", SearchOption.AllDirectories))
-					{
-						if (!File.Exists(Environment.GetEnvironmentVariable("HOME") + "/Library/MobileDevice/Provisioning Profiles/" + Path.GetFileName(Provision)) || File.GetLastWriteTime(Environment.GetEnvironmentVariable("HOME") + "/Library/MobileDevice/Provisioning Profiles/" + Path.GetFileName(Provision)) < File.GetLastWriteTime(Provision))
-						{
-							if (File.Exists(Environment.GetEnvironmentVariable("HOME") + "/Library/MobileDevice/Provisioning Profiles/" + Path.GetFileName(Provision)))
-							{
-								DestFileInfo = new FileInfo(Environment.GetEnvironmentVariable("HOME") + "/Library/MobileDevice/Provisioning Profiles/" + Path.GetFileName(Provision));
-								DestFileInfo.Attributes = DestFileInfo.Attributes & ~FileAttributes.ReadOnly;
-							}
-							File.Copy(Provision, Environment.GetEnvironmentVariable("HOME") + "/Library/MobileDevice/Provisioning Profiles/" + Path.GetFileName(Provision), true);
-							DestFileInfo = new FileInfo(Environment.GetEnvironmentVariable("HOME") + "/Library/MobileDevice/Provisioning Profiles/" + Path.GetFileName(Provision));
-							DestFileInfo.Attributes = DestFileInfo.Attributes & ~FileAttributes.ReadOnly;
-						}
-					}
-				}
-			}
-		}
-
-		// install the distribution provision
-		ProvisionWithPrefix = InEngineDir + "/Build/IOS/UE4Game_Distro.mobileprovision";
-		if (File.Exists(BuildDirectory + "/" + ShortProjectName + "_Distro.mobileprovision"))
-		{
-			ProvisionWithPrefix = BuildDirectory + "/" + ShortProjectName + "_Distro.mobileprovision";
-		}
-		else
-		{
-			if (File.Exists(BuildDirectory + "/NotForLicensees/" + ShortProjectName + "_Distro.mobileprovision"))
-			{
-				ProvisionWithPrefix = BuildDirectory + "/NotForLicensees/" + ShortProjectName + "_Distro.mobileprovision";
-			}
-			else if (!File.Exists(ProvisionWithPrefix))
-			{
-				ProvisionWithPrefix = InEngineDir + "/Build/IOS/NotForLicensees/UE4Game_Distro.mobileprovision";
-			}
-		}
-		if (File.Exists(ProvisionWithPrefix))
-		{
-			if (File.Exists(Environment.GetEnvironmentVariable("HOME") + "/Library/MobileDevice/Provisioning Profiles/" + ShortProjectName + "_Distro.mobileprovision"))
-			{
-				DestFileInfo = new FileInfo(Environment.GetEnvironmentVariable("HOME") + "/Library/MobileDevice/Provisioning Profiles/" + ShortProjectName + "_Distro.mobileprovision");
-				DestFileInfo.Attributes = DestFileInfo.Attributes & ~FileAttributes.ReadOnly;
-			}
-			File.Copy(ProvisionWithPrefix, Environment.GetEnvironmentVariable("HOME") + "/Library/MobileDevice/Provisioning Profiles/" + ShortProjectName + "_Distro.mobileprovision", true);
-			DestFileInfo = new FileInfo(Environment.GetEnvironmentVariable("HOME") + "/Library/MobileDevice/Provisioning Profiles/" + ShortProjectName + "_Distro.mobileprovision");
-			DestFileInfo.Attributes = DestFileInfo.Attributes & ~FileAttributes.ReadOnly;
-		}
-
 		return XcodeProj;
 	}
 
@@ -524,6 +387,15 @@ public class IOSPlatform : Platform
 		// check for the proper xcodeproject
 		bool bWasGenerated = false;
 		string XcodeProj = EnsureXcodeProjectExists (RawProjectPath, LocalRoot, ProjectName, ProjectDirectory, IsCode, out bWasGenerated);
+
+		// ensure the correct data os acrpss
+		var DeployHandler = UEBuildDeploy.GetBuildDeploy(UnrealTargetPlatform.IOS);
+		DeployHandler.PrepForUATPackageOrDeploy(ProjectName,
+			Path.GetDirectoryName(RawProjectPath),
+			CombinePaths(BaseDirectory, GameName),
+			CombinePaths(LocalRoot, "Engine"),
+			Distribution, "");
+
 		string Arguments = "UBT_NO_POST_DEPLOY=true";
 		Arguments += " /usr/bin/xcrun xcodebuild build -project \"" + XcodeProj + "\"";
 		Arguments += " -scheme '";
@@ -668,53 +540,6 @@ public class IOSPlatform : Platform
 		}
 	}
 
-	private static void CopyFileWithReplacements(string SourceFilename, string DestFilename, Dictionary<string, string> Replacements)
-	{
-		Console.WriteLine("Copying {0} to {1} with replacements...", SourceFilename, DestFilename);
-
-		if (!File.Exists(SourceFilename))
-		{
-			return;
-		}
-
-		// make the dst filename with the same structure as it was in SourceDir
-		if (File.Exists(DestFilename))
-		{
-			File.Delete(DestFilename);
-		}
-
-		// make the subdirectory if needed
-		string DestSubdir = Path.GetDirectoryName(DestFilename);
-		if (!Directory.Exists(DestSubdir))
-		{
-			Directory.CreateDirectory(DestSubdir);
-		}
-
-		// some files are handled specially
-		string Ext = Path.GetExtension(SourceFilename);
-		if (Ext == ".plist")
-		{
-			string Contents = File.ReadAllText(SourceFilename);
-
-			// replace some varaibles
-			foreach (var Pair in Replacements)
-			{
-				Contents = Contents.Replace(Pair.Key, Pair.Value);
-			}
-
-			// write out file
-			File.WriteAllText(DestFilename, Contents);
-		}
-		else
-		{
-			File.Copy(SourceFilename, DestFilename);
-
-			// remove any read only flags
-			FileInfo DestFileInfo = new FileInfo(DestFilename);
-			DestFileInfo.Attributes = DestFileInfo.Attributes & ~FileAttributes.ReadOnly;
-		}
-	}
-
 	public override void GetFilesToDeployOrStage(ProjectParams Params, DeploymentContext SC)
 	{
 		//		if (UnrealBuildTool.BuildHostPlatform.Current.Platform != UnrealTargetPlatform.Mac)
@@ -743,33 +568,13 @@ public class IOSPlatform : Platform
 			// copy the plist (only if code signing, as it's protected by the code sign blob in the executable and can't be modified independently)
 			if (GetCodeSignDesirability(Params))
 			{
-				string SourcePListFile = CombinePaths(SC.LocalRoot, "Engine", "Build", "IOS", "UE4Game-Info.plist");
-				if (File.Exists(SC.ProjectRoot + "/Build/IOS/Info.plist"))
-				{
-					SourcePListFile = CombinePaths(SC.ProjectRoot, "Build", "IOS", "Info.plist");
-				}
-				else if (File.Exists(SC.ProjectRoot + "/Build/IOS/" + SC.ShortProjectName + "-Info.plist"))
-				{
-					SourcePListFile = CombinePaths(SC.ProjectRoot, "Build", "IOS", SC.ShortProjectName + "-Info.plist");
-				}
-				else if (Directory.Exists(SC.ProjectRoot + "/Build/IOS"))
-				{
-					// look for any plist file
-					string[] Plists = Directory.GetFiles(SC.ProjectRoot + "/Build/IOS", "*.plist");
-					if (Plists.Length > 0)
-					{
-						SourcePListFile = Plists[0];
-					}
-				}
-
-				//@TODO: This is writing to the engine directory!
 				string SourcePath = CombinePaths((SC.IsCodeBasedProject ? SC.ProjectRoot : SC.LocalRoot + "\\Engine"), "Intermediate", "IOS");
 				string TargetPListFile = Path.Combine(SourcePath, (SC.IsCodeBasedProject ? SC.ShortProjectName : "UE4Game") + "-Info.plist");
-
-				Dictionary<string, string> Replacements = new Dictionary<string, string>();
-				Replacements.Add("${EXECUTABLE_NAME}", (SC.IsCodeBasedProject ? SC.ShortProjectName : "UE4Game"));
-				Replacements.Add("${BUNDLE_IDENTIFIER}", SC.ShortProjectName.Replace("_", ""));
-				CopyFileWithReplacements(SourcePListFile, TargetPListFile, Replacements);
+				if (!File.Exists(TargetPListFile))
+				{
+					// ensure the plist, entitlements, and provision files are properly copied
+					UnrealBuildTool.IOS.UEDeployIOS.GeneratePList((SC.IsCodeBasedProject ? SC.ProjectRoot : SC.LocalRoot + "\\Engine"), !SC.IsCodeBasedProject, (SC.IsCodeBasedProject ? SC.ShortProjectName : "UE4Game"), SC.ShortProjectName, SC.LocalRoot + "\\Engine", "");
+				}
 
 				SC.StageFiles(StagedFileType.NonUFS, SourcePath, Path.GetFileName(TargetPListFile), false, null, "", false, false, "Info.plist");
 			}
diff --git a/Engine/Source/Programs/IOS/iPhonePackager/CodeSigning.cs b/Engine/Source/Programs/IOS/iPhonePackager/CodeSigning.cs
--- a/Engine/Source/Programs/IOS/iPhonePackager/CodeSigning.cs
+++ b/Engine/Source/Programs/IOS/iPhonePackager/CodeSigning.cs
@@ -81,11 +81,11 @@ namespace iPhonePackager
 			string plistFile = Config.EngineBuildDirectory + "/UE4Game-Info.plist";
 			if (!string.IsNullOrEmpty(Config.ProjectFile))
 			{
-				plistFile = Path.GetDirectoryName(Config.ProjectFile) + "/Build/IOS/Info.plist";
+				plistFile = Path.GetDirectoryName(Config.ProjectFile) + "/Intermediate/IOS/" + Path.GetFileNameWithoutExtension(Config.ProjectFile) + "-Info.plist";
+
 				if (!File.Exists(plistFile))
 				{
-					plistFile = Path.GetDirectoryName(Config.ProjectFile) + "/Build/IOS/" + Path.GetFileNameWithoutExtension(Config.ProjectFile) + "-Info.plist";
-
+					plistFile = Config.IntermediateDirectory + "/UE4Game-Info.plist";
 					if (!File.Exists(plistFile))
 					{
 						plistFile = Config.EngineBuildDirectory + "/UE4Game-Info.plist";
@@ -725,7 +725,7 @@ namespace iPhonePackager
 				}
 				byte[] Data = OutputExeStream.ToArray();
 				Data.CopyTo(FinalExeData, (long)CurrentStreamOffset);
-				CurrentStreamOffset += DesiredExecutableLength;
+				CurrentStreamOffset += (ulong)Data.Length;
 
 				// increment the architecture index
 				ArchIndex++;
diff --git a/Engine/Source/Programs/IOS/iPhonePackager/Config.cs b/Engine/Source/Programs/IOS/iPhonePackager/Config.cs
--- a/Engine/Source/Programs/IOS/iPhonePackager/Config.cs
+++ b/Engine/Source/Programs/IOS/iPhonePackager/Config.cs
@@ -111,6 +111,21 @@ namespace iPhonePackager
 			}
 		}
 
+		/// <summary>
+		/// The local build intermediate directory (on PC)
+		/// </summary>
+		public static string IntermediateDirectory
+		{
+			get
+			{
+				string IntermediateGameBuildDir = GameDirectory + @"\Intermediate\IOS";
+				// if the normal Build dir exists, return it, otherwise, use the program Resources directory
+				return Path.GetFullPath(Directory.Exists(IntermediateGameBuildDir) ?
+					IntermediateGameBuildDir :
+					BuildDirectory);
+			}
+		}
+
 		static Config()
 		{
 			if (Environment.OSVersion.Platform == PlatformID.MacOSX || Environment.OSVersion.Platform == PlatformID.Unix) {
diff --git a/Engine/Source/Programs/IOS/iPhonePackager/Utilities.cs b/Engine/Source/Programs/IOS/iPhonePackager/Utilities.cs
--- a/Engine/Source/Programs/IOS/iPhonePackager/Utilities.cs
+++ b/Engine/Source/Programs/IOS/iPhonePackager/Utilities.cs
@@ -574,17 +574,23 @@ namespace iPhonePackager
 
 		static public string GetPrecompileSourcePListFilename()
 		{
-			// Check for a premade one
-			string SourceName = FileOperations.FindPrefixedFile(Config.BuildDirectory, Program.GameName + "-Info.plist");
+			// check for one in he project directory
+			string SourceName = FileOperations.FindPrefixedFile(Config.IntermediateDirectory, Program.GameName + "-Info.plist");
 
 			if (!File.Exists(SourceName))
 			{
-				// fallback to the shared one
-				SourceName = FileOperations.FindPrefixedFile(Config.EngineBuildDirectory, "UE4Game-Info.plist");
+				// Check for a premade one
+				SourceName = FileOperations.FindPrefixedFile(Config.BuildDirectory, Program.GameName + "-Info.plist");
 
 				if (!File.Exists(SourceName))
 				{
-					Program.Log("Failed to find " + Program.GameName + "-Info.plist. Please create new .plist or copy a base .plist from a provided game sample.");
+					// fallback to the shared one
+					SourceName = FileOperations.FindPrefixedFile(Config.EngineBuildDirectory, "UE4Game-Info.plist");
+
+					if (!File.Exists(SourceName))
+					{
+						Program.Log("Failed to find " + Program.GameName + "-Info.plist. Please create new .plist or copy a base .plist from a provided game sample.");
+					}
 				}
 			}
 
diff --git a/Engine/Source/Programs/UnrealBuildTool/IOS/IOSToolChain.cs b/Engine/Source/Programs/UnrealBuildTool/IOS/IOSToolChain.cs
--- a/Engine/Source/Programs/UnrealBuildTool/IOS/IOSToolChain.cs
+++ b/Engine/Source/Programs/UnrealBuildTool/IOS/IOSToolChain.cs
@@ -1297,6 +1297,13 @@ namespace UnrealBuildTool
 				// Generate the stub
 				if (BuildConfiguration.bCreateStubIPA || bUseDangerouslyFastMode)
 				{
+					// ensure the plist, entitlements, and provision files are properly copied
+					var DeployHandler = UEBuildDeploy.GetBuildDeploy(Target.Platform);
+					if (DeployHandler != null)
+					{
+						DeployHandler.PrepTargetForDeployment(Target);
+					}
+
 					if (!bUseDangerouslyFastMode)
 					{
 						// generate the dummy project so signing works
diff --git a/Engine/Source/Programs/UnrealBuildTool/IOS/UEDeployIOS.cs b/Engine/Source/Programs/UnrealBuildTool/IOS/UEDeployIOS.cs
--- a/Engine/Source/Programs/UnrealBuildTool/IOS/UEDeployIOS.cs
+++ b/Engine/Source/Programs/UnrealBuildTool/IOS/UEDeployIOS.cs
@@ -10,7 +10,7 @@ using System.Diagnostics;
 
 namespace UnrealBuildTool.IOS
 {
-	class UEDeployIOS : UEBuildDeploy
+	public class UEDeployIOS : UEBuildDeploy
 	{
 		/**
 		 *	Register the platform with the UEBuildDeploy class
@@ -21,83 +21,430 @@ namespace UnrealBuildTool.IOS
 			UEBuildDeploy.RegisterBuildDeploy(UnrealTargetPlatform.IOS, this);
 		}
 
-		private static void CopyFileWithReplacements(string SourceFilename, string DestFilename, Dictionary<string, string> Replacements, List<string> AdditionalLines)
+		class VersionUtilities
 		{
-			if (!File.Exists(SourceFilename))
+			public static string BuildDirectory
 			{
-				return;
+				get;
+				set;
+			}
+			public static string GameName
+			{
+				get;
+				set;
 			}
 
-			// make the dst filename with the same structure as it was in SourceDir
-			if (File.Exists(DestFilename))
+			static string RunningVersionFilename
 			{
-				File.Delete(DestFilename);
+				get { return Path.Combine(BuildDirectory, GameName + ".PackageVersionCounter"); }
 			}
 
-			// make the subdirectory if needed
-			string DestSubdir = Path.GetDirectoryName(DestFilename);
-			if (!Directory.Exists(DestSubdir))
+			/// <summary>
+			/// Reads the GameName.PackageVersionCounter from disk and bumps the minor version number in it
+			/// </summary>
+			/// <returns></returns>
+			public static string ReadRunningVersion()
 			{
-				Directory.CreateDirectory(DestSubdir);
+				string CurrentVersion = "0.0";
+				if (File.Exists(RunningVersionFilename))
+				{
+					CurrentVersion = File.ReadAllText(RunningVersionFilename);
+				}
+
+				return CurrentVersion;
 			}
 
-			// some files are handled specially
-			string Ext = Path.GetExtension(SourceFilename);
-			if (Ext == ".plist")
+			/// <summary>
+			/// Pulls apart a version string of one of the two following formats:
+			///	  "7301.15 11-01 10:28"   (Major.Minor Date Time)
+			///	  "7486.0"  (Major.Minor)
+			/// </summary>
+			/// <param name="CFBundleVersion"></param>
+			/// <param name="VersionMajor"></param>
+			/// <param name="VersionMinor"></param>
+			/// <param name="TimeStamp"></param>
+			public static void PullApartVersion(string CFBundleVersion, out int VersionMajor, out int VersionMinor, out string TimeStamp)
 			{
-				string[] Contents = File.ReadAllLines(SourceFilename);
-				StringBuilder NewContents = new StringBuilder();
+				// Expecting source to be like "7301.15 11-01 10:28" or "7486.0"
+				string[] Parts = CFBundleVersion.Split(new char[] { ' ' });
 
-				int LastDictLine = 0;
-				for (int LineIndex = Contents.Length - 1; LineIndex >= 0; LineIndex--)
+				// Parse the version string
+				string[] VersionParts = Parts[0].Split(new char[] { '.' });
+
+				if (!int.TryParse(VersionParts[0], out VersionMajor))
 				{
-					if (Contents[LineIndex].Trim() == "</dict>")
-					{
-						LastDictLine = LineIndex;
-						break;
-					}
+					VersionMajor = 0;
 				}
 
-				// replace some varaibles
-				for (int LineIndex = 0; LineIndex < Contents.Length; LineIndex++)
+				if ((VersionParts.Length < 2) || (!int.TryParse(VersionParts[1], out VersionMinor)))
 				{
-					string Line = Contents[LineIndex];
+					VersionMinor = 0;
+				}
 
-					// inject before the last line
-					if (LineIndex == LastDictLine)
-					{
-						foreach (string ExtraLine in AdditionalLines)
-						{
-							string FixedLine = ExtraLine;
-							foreach (var Pair in Replacements)
-							{
-								FixedLine = FixedLine.Replace(Pair.Key, Pair.Value);
-							}
-							NewContents.Append(FixedLine + Environment.NewLine);
-						}
-					}
+				TimeStamp = "";
+				if (Parts.Length > 1)
+				{
+					TimeStamp = String.Join(" ", Parts, 1, Parts.Length - 1);
+				}
+			}
 
-					foreach (var Pair in Replacements)
-					{
-						Line = Line.Replace(Pair.Key, Pair.Value);
-					}
+			public static string ConstructVersion(int MajorVersion, int MinorVersion)
+			{
+				return String.Format("{0}.{1}", MajorVersion, MinorVersion);
+			}
 
-					NewContents.Append(Line + Environment.NewLine);
+			/// <summary>
+			/// Parses the version string (expected to be of the form major.minor or major)
+			/// Also parses the major.minor from the running version file and increments it's minor by 1.
+			/// 
+			/// If the running version major matches and the running version minor is newer, then the bundle version is updated.
+			/// 
+			/// In either case, the running version is set to the current bundle version number and written back out.
+			/// </summary>
+			/// <returns>The (possibly updated) bundle version</returns>
+			public static string CalculateUpdatedMinorVersionString(string CFBundleVersion)
+			{
+				// Read the running version and bump it
+				int RunningMajorVersion;
+				int RunningMinorVersion;
+
+				string DummyDate;
+				string RunningVersion = ReadRunningVersion();
+				PullApartVersion(RunningVersion, out RunningMajorVersion, out RunningMinorVersion, out DummyDate);
+				RunningMinorVersion++;
+
+				// Read the passed in version and bump it
+				int MajorVersion;
+				int MinorVersion;
+				PullApartVersion(CFBundleVersion, out MajorVersion, out MinorVersion, out DummyDate);
+				MinorVersion++;
+
+				// Combine them if the stub time is older
+				if ((RunningMajorVersion == MajorVersion) && (RunningMinorVersion > MinorVersion))
+				{
+					// A subsequent cook on the same sync, the only time that we stomp on the stub version
+					MinorVersion = RunningMinorVersion;
+				}
 
+				// Combine them together
+				string ResultVersionString = ConstructVersion(MajorVersion, MinorVersion);
+
+				// Update the running version file
+				Directory.CreateDirectory(Path.GetDirectoryName(RunningVersionFilename));
+				File.WriteAllText(RunningVersionFilename, ResultVersionString);
+
+				return ResultVersionString;
+			}
+
+			/// <summary>
+			/// Updates the minor version in the CFBundleVersion key of the specified PList if this is a new package.
+			/// Also updates the key EpicAppVersion with the bundle version and the current date/time (no year)
+			/// </summary>
+			public static string UpdateBundleVersion(string OldPList)
+			{
+				string CFBundleVersion;
+				int Index = OldPList.IndexOf("CFBundleVersion");
+				if (Index != -1)
+				{
+					int Start = OldPList.IndexOf("<string>", Index) + ("<string>").Length;
+					CFBundleVersion = OldPList.Substring(Start, OldPList.IndexOf("</string>", Index) - Start);
+					CFBundleVersion = CalculateUpdatedMinorVersionString(CFBundleVersion);
+				}
+				else
+				{
+					CFBundleVersion = "0.0";
 				}
 
-				File.WriteAllText(DestFilename, NewContents.ToString());
+				return CFBundleVersion;
+			}
+		}
+
+		public static bool GeneratePList(string ProjectDirectory, bool bIsUE4Game, string GameName, string ProjectName, string InEngineDir, string AppDirectory)
+		{
+			// generate the Info.plist for future use
+			string BuildDirectory = ProjectDirectory + "/Build/IOS";
+			bool bSkipDefaultPNGs = false;
+			string IntermediateDirectory = (bIsUE4Game ? InEngineDir : ProjectDirectory) + "/Intermediate/IOS";
+			string PListFile = IntermediateDirectory + "/" + GameName + "-Info.plist";
+			VersionUtilities.BuildDirectory = BuildDirectory;
+			VersionUtilities.GameName = GameName;
+
+			// read the old file
+			string OldPListData = File.Exists(PListFile) ? File.ReadAllText(PListFile) : "";
+
+			// determine if there is a launch.xib
+			string LaunchXib = InEngineDir + "/Build/IOS/Resources/Interface/LaunchScreen.xib";
+			if (File.Exists(BuildDirectory + "/Resources/Interface/LaunchScreen.xib"))
+			{
+				LaunchXib = BuildDirectory + "/Resources/Interface/LaunchScreen.xib";
+			}
+
+			// get the settings from the ini file
+			// plist replacements
+			ConfigCacheIni Ini = new ConfigCacheIni(UnrealTargetPlatform.IOS, "Engine", UnrealBuildTool.GetUProjectPath());
+
+			// orientations
+			string SupportedOrientations = "";
+			bool bSupported = true;
+			Ini.GetBool("/Script/IOSRuntimeSettings.IOSRuntimeSettings", "bSupportsPortraitOrientation", out bSupported);
+			SupportedOrientations += bSupported ? "\t\t<string>UIInterfaceOrientationPortrait</string>\n" : "";
+			Ini.GetBool("/Script/IOSRuntimeSettings.IOSRuntimeSettings", "bSupportsUpsideDownOrientation", out bSupported);
+			SupportedOrientations += bSupported ? "\t\t<string>UIInterfaceOrientationPortraitUpsideDown</string>\n" : "";
+			Ini.GetBool("/Script/IOSRuntimeSettings.IOSRuntimeSettings", "bSupportsLandscapeLeftOrientation", out bSupported);
+			SupportedOrientations += bSupported ? "\t\t<string>UIInterfaceOrientationLandscapeLeft</string>\n" : "";
+			Ini.GetBool("/Script/IOSRuntimeSettings.IOSRuntimeSettings", "bSupportsLandscapeRightOrientation", out bSupported);
+			SupportedOrientations += bSupported ? "\t\t<string>UIInterfaceOrientationLandscapeRight</string>\n" : "";
+
+			// bundle display name
+			string BundleDisplayName;
+			Ini.GetString("/Script/IOSRuntimeSettings.IOSRuntimeSettings", "BundleDisplayName", out BundleDisplayName);
+
+			// bundle identifier
+			string BundleIdentifier;
+			Ini.GetString("/Script/IOSRuntimeSettings.IOSRuntimeSettings", "BundleIdentifier", out BundleIdentifier);
+
+			// bundle name
+			string BundleName;
+			Ini.GetString("/Script/IOSRuntimeSettings.IOSRuntimeSettings", "BundleName", out BundleName);
+
+			// short version string
+			string BundleShortVersion;
+			Ini.GetString("/Script/IOSRuntimeSettings.IOSRuntimeSettings", "VersionInfo", out BundleShortVersion);
+
+			// required capabilities
+			string RequiredCaps = "\t\t<string>armv7</string>\n";
+			Ini.GetBool("/Script/IOSRuntimeSettings.IOSRuntimeSettings", "bSupportsOpenGLES2", out bSupported);
+			RequiredCaps += bSupported ? "\t\t<string>opengles-2</string>\n" : "";
+			if (!bSupported)
+			{
+				Ini.GetBool("/Script/IOSRuntimeSettings.IOSRuntimeSettings", "bSupportsMetal", out bSupported);
+				RequiredCaps += bSupported ? "\t\t<string>metal</string>\n" : "";
+			}
+
+			// minimum iOS version
+			string MinVersion;
+			if (Ini.GetString("/Script/IOSRuntimeSettings.IOSRuntimeSettings", "MinimumiOSVersion", out MinVersion))
+			{
+				switch (MinVersion)
+				{
+					case "IOS_61":
+						MinVersion = "6.1";
+						break;
+					case "IOS_7":
+						MinVersion = "7.0";
+						break;
+					case "IOS_8":
+						MinVersion = "8.0";
+						break;
+				}
 			}
 			else
 			{
-				File.Copy(SourceFilename, DestFilename);
+				MinVersion = "6.1";
+			}
 
-				// remove any read only flags
-				FileInfo DestFileInfo = new FileInfo(DestFilename);
-				DestFileInfo.Attributes = DestFileInfo.Attributes & ~FileAttributes.ReadOnly;
+			// extra plist data
+			string ExtraData = "";
+			Ini.GetString("/Script/IOSRuntimeSettings.IOSRuntimeSettings", "AdditionalPlistData", out ExtraData);
+
+			// generate the plist file
+			StringBuilder Text = new StringBuilder();
+			Text.AppendLine("<?xml version=\"1.0\" encoding=\"UTF-8\"?>");
+			Text.AppendLine("<!DOCTYPE plist PUBLIC \"-//Apple//DTD PLIST 1.0//EN\" \"http://www.apple.com/DTDs/PropertyList-1.0.dtd\">");
+			Text.AppendLine("<plist version=\"1.0\">");
+			Text.AppendLine("<dict>");
+			Text.AppendLine("\t<key>CFBundleURLTypes</key>");
+			Text.AppendLine("\t<array>");
+			Text.AppendLine("\t\t<dict>");
+			Text.AppendLine("\t\t\t<key>CFBundleURLName</key>");
+			Text.AppendLine("\t\t\t<string>com.Epic.Unreal</string>");
+			Text.AppendLine("\t\t\t<key>CFBundleURLSchemes</key>");
+			Text.AppendLine("\t\t\t<array>");
+			Text.AppendLine(string.Format("\t\t\t\t<string>{0}</string>", bIsUE4Game ? "UE4Game" : GameName));
+			Text.AppendLine("\t\t\t</array>");
+			Text.AppendLine("\t\t</dict>");
+			Text.AppendLine("\t</array>");
+			Text.AppendLine("\t<key>CFBundleDevelopmentRegion</key>");
+			Text.AppendLine("\t<string>English</string>");
+			Text.AppendLine("\t<key>CFBundleDisplayName</key>");
+			Text.AppendLine(string.Format("\t<string>{0}</string>", BundleDisplayName.Replace("{PROJECT_NAME}", ProjectName).Replace("_","")));
+			Text.AppendLine("\t<key>CFBundleExecutable</key>");
+			Text.AppendLine(string.Format("\t<string>{0}</string>", bIsUE4Game ? "UE4Game" : GameName.Replace("_", "")));
+			Text.AppendLine("\t<key>CFBundleIdentifier</key>");
+			Text.AppendLine(string.Format("\t<string>{0}</string>", BundleIdentifier.Replace("{PROJECT_NAME}", ProjectName).Replace("_","")));
+			Text.AppendLine("\t<key>CFBundleInfoDictionaryVersion</key>");
+			Text.AppendLine("\t<string>6.0</string>");
+			Text.AppendLine("\t<key>CFBundleName</key>");
+			Text.AppendLine(string.Format("\t<string>{0}</string>", BundleName.Replace("{PROJECT_NAME}", ProjectName).Replace("_","")));
+			Text.AppendLine("\t<key>CFBundlePackageType</key>");
+			Text.AppendLine("\t<string>APPL</string>");
+			Text.AppendLine("\t<key>CFBundleSignature</key>");
+			Text.AppendLine("\t<string>????</string>");
+			Text.AppendLine("\t<key>CFBundleVersion</key>");
+			Text.AppendLine(string.Format("\t<string>{0}</string>", VersionUtilities.UpdateBundleVersion(OldPListData)));
+			Text.AppendLine("\t<key>CFBundleShortVersionString</key>");
+			Text.AppendLine(string.Format("\t<string>{0}</string>", BundleShortVersion));
+			Text.AppendLine("\t<key>LSRequiresIPhoneOS</key>");
+			Text.AppendLine("\t<true/>");
+			Text.AppendLine("\t<key>UIStatusBarHidden</key>");
+			Text.AppendLine("\t<true/>");
+			Text.AppendLine("\t<key>UISupportedInterfaceOrientations</key>");
+			Text.AppendLine("\t<array>");
+			foreach (string Line in SupportedOrientations.Split("\r\n".ToCharArray()))
+			{
+				if (!string.IsNullOrWhiteSpace(Line))
+				{
+					Text.AppendLine(Line);
+				}
 			}
-		}
+			Text.AppendLine("\t</array>");
+			Text.AppendLine("\t<key>UIRequiredDeviceCapabilities</key>");
+			Text.AppendLine("\t<array>");
+			foreach (string Line in RequiredCaps.Split("\r\n".ToCharArray()))
+			{
+				if (!string.IsNullOrWhiteSpace(Line))
+				{
+					Text.AppendLine(Line);
+				}
+			}
+			Text.AppendLine("\t</array>");
+			Text.AppendLine("\t<key>CFBundleIcons</key>");
+			Text.AppendLine("\t<dict>");
+			Text.AppendLine("\t\t<key>CFBundlePrimaryIcon</key>");
+			Text.AppendLine("\t\t<dict>");
+			Text.AppendLine("\t\t\t<key>CFBundleIconFiles</key>");
+			Text.AppendLine("\t\t\t<array>");
+			Text.AppendLine("\t\t\t\t<string>Icon29.png</string>");
+			Text.AppendLine("\t\t\t\t<string>[REDACTED_EMAIL]</string>");
+			Text.AppendLine("\t\t\t\t<string>Icon40.png</string>");
+			Text.AppendLine("\t\t\t\t<string>[REDACTED_EMAIL]</string>");
+			Text.AppendLine("\t\t\t\t<string>Icon57.png</string>");
+			Text.AppendLine("\t\t\t\t<string>[REDACTED_EMAIL]</string>");
+			Text.AppendLine("\t\t\t\t<string>[REDACTED_EMAIL]</string>");
+			Text.AppendLine("\t\t\t</array>");
+			Text.AppendLine("\t\t\t<key>UIPrerenderedIcon</key>");
+			Text.AppendLine("\t\t\t<true/>");
+			Text.AppendLine("\t\t</dict>");
+			Text.AppendLine("\t</dict>");
+			Text.AppendLine("\t<key>CFBundleIcons~ipad</key>");
+			Text.AppendLine("\t<dict>");
+			Text.AppendLine("\t\t<key>CFBundlePrimaryIcon</key>");
+			Text.AppendLine("\t\t<dict>");
+			Text.AppendLine("\t\t\t<key>CFBundleIconFiles</key>");
+			Text.AppendLine("\t\t\t<array>");
+			Text.AppendLine("\t\t\t\t<string>Icon29.png</string>");
+			Text.AppendLine("\t\t\t\t<string>[REDACTED_EMAIL]</string>");
+			Text.AppendLine("\t\t\t\t<string>Icon40.png</string>");
+			Text.AppendLine("\t\t\t\t<string>[REDACTED_EMAIL]</string>");
+			Text.AppendLine("\t\t\t\t<string>Icon50.png</string>");
+			Text.AppendLine("\t\t\t\t<string>[REDACTED_EMAIL]</string>");
+			Text.AppendLine("\t\t\t\t<string>Icon72.png</string>");
+			Text.AppendLine("\t\t\t\t<string>[REDACTED_EMAIL]</string>");
+			Text.AppendLine("\t\t\t\t<string>Icon76.png</string>");
+			Text.AppendLine("\t\t\t\t<string>[REDACTED_EMAIL]</string>");
+			Text.AppendLine("\t\t\t</array>");
+			Text.AppendLine("\t\t\t<key>UIPrerenderedIcon</key>");
+			Text.AppendLine("\t\t\t<true/>");
+			Text.AppendLine("\t\t</dict>");
+			Text.AppendLine("\t</dict>");
+ 			if (File.Exists(LaunchXib))
+			{
+				// TODO: compile the xib via remote tool
+				Text.AppendLine("\t<key>UILaunchStoryboardName</key>");
+				Text.AppendLine("\t<string>LaunchScreen</string>");
+				bSkipDefaultPNGs = true;
+			}
+			else
+			{
+				// this is a temp way to inject the iphone 6 images without needing to upgrade everyone's plist
+				// eventually we want to generate this based on what the user has set in the project settings
+				string[] IPhoneConfigs =  
+					{ 
+						"Default-IPhone6", "Landscape", "{375, 667}", 
+						"Default-IPhone6", "Portrait", "{375, 667}", 
+						"Default-IPhone6Plus-Landscape", "Landscape", "{414, 736}", 
+						"Default-IPhone6Plus-Portrait", "Portrait", "{414, 736}", 
+						"Default", "Landscape", "{320, 480}",
+						"Default", "Portrait", "{320, 480}",
+						"Default-568h", "Landscape", "{320, 568}",
+						"Default-568h", "Portrait", "{320, 568}",
+					};
+
+				Text.AppendLine("\t<key>UILaunchImages~iphone</key>");
+				Text.AppendLine("\t<array>");
+				for (int ConfigIndex = 0; ConfigIndex < IPhoneConfigs.Length; ConfigIndex += 3)
+				{
+					Text.AppendLine("\t\t<dict>");
+					Text.AppendLine("\t\t\t<key>UILaunchImageMinimumOSVersion</key>");
+					Text.AppendLine("\t\t\t<string>8.0</string>");
+					Text.AppendLine("\t\t\t<key>UILaunchImageName</key>");
+					Text.AppendLine(string.Format("\t\t\t<string>{0}</string>", IPhoneConfigs[ConfigIndex + 0]));
+					Text.AppendLine("\t\t\t<key>UILaunchImageOrientation</key>");
+					Text.AppendLine(string.Format("\t\t\t<string>{0}</string>", IPhoneConfigs[ConfigIndex + 1]));
+					Text.AppendLine("\t\t\t<key>UILaunchImageSize</key>");
+					Text.AppendLine(string.Format("\t\t\t<string>{0}</string>", IPhoneConfigs[ConfigIndex + 2]));
+					Text.AppendLine("\t\t</dict>");
+				}
 
+				// close it out
+				Text.AppendLine("\t</array>");
+			}
+			Text.AppendLine("\t<key>UILaunchImages~ipad</key>");
+			Text.AppendLine("\t<array>");
+			Text.AppendLine("\t\t<dict>");
+			Text.AppendLine("\t\t\t<key>UILaunchImageMinimumOSVersion</key>");
+			Text.AppendLine("\t\t\t<string>7.0</string>");
+			Text.AppendLine("\t\t\t<key>UILaunchImageName</key>");
+			Text.AppendLine("\t\t\t<string>Default-Landscape</string>");
+			Text.AppendLine("\t\t\t<key>UILaunchImageOrientation</key>");
+			Text.AppendLine("\t\t\t<string>Landscape</string>");
+			Text.AppendLine("\t\t\t<key>UILaunchImageSize</key>");
+			Text.AppendLine("\t\t\t<string>{768, 1024}</string>");
+			Text.AppendLine("\t\t</dict>");
+			Text.AppendLine("\t\t<dict>");
+			Text.AppendLine("\t\t\t<key>UILaunchImageMinimumOSVersion</key>");
+			Text.AppendLine("\t\t\t<string>7.0</string>");
+			Text.AppendLine("\t\t\t<key>UILaunchImageName</key>");
+			Text.AppendLine("\t\t\t<string>Default-Portrait</string>");
+			Text.AppendLine("\t\t\t<key>UILaunchImageOrientation</key>");
+			Text.AppendLine("\t\t\t<string>Portrait</string>");
+			Text.AppendLine("\t\t\t<key>UILaunchImageSize</key>");
+			Text.AppendLine("\t\t\t<string>{768, 1024}</string>");
+			Text.AppendLine("\t\t</dict>");
+			Text.AppendLine("\t</array>");
+			Text.AppendLine("\t<key>CFBundleSupportedPlatforms</key>");
+			Text.AppendLine("\t<array>");
+			Text.AppendLine("\t\t<string>iPhoneOS</string>");
+			Text.AppendLine("\t</array>");
+			Text.AppendLine("\t<key>MinimumOSVersion</key>");
+			Text.AppendLine(string.Format("\t<string>{0}</string>", MinVersion));
+			if (!string.IsNullOrEmpty(ExtraData))
+			{
+				ExtraData = ExtraData.Replace("\\n", "\n");
+				foreach (string Line in ExtraData.Split("\r\n".ToCharArray()))
+				{
+					if (!string.IsNullOrWhiteSpace(Line))
+					{
+						Text.AppendLine("\t" + Line);
+					}
+				}
+			}
+			Text.AppendLine("</dict>");
+			Text.AppendLine("</plist>");
+
+			// write the file out
+			if (!Directory.Exists(IntermediateDirectory))
+			{
+				Directory.CreateDirectory(IntermediateDirectory);
+			}
+			File.WriteAllText(PListFile, Text.ToString());
+			if (BuildHostPlatform.Current.Platform == UnrealTargetPlatform.Mac)
+			{
+				File.WriteAllText(AppDirectory + "/Info.plist", Text.ToString());
+			}
+
+			return bSkipDefaultPNGs;
+		}
 
 		public override bool PrepForUATPackageOrDeploy(string InProjectName, string InProjectDirectory, string InExecutablePath, string InEngineDir, bool bForDistribution, string CookFlavor)
 		{
@@ -254,82 +601,7 @@ namespace UnrealBuildTool.IOS
 				LaunchXib = BuildDirectory + "/Resources/Interface/LaunchScreen.xib";
 			}
 
-			List<string> PListAdditionalLines = new List<string>();
-			Dictionary<string, string> Replacements = new Dictionary<string, string>();
-
-			bool bSkipDefaultPNGs = false;
-			if (File.Exists(LaunchXib))
-			{
-				string CommandLine = string.Format("--target-device iphone --target-device ipad --errors --warnings --notices --module {0} --minimum-deployment-target 8.0 --auto-activate-custom-fonts --output-format human-readable-text --compile {1}/LaunchScreen.nib {2}", InProjectName, Path.GetFullPath(AppDirectory), Path.GetFileName(LaunchXib));
-
-				// now we need to zipalign the apk to the final destination (to 4 bytes, must be 4)
-				ProcessStartInfo IBToolStartInfo = new ProcessStartInfo();
-				IBToolStartInfo.WorkingDirectory = Path.GetDirectoryName(LaunchXib);
-				IBToolStartInfo.FileName = "/Applications/Xcode.app/Contents/Developer/usr/bin/ibtool";
-				IBToolStartInfo.Arguments = CommandLine;
-				IBToolStartInfo.UseShellExecute = false;
-				Process CallIBTool = new Process();
-				CallIBTool.StartInfo = IBToolStartInfo;
-				CallIBTool.Start();
-				CallIBTool.WaitForExit();
-
-				PListAdditionalLines.Add("\t<key>UILaunchStoryboardName</key>");
-				PListAdditionalLines.Add("\t<string>LaunchScreen</string>");
-
-				bSkipDefaultPNGs = true;
-			}
-			else
-			{
-				// this is a temp way to inject the iphone 6 images without needing to upgrade everyone's plist
-				// eventually we want to generate this based on what the user has set in the project settings
-				string[] IPhoneConfigs =  
-				{ 
-					"Default-IPhone6", "Landscape", "{375, 667}", 
-					"Default-IPhone6", "Portrait", "{375, 667}", 
-					"Default-IPhone6Plus-Landscape", "Landscape", "{414, 736}", 
-					"Default-IPhone6Plus-Portrait", "Portrait", "{414, 736}", 
-					"Default", "Landscape", "{320, 480}",
-					"Default", "Portrait", "{320, 480}",
-					"Default-568h", "Landscape", "{320, 568}",
-					"Default-568h", "Portrait", "{320, 568}",
-				};
-
-				StringBuilder NewLaunchImagesString = new StringBuilder("<key>UILaunchImages~iphone</key>\n\t\t<array>\n");
-				for (int ConfigIndex = 0; ConfigIndex < IPhoneConfigs.Length; ConfigIndex += 3)
-				{
-					NewLaunchImagesString.Append("\t\t\t<dict>\n");
-					NewLaunchImagesString.Append("\t\t\t\t<key>UILaunchImageMinimumOSVersion</key>\n");
-					NewLaunchImagesString.Append("\t\t\t\t<string>8.0</string>\n");
-					NewLaunchImagesString.Append("\t\t\t\t<key>UILaunchImageName</key>\n");
-					NewLaunchImagesString.AppendFormat("\t\t\t\t<string>{0}</string>\n", IPhoneConfigs[ConfigIndex + 0]);
-					NewLaunchImagesString.Append("\t\t\t\t<key>UILaunchImageOrientation</key>\n");
-					NewLaunchImagesString.AppendFormat("\t\t\t\t<string>{0}</string>\n", IPhoneConfigs[ConfigIndex + 1]);
-					NewLaunchImagesString.Append("\t\t\t\t<key>UILaunchImageSize</key>\n");
-					NewLaunchImagesString.AppendFormat("\t\t\t\t<string>{0}</string>\n", IPhoneConfigs[ConfigIndex + 2]);
-					NewLaunchImagesString.Append("\t\t\t</dict>\n");
-				}
-
-				// close it out
-				NewLaunchImagesString.Append("\t\t\t</array>\n\t\t<key>UILaunchImages~ipad</key>");
-				Replacements.Add("<key>UILaunchImages~ipad</key>", NewLaunchImagesString.ToString());
-			}
-
-			// copy plist file
-			string PListFile = InEngineDir + "/Build/IOS/UE4Game-Info.plist";
-			if (File.Exists(BuildDirectory + "/Info.plist"))
-			{
-				PListFile = BuildDirectory + "/Info.plist";
-			}
-			else if (File.Exists(BuildDirectory + "/" + InProjectName + "-Info.plist"))
-			{
-				PListFile = BuildDirectory + "/" + InProjectName + "-Info.plist";
-			}
-
-			// plist replacements
-			Replacements.Add("${EXECUTABLE_NAME}", GameName);
-			Replacements.Add("${BUNDLE_IDENTIFIER}", InProjectName.Replace("_", ""));
-			CopyFileWithReplacements(PListFile, AppDirectory + "/Info.plist", Replacements, PListAdditionalLines);
-			CopyFileWithReplacements(PListFile, IntermediateDirectory + "/" + GameName + "-Info.plist", Replacements, PListAdditionalLines);
+			bool bSkipDefaultPNGs = GeneratePList(InProjectDirectory, bIsUE4Game, GameName, InProjectName, InEngineDir, AppDirectory);
 
 			// ensure the destination is writable
 			if (File.Exists(AppDirectory + "/" + GameName))
@@ -385,19 +657,20 @@ namespace UnrealBuildTool.IOS
 			string GameName = InTarget.AppName;
 			string BuildPath = (GameName == "UE4Game" ? "../../Engine" : InTarget.ProjectDirectory) + "/Binaries/IOS";
 			string ProjectDirectory = InTarget.ProjectDirectory;
+			bool bIsUE4Game = GameName.Contains("UE4Game");
 
-			if (BuildHostPlatform.Current.Platform == UnrealTargetPlatform.Mac && Environment.GetEnvironmentVariable("UBT_NO_POST_DEPLOY") != "true")
+			string DecoratedGameName;
+			if (InTarget.Configuration == UnrealTargetConfiguration.Development)
 			{
-				string DecoratedGameName;
-				if (InTarget.Configuration == UnrealTargetConfiguration.Development)
-				{
-					DecoratedGameName = GameName;
-				}
-				else
-				{
-					DecoratedGameName = String.Format("{0}-{1}-{2}", GameName, InTarget.Platform.ToString(), InTarget.Configuration.ToString());
-				}
+				DecoratedGameName = GameName;
+			}
+			else
+			{
+				DecoratedGameName = String.Format("{0}-{1}-{2}", GameName, InTarget.Platform.ToString(), InTarget.Configuration.ToString());
+			}
 
+			if (BuildHostPlatform.Current.Platform == UnrealTargetPlatform.Mac && Environment.GetEnvironmentVariable("UBT_NO_POST_DEPLOY") != "true")
+			{
 				return PrepForUATPackageOrDeploy(GameName, ProjectDirectory, BuildPath + "/" + DecoratedGameName, "../../Engine", false, "");
 			}
 			else
@@ -439,6 +712,8 @@ namespace UnrealBuildTool.IOS
 						Log.TraceInformation("Copying binaries back to this device failed.");
 					}
 				}
+
+				GeneratePList(ProjectDirectory, bIsUE4Game, GameName, Path.GetFileNameWithoutExtension(UnrealBuildTool.GetUProjectFile()), "../../Engine", "");
 			}
 			return true;
 		}
diff --git a/Engine/Source/Programs/UnrealBuildTool/System/XcodeProjectFileGenerator.cs b/Engine/Source/Programs/UnrealBuildTool/System/XcodeProjectFileGenerator.cs
--- a/Engine/Source/Programs/UnrealBuildTool/System/XcodeProjectFileGenerator.cs
+++ b/Engine/Source/Programs/UnrealBuildTool/System/XcodeProjectFileGenerator.cs
@@ -608,16 +608,12 @@ namespace UnrealBuildTool
 			// This is needed for the target to pass the settings validation before code signing. UBT will overwrite this plist file later, with proper contents.
 			if (BuildHostPlatform.Current.Platform == UnrealTargetPlatform.Mac)
 			{
-				if (File.Exists(InfoPlistPath))
+				if (!File.Exists(InfoPlistPath))
 				{
 					File.Delete(InfoPlistPath);
-				}
-				else
-				{
 					Directory.CreateDirectory(Path.GetDirectoryName(InfoPlistPath));
+					IOS.UEDeployIOS.GeneratePList(bIsUE4Game ? EngineRelative + "Engine" : GamePath, bIsUE4Game, bIsUE4Game ? "UE4Game" : TargetName, TargetName, EngineRelative + "Engine", bIsUE4Game ? EngineRelative + "Engine/Binaries/IOS/Payload" : GamePath + "/Binaries/IOS/Payload");
 				}
-				string InfoPlistContents = File.ReadAllText(EngineRelative + "Engine/Build/IOS/UE4Game-Info.plist").Replace("${BUNDLE_IDENTIFIER}", TargetName).Replace("${EXECUTABLE_NAME}", TargetName);
-				File.WriteAllText(InfoPlistPath, InfoPlistContents);
 			}
 		}
 
diff --git a/Engine/Source/Runtime/IOS/IOSRuntimeSettings/Classes/IOSRuntimeSettings.h b/Engine/Source/Runtime/IOS/IOSRuntimeSettings/Classes/IOSRuntimeSettings.h
--- a/Engine/Source/Runtime/IOS/IOSRuntimeSettings/Classes/IOSRuntimeSettings.h
+++ b/Engine/Source/Runtime/IOS/IOSRuntimeSettings/Classes/IOSRuntimeSettings.h
@@ -200,6 +200,10 @@ public:
 	UPROPERTY(GlobalConfig, EditAnywhere, Category = DeviceUsage)
 	uint32 bSupportsIPhone : 1;
 
+	// extra data to be stored in the plist
+	UPROPERTY(GlobalConfig, EditAnywhere, Category = ExtraData)
+	FString AdditionalPlistData;
+
 #if WITH_EDITOR
 	// UObject interface
 	virtual void PostEditChangeProperty(struct FPropertyChangedEvent& PropertyChangedEvent) override;
diff --git a/Engine/Source/Runtime/IOS/IOSRuntimeSettings/Private/IOSRuntimeSettings.cpp b/Engine/Source/Runtime/IOS/IOSRuntimeSettings/Private/IOSRuntimeSettings.cpp
--- a/Engine/Source/Runtime/IOS/IOSRuntimeSettings/Private/IOSRuntimeSettings.cpp
+++ b/Engine/Source/Runtime/IOS/IOSRuntimeSettings/Private/IOSRuntimeSettings.cpp
@@ -24,6 +24,7 @@ UIOSRuntimeSettings::UIOSRuntimeSettings(const FObjectInitializer& ObjectInitial
 	bShipForArm64 = true;
 	bShipForArmV7S = false;
 	bUseRSync = false;
+	AdditionalPlistData = TEXT("");
 }
 
 #if WITH_EDITOR
