fix for CIS build failures for IOS on mac
#ios

diff --git a/Engine/Source/Programs/UnrealBuildTool/System/XcodeProjectFileGenerator.cs b/Engine/Source/Programs/UnrealBuildTool/System/XcodeProjectFileGenerator.cs
--- a/Engine/Source/Programs/UnrealBuildTool/System/XcodeProjectFileGenerator.cs
+++ b/Engine/Source/Programs/UnrealBuildTool/System/XcodeProjectFileGenerator.cs
@@ -610,9 +610,19 @@ namespace UnrealBuildTool
 			{
 				if (!File.Exists(InfoPlistPath))
 				{
-					File.Delete(InfoPlistPath);
+					string ProjectPath = GamePath;
+					string GameName = TargetName;
+					if (string.IsNullOrEmpty (ProjectPath))
+					{
+						ProjectPath = EngineRelative + "Engine";
+					}
+					if (bIsUE4Game)
+					{
+						ProjectPath = EngineRelative + "Engine";
+						GameName = "UE4Game";
+					}
 					Directory.CreateDirectory(Path.GetDirectoryName(InfoPlistPath));
-					IOS.UEDeployIOS.GeneratePList(bIsUE4Game ? EngineRelative + "Engine" : GamePath, bIsUE4Game, bIsUE4Game ? "UE4Game" : TargetName, TargetName, EngineRelative + "Engine", bIsUE4Game ? EngineRelative + "Engine/Binaries/IOS/Payload" : GamePath + "/Binaries/IOS/Payload");
+					IOS.UEDeployIOS.GeneratePList(ProjectPath, bIsUE4Game, GameName, TargetName, EngineRelative + "Engine", ProjectPath + "/Binaries/IOS/Payload");
 				}
 			}
 		}
