fix for editor not liking {} for display in project settings.
#ios

diff --git a/Engine/Config/BaseEngine.ini b/Engine/Config/BaseEngine.ini
--- a/Engine/Config/BaseEngine.ini
+++ b/Engine/Config/BaseEngine.ini
@@ -2119,9 +2119,9 @@ bShipForArmV7=True
 bShipForArm64=True
 bShipForArmV7S=False
 bUseRSync=False
-BundleDisplayName={PROJECT_NAME}
-BundleName={PROJECT_NAME}
-BundleIdentifier=com.YourCompany.{PROJECT_NAME}
+BundleDisplayName=[PROJECT_NAME]
+BundleName=[PROJECT_NAME]
+BundleIdentifier=com.YourCompany.[PROJECT_NAME]
 VersionInfo=1.0
 FrameRateLock=PUFRL_30
 MinimumiOSVersion=IOS_61
diff --git a/Engine/Source/Programs/UnrealBuildTool/IOS/UEDeployIOS.cs b/Engine/Source/Programs/UnrealBuildTool/IOS/UEDeployIOS.cs
--- a/Engine/Source/Programs/UnrealBuildTool/IOS/UEDeployIOS.cs
+++ b/Engine/Source/Programs/UnrealBuildTool/IOS/UEDeployIOS.cs
@@ -267,15 +267,15 @@ namespace UnrealBuildTool.IOS
 			Text.AppendLine("\t<key>CFBundleDevelopmentRegion</key>");
 			Text.AppendLine("\t<string>English</string>");
 			Text.AppendLine("\t<key>CFBundleDisplayName</key>");
-			Text.AppendLine(string.Format("\t<string>{0}</string>", BundleDisplayName.Replace("{PROJECT_NAME}", ProjectName).Replace("_","")));
+			Text.AppendLine(string.Format("\t<string>{0}</string>", BundleDisplayName.Replace("[PROJECT_NAME]", ProjectName).Replace("_","")));
 			Text.AppendLine("\t<key>CFBundleExecutable</key>");
 			Text.AppendLine(string.Format("\t<string>{0}</string>", bIsUE4Game ? "UE4Game" : GameName.Replace("_", "")));
 			Text.AppendLine("\t<key>CFBundleIdentifier</key>");
-			Text.AppendLine(string.Format("\t<string>{0}</string>", BundleIdentifier.Replace("{PROJECT_NAME}", ProjectName).Replace("_","")));
+			Text.AppendLine(string.Format("\t<string>{0}</string>", BundleIdentifier.Replace("[PROJECT_NAME]", ProjectName).Replace("_","")));
 			Text.AppendLine("\t<key>CFBundleInfoDictionaryVersion</key>");
 			Text.AppendLine("\t<string>6.0</string>");
 			Text.AppendLine("\t<key>CFBundleName</key>");
-			Text.AppendLine(string.Format("\t<string>{0}</string>", BundleName.Replace("{PROJECT_NAME}", ProjectName).Replace("_","")));
+			Text.AppendLine(string.Format("\t<string>{0}</string>", BundleName.Replace("[PROJECT_NAME]", ProjectName).Replace("_","")));
 			Text.AppendLine("\t<key>CFBundlePackageType</key>");
 			Text.AppendLine("\t<string>APPL</string>");
 			Text.AppendLine("\t<key>CFBundleSignature</key>");
