fix for potential crash when trying to copy the Info.plist to the payload directory
#ios

diff --git a/Engine/Source/Programs/UnrealBuildTool/IOS/UEDeployIOS.cs b/Engine/Source/Programs/UnrealBuildTool/IOS/UEDeployIOS.cs
--- a/Engine/Source/Programs/UnrealBuildTool/IOS/UEDeployIOS.cs
+++ b/Engine/Source/Programs/UnrealBuildTool/IOS/UEDeployIOS.cs
@@ -440,6 +440,10 @@ namespace UnrealBuildTool.IOS
 			File.WriteAllText(PListFile, Text.ToString());
 			if (BuildHostPlatform.Current.Platform == UnrealTargetPlatform.Mac)
 			{
+				if (!Directory.Exists(AppDirectory))
+				{
+					Directory.CreateDirectory(AppDirectory);
+				}
 				File.WriteAllText(AppDirectory + "/Info.plist", Text.ToString());
 			}
 
