Always build Android target platform and texture format modules on builders, even if the SDK is not installed

diff --git a/Engine/Source/Programs/UnrealBuildTool/Android/UEBuildAndroid.cs b/Engine/Source/Programs/UnrealBuildTool/Android/UEBuildAndroid.cs
--- a/Engine/Source/Programs/UnrealBuildTool/Android/UEBuildAndroid.cs
+++ b/Engine/Source/Programs/UnrealBuildTool/Android/UEBuildAndroid.cs
@@ -162,7 +162,7 @@ namespace UnrealBuildTool
 
 		protected override void RegisterBuildPlatformInternal()
 		{
-			if ((ProjectFileGenerator.bGenerateProjectFiles == true) || (HasRequiredSDKsInstalled() == SDKStatus.Valid))
+			if ((ProjectFileGenerator.bGenerateProjectFiles == true) || (HasRequiredSDKsInstalled() == SDKStatus.Valid) || Environment.GetEnvironmentVariable("IsBuildMachine") == "1")
 			{
 				bool bRegisterBuildPlatform = true;
 
