Fix remote compiling using settings from BuildConfiguration.xml. Empty settings in BaseEngine.ini aren't ignored by UBT, and just overwrite valid settings in XML file.

diff --git a/Engine/Source/Programs/UnrealBuildTool/ToolChain/RemoteToolChain.cs b/Engine/Source/Programs/UnrealBuildTool/ToolChain/RemoteToolChain.cs
--- a/Engine/Source/Programs/UnrealBuildTool/ToolChain/RemoteToolChain.cs
+++ b/Engine/Source/Programs/UnrealBuildTool/ToolChain/RemoteToolChain.cs
@@ -193,7 +193,7 @@ namespace UnrealBuildTool
 
 			ConfigCacheIni Ini = new ConfigCacheIni(UnrealTargetPlatform.IOS, "Engine", UnrealBuildTool.GetUProjectPath());
 			string ServerName = RemoteServerName;
-			if (Ini.GetString("/Script/IOSRuntimeSettings.IOSRuntimeSettings", "RemoteServerName", out ServerName))
+			if (Ini.GetString("/Script/IOSRuntimeSettings.IOSRuntimeSettings", "RemoteServerName", out ServerName) && !String.IsNullOrEmpty(ServerName))
 			{
 				RemoteServerName = ServerName;
 			}
