Remove ../../Rocket from automation search paths. Nothing's in it any more.

diff --git a/Engine/Source/Programs/UnrealBuildTool/Configuration/UEBuildTarget.cs b/Engine/Source/Programs/UnrealBuildTool/Configuration/UEBuildTarget.cs
--- a/Engine/Source/Programs/UnrealBuildTool/Configuration/UEBuildTarget.cs
+++ b/Engine/Source/Programs/UnrealBuildTool/Configuration/UEBuildTarget.cs
@@ -605,16 +605,6 @@ namespace UnrealBuildTool
 				AllGameFolders.Add(GameProject.Folder);
 			}
 
-			// @todo: Perversely, if we're not running rocket we need to add the Rocket root folder so that the Rocket build automation script can be found...
-			if (!UnrealBuildTool.RunningRocket())
-			{
-				string RocketFolder = "../../Rocket";
-				if (Directory.Exists(RocketFolder))
-				{
-					AllGameFolders.Add(RocketFolder);
-				}
-			}
-
 			return AllGameFolders;
 		}
 
