Fix UBT makefile regeneration.

diff --git a/Engine/Source/Programs/UnrealBuildTool/System/ProjectFileGenerator.cs b/Engine/Source/Programs/UnrealBuildTool/System/ProjectFileGenerator.cs
--- a/Engine/Source/Programs/UnrealBuildTool/System/ProjectFileGenerator.cs
+++ b/Engine/Source/Programs/UnrealBuildTool/System/ProjectFileGenerator.cs
@@ -182,6 +182,9 @@ namespace UnrealBuildTool
 		// @todo projectfiles: Ideally, projects for game modules/targets would be created in the game's Intermediate folder!
 		public static string IntermediateProjectFilesPath = Path.Combine( EngineRelativePath, "Intermediate", "ProjectFiles" );
 
+		/// Path to timestamp file, recording when was the last time projects were created.
+		public static string ProjectTimestampFile = Path.Combine(IntermediateProjectFilesPath, "Timestamp");
+
 		/// Global static new line string used by ProjectFileGenerator to generate project files.
 		public static readonly string NewLine = Environment.NewLine;
 
diff --git a/Engine/Source/Programs/UnrealBuildTool/System/UnrealBuildTool.cs b/Engine/Source/Programs/UnrealBuildTool/System/UnrealBuildTool.cs
--- a/Engine/Source/Programs/UnrealBuildTool/System/UnrealBuildTool.cs
+++ b/Engine/Source/Programs/UnrealBuildTool/System/UnrealBuildTool.cs
@@ -1426,6 +1426,12 @@ namespace UnrealBuildTool
 				bool bIsHotReload = UEBuildConfiguration.bHotReloadFromIDE || ( TargetDescs.Count == 1 && TargetDescs[0].OnlyModules.Count > 0 );
 				TargetDescriptor HotReloadTargetDesc = bIsHotReload ? TargetDescs[0] : null;
 
+				if (ProjectFileGenerator.bGenerateProjectFiles)
+				{
+					// Create empty timestamp file to record when was the last time we regenerated projects.
+					File.Create(ProjectFileGenerator.ProjectTimestampFile).Dispose();
+				}
+
                 if( !ProjectFileGenerator.bGenerateProjectFiles )
                 {
 					if( BuildConfiguration.bUseUBTMakefiles )
@@ -2209,7 +2215,7 @@ namespace UnrealBuildTool
                 {
                     if( Directory.Exists( ProjectFileGenerator.IntermediateProjectFilesPath ) )
                     {
-                        var EngineProjectFilesLastUpdateTime = new DirectoryInfo( ProjectFileGenerator.IntermediateProjectFilesPath ).LastWriteTime;
+                        var EngineProjectFilesLastUpdateTime = new FileInfo(ProjectFileGenerator.ProjectTimestampFile).LastWriteTime;
 						if( UBTMakefileItem.LastWriteTime.CompareTo( EngineProjectFilesLastUpdateTime ) < 0 )
 						{
 							// Engine project files are newer than UBTMakefile
