Fix UBT timestamp file creation

diff --git a/Engine/Source/Programs/UnrealBuildTool/System/UnrealBuildTool.cs b/Engine/Source/Programs/UnrealBuildTool/System/UnrealBuildTool.cs
--- a/Engine/Source/Programs/UnrealBuildTool/System/UnrealBuildTool.cs
+++ b/Engine/Source/Programs/UnrealBuildTool/System/UnrealBuildTool.cs
@@ -1429,6 +1429,7 @@ namespace UnrealBuildTool
 				if (ProjectFileGenerator.bGenerateProjectFiles)
 				{
 					// Create empty timestamp file to record when was the last time we regenerated projects.
+					Directory.CreateDirectory( Path.GetDirectoryName( ProjectFileGenerator.ProjectTimestampFile ) );
 					File.Create(ProjectFileGenerator.ProjectTimestampFile).Dispose();
 				}
 
