HTML5 SDK settings are now set via the Editor.
Android SDK path backslash duplication fix.
Fix for User inis not saving array properties correctly.

diff --git a/Engine/Config/BaseEngine.ini b/Engine/Config/BaseEngine.ini
--- a/Engine/Config/BaseEngine.ini
+++ b/Engine/Config/BaseEngine.ini
@@ -2207,4 +2207,7 @@ NameForGameView3="GameView3"
 NameForGameView4="GameView4"
 NameForGameView5="GameView5"
 
-
+; These are the defaults for HTML5 settings, and they need to be in the .ini since UBT reads the .ini settings, without instantiating the class
+[/Script/HTML5PlatformEditor.HTML5TargetSettings]
+HeapSizeDevelopment=1024
+HeapSizeShipping=1024
\ No newline at end of file
diff --git a/Engine/Config/HTML5/HTML5Engine.ini b/Engine/Config/HTML5/HTML5Engine.ini
--- a/Engine/Config/HTML5/HTML5Engine.ini
+++ b/Engine/Config/HTML5/HTML5Engine.ini
@@ -7,24 +7,26 @@ benabled=false
 [SystemSettings]
 r.setres=800x600
 
-[BuildSettings]
-HeapSize=1024
+; These following settings have moved to the User inis and can be set from the editor
+; These old value are left as reference until all old code paths are removed
+;[BuildSettings]
+;HeapSize=1024
 
-[HTML5DevicesWindows]
-Firefox Nightly (64 bit)="C:/Program Files/Nightly/firefox.exe"
-Firefox Nightly (32 bit)="C:/Program Files (x86)/Nightly/firefox.exe"
-Firefox (32 bit)="C:/Program Files (x86)/Mozilla Firefox/firefox.exe"
-Firefox Aurora  (32 bit)="C:/Program Files (x86)/Aurora/firefox.exe"
-Chrome="C:/Program Files (x86)/Google/Chrome/Application/chrome.exe"
+;[HTML5DevicesWindows]
+;Firefox Nightly (64 bit)="C:/Program Files/Nightly/firefox.exe"
+;Firefox Nightly (32 bit)="C:/Program Files (x86)/Nightly/firefox.exe"
+;Firefox (32 bit)="C:/Program Files (x86)/Mozilla Firefox/firefox.exe"
+;Firefox Aurora  (32 bit)="C:/Program Files (x86)/Aurora/firefox.exe"
+;Chrome="C:/Program Files (x86)/Google/Chrome/Application/chrome.exe"
 
 ;all browsers on mac are 64 bit, 64 bit browsers can allocate more than 512MB of memory to around 1.5 GB. 
-[HTML5DevicesMac]
-Safari="/Applications/Safari.app/Contents/MacOS/Safari"
-Firefox="/Applications/Firefox.app/Contents/MacOS/firefox-bin"
-Firefox (Aurora)="/Applications/FirefoxAurora.app/Contents/MacOS/firefox-bin"
+;[HTML5DevicesMac]
+;Safari="/Applications/Safari.app/Contents/MacOS/Safari"
+;Firefox="/Applications/Firefox.app/Contents/MacOS/firefox-bin"
+;Firefox (Aurora)="/Applications/FirefoxAurora.app/Contents/MacOS/firefox-bin"
 
-[HTML5DevicesLinux]
-Firefox="/usr/bin/firefox"
+;[HTML5DevicesLinux]
+;Firefox="/usr/bin/firefox"
 
-[HTML5SDKPaths]
-Mac="/path/to/emscripten/sdk"
+;[HTML5SDKPaths]
+;Mac="/path/to/emscripten/sdk"
diff --git a/Engine/Source/Developer/Android/AndroidPlatformEditor/Private/AndroidSDKSettings.cpp b/Engine/Source/Developer/Android/AndroidPlatformEditor/Private/AndroidSDKSettings.cpp
--- a/Engine/Source/Developer/Android/AndroidPlatformEditor/Private/AndroidSDKSettings.cpp
+++ b/Engine/Source/Developer/Android/AndroidPlatformEditor/Private/AndroidSDKSettings.cpp
@@ -125,24 +125,28 @@ void UAndroidSDKSettings::UpdateTargetModulePaths(bool bForceUpdate)
 	
 	if (bForceUpdate || !SDKPath.Path.IsEmpty())
 	{
+		FPaths::NormalizeFilename(SDKPath.Path);
 		Keys.Add(TEXT("ANDROID_HOME"));
 		Values.Add(SDKPath.Path);
 	}
 	
 	if (bForceUpdate || !NDKPath.Path.IsEmpty())
 	{
+		FPaths::NormalizeFilename(NDKPath.Path);
 		Keys.Add(TEXT("NDKROOT"));
 		Values.Add(NDKPath.Path);
 	}
 	
 	if (bForceUpdate || !ANTPath.Path.IsEmpty())
 	{
+		FPaths::NormalizeFilename(ANTPath.Path);
 		Keys.Add(TEXT("ANT_HOME"));
 		Values.Add(ANTPath.Path);
 	}
 
 	if (bForceUpdate || !JavaPath.Path.IsEmpty())
 	{
+		FPaths::NormalizeFilename(JavaPath.Path);
 		Keys.Add(TEXT("JAVA_HOME"));
 		Values.Add(JavaPath.Path);
 	}
diff --git a/Engine/Source/Developer/HTML5/HTML5PlatformEditor/Classes/HTML5TargetSettings.h b/Engine/Source/Developer/HTML5/HTML5PlatformEditor/Classes/HTML5TargetSettings.h
new file mode 100644
--- /dev/null
+++ b/Engine/Source/Developer/HTML5/HTML5PlatformEditor/Classes/HTML5TargetSettings.h
@@ -0,0 +1,34 @@
+// Copyright 1998-2015 Epic Games, Inc. All Rights Reserved.
+
+/*=============================================================================
+	HTML5TargetSettings.h: Declares the UHTML5TargetSettings class.
+=============================================================================*/
+
+#pragma once
+
+#include "HTML5TargetSettings.generated.h"
+
+
+/**
+ * Implements the settings for the HTML5 target platform.
+ */
+UCLASS(config=Engine, defaultconfig)
+class HTML5PLATFORMEDITOR_API UHTML5TargetSettings
+	: public UObject
+{
+public:
+
+	GENERATED_UCLASS_BODY()
+
+ 	/**
+ 	 * Setting to control HTML5 Heap size (in Development)
+ 	 */
+ 	UPROPERTY(GlobalConfig, EditAnywhere, Category=Memory, Meta = (DisplayName = "Development Heap Size (in MB)", ClampMin="1", ClampMax="4096"))
+ 	int32 HeapSizeDevelopment;
+
+	/**
+ 	 * Setting to control HTML5 Heap size
+ 	 */
+ 	UPROPERTY(GlobalConfig, EditAnywhere, Category=Memory, Meta = (DisplayName = "Heap Size (in MB)", ClampMin="1", ClampMax="4096"))
+ 	int32 HeapSizeShipping;
+};
diff --git a/Engine/Source/Developer/HTML5/HTML5PlatformEditor/HTML5PlatformEditor.Build.cs b/Engine/Source/Developer/HTML5/HTML5PlatformEditor/HTML5PlatformEditor.Build.cs
new file mode 100644
--- /dev/null
+++ b/Engine/Source/Developer/HTML5/HTML5PlatformEditor/HTML5PlatformEditor.Build.cs
@@ -0,0 +1,36 @@
+// Copyright 1998-2015 Epic Games, Inc. All Rights Reserved.
+
+using UnrealBuildTool;
+
+public class HTML5PlatformEditor : ModuleRules
+{
+	public HTML5PlatformEditor(TargetInfo Target)
+	{
+		BinariesSubFolder = "HTML5";
+
+		PrivateDependencyModuleNames.AddRange(
+			new string[] {
+				"Core",
+				"CoreUObject",
+				"InputCore",
+				"Engine",
+				"Slate",
+				"SlateCore",
+				"EditorStyle",
+                "EditorWidgets",
+                "DesktopWidgets",
+				"PropertyEditor",
+				"SharedSettingsWidgets",
+				"SourceControl",
+                "TargetPlatform",
+				"HTML5TargetPlatform",
+			}
+		);
+
+		PrivateIncludePathModuleNames.AddRange(
+			new string[] {
+				"Settings",
+			}
+		);
+	}
+}
diff --git a/Engine/Source/Developer/HTML5/HTML5PlatformEditor/Private/HTML5PlatformEditorClasses.cpp b/Engine/Source/Developer/HTML5/HTML5PlatformEditor/Private/HTML5PlatformEditorClasses.cpp
new file mode 100644
--- /dev/null
+++ b/Engine/Source/Developer/HTML5/HTML5PlatformEditor/Private/HTML5PlatformEditorClasses.cpp
@@ -0,0 +1,16 @@
+// Copyright 1998-2015 Epic Games, Inc. All Rights Reserved.
+
+/*=============================================================================
+	HTML5TargetPlatformClasses.cpp: Implements the module's UClasses.
+=============================================================================*/
+
+#include "HTML5PlatformEditorPrivatePCH.h"
+
+
+
+/* UHTML5TargetSettings structors
+ *****************************************************************************/
+
+UHTML5TargetSettings::UHTML5TargetSettings( const FObjectInitializer& ObjectInitializer )
+	: Super(ObjectInitializer)
+{ }
diff --git a/Engine/Source/Developer/HTML5/HTML5PlatformEditor/Private/HTML5PlatformEditorModule.cpp b/Engine/Source/Developer/HTML5/HTML5PlatformEditor/Private/HTML5PlatformEditorModule.cpp
new file mode 100644
--- /dev/null
+++ b/Engine/Source/Developer/HTML5/HTML5PlatformEditor/Private/HTML5PlatformEditorModule.cpp
@@ -0,0 +1,65 @@
+// Copyright 1998-2015 Epic Games, Inc. All Rights Reserved.
+
+#include "HTML5PlatformEditorPrivatePCH.h"
+#include "ModuleInterface.h"
+#include "ISettingsModule.h"
+#include "ModuleManager.h"
+#include "HTML5SDKSettings.h"
+
+
+#define LOCTEXT_NAMESPACE "FHTML5PlatformEditorModule"
+
+
+/**
+ * Module for Android platform editor utilities
+ */
+class FHTML5PlatformEditorModule
+	: public IModuleInterface
+{
+	// IModuleInterface interface
+
+	virtual void StartupModule() override
+	{
+		// register settings
+		ISettingsModule* SettingsModule = FModuleManager::GetModulePtr<ISettingsModule>("Settings");
+
+		if (SettingsModule != nullptr)
+		{
+			SettingsModule->RegisterSettings("Project", "Platforms", "HTML5",
+				LOCTEXT("TargetSettingsName", "HTML5"),
+				LOCTEXT("SDKSettingsDescription", "Settings for HTML5"),
+				GetMutableDefault<UHTML5TargetSettings>()
+			);
+
+ 			SettingsModule->RegisterSettings("Project", "Platforms", "HTML5SDK",
+ 				LOCTEXT("SDKSettingsName", "HTML5 SDK"),
+ 				LOCTEXT("SDKSettingsDescription", "Settings for HTML5 SDK (system wide)"),
+ 				GetMutableDefault<UHTML5SDKSettings>()
+			);
+		}
+
+		auto* Settings = GetMutableDefault<UHTML5SDKSettings>();
+		if (Settings->DeviceMap.Num() == 0)
+		{
+			Settings->QueryKnownBrowserLocations();
+		}
+	}
+
+	virtual void ShutdownModule() override
+	{
+		auto* Settings = GetMutableDefault<UHTML5SDKSettings>();
+
+		ISettingsModule* SettingsModule = FModuleManager::GetModulePtr<ISettingsModule>("Settings");
+
+		if (SettingsModule != nullptr)
+		{
+			SettingsModule->UnregisterSettings("Project", "Platforms", "HTML5");
+			SettingsModule->UnregisterSettings("Project", "Platforms", "HTML5SDK");
+		}
+	}
+};
+
+
+IMPLEMENT_MODULE(FHTML5PlatformEditorModule, HTML5PlatformEditor);
+
+#undef LOCTEXT_NAMESPACE
diff --git a/Engine/Source/Developer/HTML5/HTML5PlatformEditor/Private/HTML5PlatformEditorPrivatePCH.h b/Engine/Source/Developer/HTML5/HTML5PlatformEditor/Private/HTML5PlatformEditorPrivatePCH.h
new file mode 100644
--- /dev/null
+++ b/Engine/Source/Developer/HTML5/HTML5PlatformEditor/Private/HTML5PlatformEditorPrivatePCH.h
@@ -0,0 +1,10 @@
+// Copyright 1998-2015 Epic Games, Inc. All Rights Reserved.
+
+#pragma once
+
+#include "Core.h"
+#include "ModuleInterface.h"
+#include "ModuleManager.h"
+#include "Engine.h"
+
+#include "HTML5TargetSettings.h"
\ No newline at end of file
diff --git a/Engine/Source/Developer/HTML5/HTML5PlatformEditor/Private/HTML5SDKSettings.cpp b/Engine/Source/Developer/HTML5/HTML5PlatformEditor/Private/HTML5SDKSettings.cpp
new file mode 100644
--- /dev/null
+++ b/Engine/Source/Developer/HTML5/HTML5PlatformEditor/Private/HTML5SDKSettings.cpp
@@ -0,0 +1,73 @@
+// Copyright 1998-2014 Epic Games, Inc. All Rights Reserved.
+
+#include "HTML5PlatformEditorPrivatePCH.h"
+
+//#include "EngineTypes.h"
+#include "HTML5SDKSettings.h"
+#include "IHTML5TargetPlatformModule.h"
+
+DEFINE_LOG_CATEGORY_STATIC(HTML5SDKSettings, Log, All);
+
+UHTML5SDKSettings::UHTML5SDKSettings(const FObjectInitializer& ObjectInitializer)
+	: Super(ObjectInitializer)
+	, TargetPlatformModule(nullptr)
+{
+}
+
+#if WITH_EDITOR
+void UHTML5SDKSettings::PostEditChangeProperty(struct FPropertyChangedEvent& PropertyChangedEvent)
+{
+	Super::PostEditChangeProperty(PropertyChangedEvent);
+	if (!TargetPlatformModule)
+	{
+		TargetPlatformModule = FModuleManager::LoadModulePtr<IHTML5TargetPlatformModule>("HTML5TargetPlatform");
+	}
+	
+	if (TargetPlatformModule) 
+	{
+		SaveConfig();
+		TargetPlatformModule->RefreshAvailableDevices();
+	}
+}
+
+void UHTML5SDKSettings::QueryKnownBrowserLocations()
+{
+	if (DeviceMap.Num() != 0) 
+	{
+		return;
+	}
+
+	struct FBrowserLocation 
+	{
+		FString Name;
+		FString Path;
+	} PossibleLocations[] = 
+	{
+#if PLATFORM_WINDOWS
+		{TEXT("Nightly"), TEXT("C:/Program Files/Nightly/firefox.exe")},
+		{TEXT("Firefox"), TEXT("C:/Program Files (x86)/Mozilla Firefox/firefox.exe")},
+		{TEXT("Chrome"), TEXT("C:/Program Files (x86)/Google/Chrome/Application/chrome.exe")},
+#elif PLATFORM_MAC
+		{TEXT("Safari"), TEXT("/Applications/Safari.app/Contents/MacOS/Safari")},
+		{TEXT("Firefox"),TEXT("/Applications/Firefox.app/Contents/MacOS/firefox-bin")},
+		{TEXT("Chrome"),TEXT("/Applications/Google Chrome.app/Contents/MacOS/Google Chrome")},
+#elif PLATFORM_LINUX
+		{TEXT("Firefox"), TEXT("/usr/bin/firefox")},
+#else
+		{TEXT("Firefox"), TEXT("")},
+#endif
+	};
+
+	for (const auto& Loc : PossibleLocations)
+	{
+		if (FPlatformFileManager::Get().GetPlatformFile().FileExists(*Loc.Name))
+		{
+			FHTML5DeviceMapping NewDevice;
+			NewDevice.DeviceName = Loc.Name;
+			NewDevice.DevicePath.FilePath = Loc.Path;
+			DeviceMap.Add(NewDevice);
+		}
+	}
+}
+
+#endif //WITH_EDITOR
\ No newline at end of file
diff --git a/Engine/Source/Developer/HTML5/HTML5PlatformEditor/Private/HTML5SDKSettings.h b/Engine/Source/Developer/HTML5/HTML5PlatformEditor/Private/HTML5SDKSettings.h
new file mode 100644
--- /dev/null
+++ b/Engine/Source/Developer/HTML5/HTML5PlatformEditor/Private/HTML5SDKSettings.h
@@ -0,0 +1,49 @@
+// Copyright 1998-2014 Epic Games, Inc. All Rights Reserved.
+
+#pragma once
+
+#include "ITargetPlatformManagerModule.h"
+#include "HTML5SDKSettings.generated.h"
+
+class IHTML5TargetPlatformModule;
+
+USTRUCT()
+struct FHTML5DeviceMapping
+{
+	GENERATED_USTRUCT_BODY()
+
+	UPROPERTY(EditAnywhere, Category = HTML5_Devices, Meta = (DisplayName = "Browser name"))
+	FString DeviceName;
+
+	UPROPERTY(EditAnywhere, Category = HTML5_Devices, Meta = (DisplayName = "Browser filepath"))
+	FFilePath DevicePath;
+};
+
+/**
+ * Implements the settings for the HTML5 SDK setup.
+ */
+UCLASS(config=Engine, globaluserconfig)
+class HTML5PLATFORMEDITOR_API UHTML5SDKSettings : public UObject
+{
+public:
+	GENERATED_UCLASS_BODY()
+
+	UPROPERTY(GlobalConfig, EditAnywhere, Category = HTML5_SDK_Paths, Meta = (DisplayName = "Location of Emscripten SDK (This directory contains 'emcc')"))
+	FDirectoryPath Emscripten;
+
+	UPROPERTY(GlobalConfig, EditAnywhere, Category = HTML5_SDK_Paths, Meta = (DisplayName = "Location of Python exe (This can be blank if python is on the PATH)"))
+	FFilePath Python;
+
+	// Available browsers that can be used when launching HTML5 builds.
+	UPROPERTY(GlobalConfig, EditAnywhere, Category = HTML5_Devices, Meta = (DisplayName = "Available browsers"))
+	TArray<FHTML5DeviceMapping> DeviceMap;
+
+#if WITH_EDITOR
+	// UObject interface
+	virtual void PostEditChangeProperty(struct FPropertyChangedEvent& PropertyChangedEvent) override;
+	// End of UObject interface
+	void QueryKnownBrowserLocations();
+#endif
+private:
+	IHTML5TargetPlatformModule* TargetPlatformModule;
+};
diff --git a/Engine/Source/Developer/HTML5/HTML5TargetPlatform/Classes/HTML5TargetSettings.h b/Engine/Source/Developer/HTML5/HTML5TargetPlatform/Classes/HTML5TargetSettings.h
deleted file mode 100644
diff --git a/Engine/Source/Developer/HTML5/HTML5TargetPlatform/Private/HTML5TargetPlatform.cpp b/Engine/Source/Developer/HTML5/HTML5TargetPlatform/Private/HTML5TargetPlatform.cpp
--- a/Engine/Source/Developer/HTML5/HTML5TargetPlatform/Private/HTML5TargetPlatform.cpp
+++ b/Engine/Source/Developer/HTML5/HTML5TargetPlatform/Private/HTML5TargetPlatform.cpp
@@ -18,24 +18,7 @@ FHTML5TargetPlatform::FHTML5TargetPlatform( )
 	// load the final HTML5 engine settings for this game
 	FConfigCacheIni::LoadLocalIniFile(HTML5EngineSettings, TEXT("Engine"), true, *PlatformName());
 
-	FString DeviceSectionName;
-
-#if PLATFORM_WINDOWS
-	DeviceSectionName = "HTML5DevicesWindows";
-#elif PLATFORM_MAC
-	DeviceSectionName = "HTML5DevicesMac";
-#else 
-	DeviceSectionName = "HTML5DevicesLinux";
-#endif 
-
-	FConfigSection AvaliableDevicesSection = HTML5EngineSettings[DeviceSectionName];
-	for (auto It : AvaliableDevicesSection)
-	{
-		const FString& BrowserName = It.Key.ToString();
-		const FString& BrowserPath = It.Value;
-		if(FPlatformFileManager::Get().GetPlatformFile().FileExists(*It.Value))
-			LocalDevice.Add(MakeShareable(new FHTML5TargetDevice(*this, FString::Printf(TEXT("%s on %s"), *It.Key.ToString(), FPlatformProcess::ComputerName()))));
-	}
+	RefreshAvailableDevices();
 
 #if WITH_ENGINE
 	// load up texture settings from the config file
@@ -82,8 +65,25 @@ ITargetDevicePtr FHTML5TargetPlatform::GetDevice( const FTargetDeviceId& DeviceI
 
 bool FHTML5TargetPlatform::IsSdkInstalled(bool bProjectHasCode, FString& OutDocumentationPath) const
 {
+	//New style detection of devices
+	FConfigSection SDKPaths = HTML5EngineSettings["/Script/HTML5PlatformEditor.HTML5SDKSettings"];
+	for (auto It : SDKPaths)
+	{
+		const FString& Platform = It.Key.ToString();
+		FString Path = It.Value;
+		{	
+			Path.RemoveFromStart(TEXT("(Path=\""));
+			Path.RemoveFromEnd(TEXT("\")"));
+			if (Platform == "Emscripten" && IFileManager::Get().DirectoryExists(*Path))
+			{
+				return true;
+			}
+		}
+	}
+
+	// Old fallbacks
 	FString SectionName = "HTML5SDKPaths";
-	FConfigSection SDKPaths = HTML5EngineSettings[SectionName];
+	SDKPaths = HTML5EngineSettings[SectionName];
 	for (auto It : SDKPaths )
 	{
 		const FString& Platform = It.Key.ToString();
@@ -266,4 +266,82 @@ FName FHTML5TargetPlatform::GetWaveFormat( const USoundWave* Wave ) const
 	return NAME_OGG;
 }
 
+void FHTML5TargetPlatform::RefreshAvailableDevices()
+{
+	//New style detection of devices
+	for (const auto& Device : LocalDevice)
+	{
+		DeviceLostEvent.Broadcast(Device.ToSharedRef());
+	}
+	LocalDevice.Reset();
+
+	auto* Config = GConfig->FindConfigFile(GEngineIni);
+	if (!Config)
+	{
+		Config = &HTML5EngineSettings;
+	}
+	TArray<FString> ValueArray;
+	FConfigSection AvaliableDevicesNewSection =  (*Config)["/Script/HTML5PlatformEditor.HTML5SDKSettings"];
+	for (auto It : AvaliableDevicesNewSection)
+	{
+		ValueArray.Reset();
+		if (It.Key == TEXT("DeviceMap"))
+		{
+			FString DeviceName;
+			FString DevicePath;
+			It.Value.RemoveFromStart(TEXT("("));
+			It.Value.RemoveFromEnd(TEXT(")"));
+			It.Value.ParseIntoArray(&ValueArray, TEXT(","), 1);
+			for (auto& Value : ValueArray)
+			{
+				if (Value.StartsWith(TEXT("DeviceName=")))
+				{
+					DeviceName = Value.RightChop(11).TrimQuotes();
+				}
+				else if (Value.StartsWith(TEXT("DevicePath=(FilePath=")))
+				{
+					DevicePath = Value.RightChop(21);
+					DevicePath.RemoveFromEnd(TEXT(")"));
+					DevicePath = DevicePath.TrimQuotes();
+				}
+			}
+
+			if (!DeviceName.IsEmpty() && !DevicePath.IsEmpty() &&
+				FPlatformFileManager::Get().GetPlatformFile().FileExists(*DevicePath))
+			{
+				ITargetDevicePtr Device = MakeShareable(new FHTML5TargetDevice(*this, FString::Printf(TEXT("%s on %s"), *DeviceName, FPlatformProcess::ComputerName())));
+				LocalDevice.Add(Device);
+				DeviceDiscoveredEvent.Broadcast(Device.ToSharedRef());
+			}
+		}
+	}
+
+	// Old Fallback
+	if (LocalDevice.Num() == 0)
+	{
+		FString DeviceSectionName;
+
+#if PLATFORM_WINDOWS
+		DeviceSectionName = "HTML5DevicesWindows";
+#elif PLATFORM_MAC
+		DeviceSectionName = "HTML5DevicesMac";
+#else 
+		DeviceSectionName = "HTML5DevicesLinux";
+#endif 
+
+		FConfigSection AvaliableDevicesSection = HTML5EngineSettings[DeviceSectionName];
+		for (auto It : AvaliableDevicesSection)
+		{
+			const FString& BrowserName = It.Key.ToString();
+			const FString& BrowserPath = It.Value;
+			if (FPlatformFileManager::Get().GetPlatformFile().FileExists(*It.Value))
+			{
+				ITargetDevicePtr Device = MakeShareable(new FHTML5TargetDevice(*this, FString::Printf(TEXT("%s on %s"), *It.Key.ToString(), FPlatformProcess::ComputerName())));
+				LocalDevice.Add(Device);
+				DeviceDiscoveredEvent.Broadcast(Device.ToSharedRef());
+			}
+		}
+	}
+}
+
 #endif // WITH_ENGINE
diff --git a/Engine/Source/Developer/HTML5/HTML5TargetPlatform/Private/HTML5TargetPlatform.h b/Engine/Source/Developer/HTML5/HTML5TargetPlatform/Private/HTML5TargetPlatform.h
--- a/Engine/Source/Developer/HTML5/HTML5TargetPlatform/Private/HTML5TargetPlatform.h
+++ b/Engine/Source/Developer/HTML5/HTML5TargetPlatform/Private/HTML5TargetPlatform.h
@@ -23,7 +23,7 @@ public:
 	 */
 	FHTML5TargetPlatform( );
 
-public:
+	void RefreshAvailableDevices();
 
 	// Begin ITargetPlatform interface
 
diff --git a/Engine/Source/Developer/HTML5/HTML5TargetPlatform/Private/HTML5TargetPlatformClasses.cpp b/Engine/Source/Developer/HTML5/HTML5TargetPlatform/Private/HTML5TargetPlatformClasses.cpp
deleted file mode 100644
diff --git a/Engine/Source/Developer/HTML5/HTML5TargetPlatform/Private/HTML5TargetPlatformModule.cpp b/Engine/Source/Developer/HTML5/HTML5TargetPlatform/Private/HTML5TargetPlatformModule.cpp
--- a/Engine/Source/Developer/HTML5/HTML5TargetPlatform/Private/HTML5TargetPlatformModule.cpp
+++ b/Engine/Source/Developer/HTML5/HTML5TargetPlatform/Private/HTML5TargetPlatformModule.cpp
@@ -15,7 +15,7 @@ static ITargetPlatform* Singleton = NULL;
  * Module for the HTML5 target platform.
  */
 class FHTML5TargetPlatformModule
-	: public ITargetPlatformModule
+	: public IHTML5TargetPlatformModule
 {
 public:
 
@@ -51,6 +51,13 @@ public:
 
 	// End ITargetPlatformModule interface
 
+	virtual void RefreshAvailableDevices() override
+	{
+		if (Singleton)
+		{
+			((FHTML5TargetPlatform*)Singleton)->RefreshAvailableDevices();
+		}
+	}
 
 public:
 
@@ -58,39 +65,12 @@ public:
 
 	virtual void StartupModule() override
 	{
-		TargetSettings = ConstructObject<UHTML5TargetSettings>(UHTML5TargetSettings::StaticClass(), GetTransientPackage(), "HTML5TargetSettings", RF_Standalone);
-		TargetSettings->AddToRoot();
 
-		ISettingsModule* SettingsModule = FModuleManager::GetModulePtr<ISettingsModule>("Settings");
-
-		if (SettingsModule != nullptr)
-		{
-			SettingsModule->RegisterSettings("Project", "Platforms", "HTML5",
-				LOCTEXT("TargetSettingsName", "HTML5"),
-				LOCTEXT("TargetSettingsDescription", "HTML5 platform settings description text here"),
-				TargetSettings
-			);
-		}
 	}
 
 	virtual void ShutdownModule() override
 	{
-		ISettingsModule* SettingsModule = FModuleManager::GetModulePtr<ISettingsModule>("Settings");
 
-		if (SettingsModule != nullptr)
-		{
-			SettingsModule->UnregisterSettings("Project", "Platforms", "HTML5");
-		}
-
-		if (!GExitPurge)
-		{
-			// If we're in exit purge, this object has already been destroyed
-			TargetSettings->RemoveFromRoot();
-		}
-		else
-		{
-			TargetSettings = NULL;
-		}
 	}
 
 	// End IModuleInterface interface
@@ -98,8 +78,6 @@ public:
 
 private:
 
-	// Holds the target settings.
-	UHTML5TargetSettings* TargetSettings;
 };
 
 
diff --git a/Engine/Source/Developer/HTML5/HTML5TargetPlatform/Private/HTML5TargetPlatformPrivatePCH.h b/Engine/Source/Developer/HTML5/HTML5TargetPlatform/Private/HTML5TargetPlatformPrivatePCH.h
--- a/Engine/Source/Developer/HTML5/HTML5TargetPlatform/Private/HTML5TargetPlatformPrivatePCH.h
+++ b/Engine/Source/Developer/HTML5/HTML5TargetPlatform/Private/HTML5TargetPlatformPrivatePCH.h
@@ -17,10 +17,10 @@
 #include "TargetPlatform.h"
 #include "TargetPlatformBase.h"
 
+#include "IHTML5TargetPlatformModule.h"
 
 /* Private includes
  *****************************************************************************/
 
-#include "HTML5TargetSettings.h"
 #include "HTML5TargetDevice.h"
 #include "HTML5TargetPlatform.h"
diff --git a/Engine/Source/Developer/HTML5/HTML5TargetPlatform/Public/IHTML5TargetPlatformModule.h b/Engine/Source/Developer/HTML5/HTML5TargetPlatform/Public/IHTML5TargetPlatformModule.h
new file mode 100644
--- /dev/null
+++ b/Engine/Source/Developer/HTML5/HTML5TargetPlatform/Public/IHTML5TargetPlatformModule.h
@@ -0,0 +1,25 @@
+// Copyright 1998-2015 Epic Games, Inc. All Rights Reserved.
+
+#pragma once
+
+#include "ITargetPlatformModule.h"
+
+/**
+ * Interface for AndroidDeviceDetection module.
+ */
+class IHTML5TargetPlatformModule
+	: public ITargetPlatformModule
+{
+public:
+	/**
+	 * Refresh the list of HTML5 browsers that exist on the system
+	 */
+	virtual void RefreshAvailableDevices() = 0;
+
+protected:
+
+	/**
+	 * Virtual destructor
+	 */
+	virtual ~IHTML5TargetPlatformModule( ) { }
+};
\ No newline at end of file
diff --git a/Engine/Source/Editor/UnrealEd/Private/Editor.cpp b/Engine/Source/Editor/UnrealEd/Private/Editor.cpp
--- a/Engine/Source/Editor/UnrealEd/Private/Editor.cpp
+++ b/Engine/Source/Editor/UnrealEd/Private/Editor.cpp
@@ -648,6 +648,7 @@ void UEditorEngine::Init(IEngineLoop* InEngineLoop)
 			TEXT("ProjectSettingsViewer"),
 			TEXT("AndroidRuntimeSettings"),
 			TEXT("AndroidPlatformEditor"),
+			TEXT("HTML5PlatformEditor"),
 			TEXT("IOSRuntimeSettings"),
 			TEXT("IOSPlatformEditor"),
 			TEXT("Blutility"),
diff --git a/Engine/Source/Editor/UnrealEd/UnrealEd.Build.cs b/Engine/Source/Editor/UnrealEd/UnrealEd.Build.cs
--- a/Engine/Source/Editor/UnrealEd/UnrealEd.Build.cs
+++ b/Engine/Source/Editor/UnrealEd/UnrealEd.Build.cs
@@ -191,6 +191,7 @@ public class UnrealEd : ModuleRules
 				"EditorLiveStreaming",
 				"HotReload",
                 "IOSPlatformEditor",
+				"HTML5PlatformEditor",
 			}
 		);
 
diff --git a/Engine/Source/Programs/AutomationTool/HTML5/HTML5Platform.Automation.cs b/Engine/Source/Programs/AutomationTool/HTML5/HTML5Platform.Automation.cs
--- a/Engine/Source/Programs/AutomationTool/HTML5/HTML5Platform.Automation.cs
+++ b/Engine/Source/Programs/AutomationTool/HTML5/HTML5Platform.Automation.cs
@@ -61,8 +61,12 @@ public class HTML5Platform : Platform
             ulong HeapSize;
             var ConfigCache = new UnrealBuildTool.ConfigCacheIni(UnrealTargetPlatform.HTML5, "Engine", Path.GetDirectoryName(Params.RawProjectPath), CombinePaths(CmdEnv.LocalRoot, "Engine"));
 
-            int ConfigHeapSize;
-            if (!ConfigCache.GetInt32("BuildSettings", "HeapSize" + Params.ClientConfigsToBuild[0].ToString(), out ConfigHeapSize)) // in Megs.
+            int ConfigHeapSize = 0;
+			// Valuer set by Editor UI
+			var bGotHeapSize = ConfigCache.GetInt32("/Script/HTML5PlatformEditor.HTML5TargetSettings", "HeapSize" + Params.ClientConfigsToBuild[0].ToString(), out ConfigHeapSize);
+
+			// Fallback if the previous method failed
+            if (!bGotHeapSize && !ConfigCache.GetInt32("BuildSettings", "HeapSize" + Params.ClientConfigsToBuild[0].ToString(), out ConfigHeapSize)) // in Megs.
             {
                 // we couldn't find a per config heap size, look for a common one.
                 if (!ConfigCache.GetInt32("BuildSettings", "HeapSize", out ConfigHeapSize))
@@ -203,25 +207,46 @@ public class HTML5Platform : Platform
 	{
 		// look for browser
 		var ConfigCache = new UnrealBuildTool.ConfigCacheIni(UnrealTargetPlatform.HTML5, "Engine", Path.GetDirectoryName(Params.RawProjectPath), CombinePaths(CmdEnv.LocalRoot, "Engine"));
+		bool ok = false;
+		List<string> Devices;
+		string browserPath = "";
+		string DeviceName = Params.Device.Split('@')[1];
+		DeviceName = DeviceName.Substring(0, DeviceName.LastIndexOf(" on "));
 
-		string DeviceSection;
-
-		if ( Utils.IsRunningOnMono )
+		if (ConfigCache.GetArray("/Script/HTML5PlatformEditor.HTML5SDKSettings", "DeviceMap", out Devices))
 		{
-			DeviceSection = "HTML5DevicesMac";
+			foreach (var Dev in Devices)
+			{
+				var Matched = Regex.Match(Dev, "\\(DeviceName=\"(.*)\",DevicePath=\\(FilePath=\"(.*)\"\\)\\)", RegexOptions.IgnoreCase);
+				if (Matched.Success && Matched.Groups[1].ToString() == DeviceName)
+				{
+					browserPath = Matched.Groups[2].ToString();
+					ok = true;
+					break;
+				}
+			}
 		}
-		else
+
+		if (!ok && HTML5SDKInfo.bAllowFallbackSDKSettings)
 		{
-			DeviceSection = "HTML5DevicesWindows";
-		}
+			string DeviceSection;
 
-		string browserPath;
-		string DeviceName = Params.Device.Split('@')[1];
-		DeviceName = DeviceName.Substring(0, DeviceName.LastIndexOf(" on "));
-		bool ok = ConfigCache.GetString(DeviceSection, DeviceName, out browserPath);
+			if (Utils.IsRunningOnMono)
+			{
+				DeviceSection = "HTML5DevicesMac";
+			}
+			else
+			{
+				DeviceSection = "HTML5DevicesWindows";
+			}
+
+			ok = ConfigCache.GetString(DeviceSection, DeviceName, out browserPath);
+		}
 
 		if (!ok)
-			throw new System.Exception ("Incorrect browser configuration in HTML5Engine.ini ");
+		{
+			throw new System.Exception("Incorrect browser configuration in HTML5Engine.ini ");
+		}
 
 		// open the webpage
 		string directory = Path.GetDirectoryName(ClientApp);
diff --git a/Engine/Source/Programs/UnrealBuildTool/HTML5/HTML5SDKInfo.cs b/Engine/Source/Programs/UnrealBuildTool/HTML5/HTML5SDKInfo.cs
--- a/Engine/Source/Programs/UnrealBuildTool/HTML5/HTML5SDKInfo.cs
+++ b/Engine/Source/Programs/UnrealBuildTool/HTML5/HTML5SDKInfo.cs
@@ -13,6 +13,7 @@ namespace UnrealBuildTool
     {
         static string EmscriptenSettingsPath = Path.Combine(Environment.GetFolderPath(Environment.SpecialFolder.UserProfile), ".emscripten");
 		static ConfigCacheIni ConfigCache = null;
+		public const bool bAllowFallbackSDKSettings = true;
 
         static Dictionary<string, string> ReadEmscriptenSettings()
         {
@@ -132,32 +133,39 @@ namespace UnrealBuildTool
 			EnsureConfigCacheIsReady();
 			// Search the ini file for Emscripten root first
 			string EMCCPath;
-			bool ok = ConfigCache.GetString("HTML5SDKPaths", "Emscripten", out EMCCPath);
-			if (ok && System.IO.Directory.Exists(EMCCPath))
-				return EMCCPath;
+			if (ConfigCache.GetPath("/Script/HTML5PlatformEditor.HTML5SDKSettings", "Emscripten", out EMCCPath) && System.IO.Directory.Exists(EMCCPath))
+ 				return EMCCPath;
 
-			// Older method used platform name
-            string PlatformName = "";
-            if (!Utils.IsRunningOnMono)
-            {
-                PlatformName = "Windows";
-            }
-            else
-            {
-                PlatformName = "Mac";
-            }
-            ok = ConfigCache.GetString("HTML5SDKPaths", PlatformName, out EMCCPath);
+			// Old Method
+			if (bAllowFallbackSDKSettings)
+			{
+				bool ok = ConfigCache.GetString("HTML5SDKPaths", "Emscripten", out EMCCPath);
+				if (ok && System.IO.Directory.Exists(EMCCPath))
+					return EMCCPath;
+
+				// Older method used platform name
+				string PlatformName = "";
+				if (!Utils.IsRunningOnMono)
+				{
+					PlatformName = "Windows";
+				}
+				else
+				{
+					PlatformName = "Mac";
+				}
+				ok = ConfigCache.GetString("HTML5SDKPaths", PlatformName, out EMCCPath);
 
-            if (ok && System.IO.Directory.Exists(EMCCPath))
-                return EMCCPath;
+				if (ok && System.IO.Directory.Exists(EMCCPath))
+					return EMCCPath;
 
-            // try to find SDK Location from env.
-            if (Environment.GetEnvironmentVariable("EMSCRIPTEN") != null
-                && System.IO.Directory.Exists(Environment.GetEnvironmentVariable("EMSCRIPTEN"))
-                )
-            {
-                return Environment.GetEnvironmentVariable("EMSCRIPTEN");
-            }
+				// try to find SDK Location from env.
+				if (Environment.GetEnvironmentVariable("EMSCRIPTEN") != null
+					&& System.IO.Directory.Exists(Environment.GetEnvironmentVariable("EMSCRIPTEN"))
+					)
+				{
+					return Environment.GetEnvironmentVariable("EMSCRIPTEN");
+				}
+			}
 
             return ""; 
         }
@@ -166,31 +174,37 @@ namespace UnrealBuildTool
         {
 			// find Python.
 			EnsureConfigCacheIsReady();
-			
-			// Check the ini first. Check for Python="path"
+
 			string PythonPath;
-			bool ok = ConfigCache.GetString("HTML5SDKPaths", "Python", out PythonPath);
-			if (ok && System.IO.File.Exists(PythonPath))
+			if (ConfigCache.GetPath("/Script/HTML5PlatformEditor.HTML5SDKSettings", "Python", out PythonPath) && System.IO.File.Exists(PythonPath))
 				return PythonPath;
-           
-            var EmscriptenSettings = ReadEmscriptenSettings();
 
-			// check emscripten generated config file. 
-            if (EmscriptenSettings.ContainsKey("PYTHON")
-                && System.IO.File.Exists(EmscriptenSettings["PYTHON"])
-                )
-            {
-                return EmscriptenSettings["PYTHON"];
-            }
+			if (bAllowFallbackSDKSettings)
+			{
+				// Check the ini first. Check for Python="path"
+				bool ok = ConfigCache.GetString("HTML5SDKPaths", "Python", out PythonPath);
+				if (ok && System.IO.File.Exists(PythonPath))
+					return PythonPath;
 
-            // It might be setup as a env variable. 
-            if (Environment.GetEnvironmentVariable("PYTHON") != null
-                &&
-                System.IO.File.Exists(Environment.GetEnvironmentVariable("PYTHON"))
-                )
-            {
-                return Environment.GetEnvironmentVariable("PYTHON");
-            }
+				var EmscriptenSettings = ReadEmscriptenSettings();
+
+				// check emscripten generated config file. 
+				if (EmscriptenSettings.ContainsKey("PYTHON")
+					&& System.IO.File.Exists(EmscriptenSettings["PYTHON"])
+					)
+				{
+					return EmscriptenSettings["PYTHON"];
+				}
+
+				// It might be setup as a env variable. 
+				if (Environment.GetEnvironmentVariable("PYTHON") != null
+					&&
+					System.IO.File.Exists(Environment.GetEnvironmentVariable("PYTHON"))
+					)
+				{
+					return Environment.GetEnvironmentVariable("PYTHON");
+				}
+			}
 
             // it might just be on path. 
             ProcessStartInfo startInfo = new ProcessStartInfo();
diff --git a/Engine/Source/Runtime/CoreUObject/Private/UObject/Obj.cpp b/Engine/Source/Runtime/CoreUObject/Private/UObject/Obj.cpp
--- a/Engine/Source/Runtime/CoreUObject/Private/UObject/Obj.cpp
+++ b/Engine/Source/Runtime/CoreUObject/Private/UObject/Obj.cpp
@@ -1625,7 +1625,7 @@ void UObject::SaveConfig( uint64 Flags, const TCHAR* InFilename, FConfigCacheIni
 	// Determine whether the file we are writing is a default file config.
 	const bool bIsADefaultIniWrite = !Filename.Contains(FPaths::GameSavedDir())
 		&& !Filename.Contains(FPaths::EngineSavedDir())
-		&& FPaths::GetBaseFilename(Filename).StartsWith(TEXT("Default"));
+		&& (FPaths::GetBaseFilename(Filename).StartsWith(TEXT("Default")) || FPaths::GetBaseFilename(Filename).StartsWith(TEXT("User")));
 
 	const bool bPerObject = UsesPerObjectConfig(this);
 	FString Section;
