- Always add the GMS version to the AndroidManifest.xml [UE-8259]

diff --git a/Engine/Source/Programs/UnrealBuildTool/Android/UEDeployAndroid.cs b/Engine/Source/Programs/UnrealBuildTool/Android/UEDeployAndroid.cs
--- a/Engine/Source/Programs/UnrealBuildTool/Android/UEDeployAndroid.cs
+++ b/Engine/Source/Programs/UnrealBuildTool/Android/UEDeployAndroid.cs
@@ -692,15 +692,12 @@ namespace UnrealBuildTool.Android
 			}
 			Text.AppendLine("\t\t</activity>");
 			Text.AppendLine(string.Format("\t\t<meta-data android:name=\"com.epicgames.ue4.GameActivity.DepthBufferPreference\" android:value=\"{0}\"/>", ConvertDepthBufferIniValue(DepthBufferPreference)));
-			if (bEnableGooglePlaySupport)
-			{
-				Text.AppendLine("\t\t<meta-data android:name=\"com.google.android.gms.games.APP_ID\"");
-				Text.AppendLine("\t\t           android:value=\"@string/app_id\" />");
-				Text.AppendLine("\t\t<meta-data android:name=\"com.google.android.gms.version\"");
-				Text.AppendLine("\t\t           android:value=\"@integer/google_play_services_version\" />");
-				Text.AppendLine("\t\t<activity android:name=\"com.google.android.gms.ads.AdActivity\"");
-				Text.AppendLine("\t\t          android:configChanges=\"keyboard|keyboardHidden|orientation|screenLayout|uiMode|screenSize|smallestScreenSize\"/>");
-			}
+			Text.AppendLine("\t\t<meta-data android:name=\"com.google.android.gms.games.APP_ID\"");
+			Text.AppendLine("\t\t           android:value=\"@string/app_id\" />");
+			Text.AppendLine("\t\t<meta-data android:name=\"com.google.android.gms.version\"");
+			Text.AppendLine("\t\t           android:value=\"@integer/google_play_services_version\" />");
+			Text.AppendLine("\t\t<activity android:name=\"com.google.android.gms.ads.AdActivity\"");
+			Text.AppendLine("\t\t          android:configChanges=\"keyboard|keyboardHidden|orientation|screenLayout|uiMode|screenSize|smallestScreenSize\"/>");
 			if (!string.IsNullOrEmpty(ExtraApplicationSettings))
 			{
 				ExtraApplicationSettings = ExtraApplicationSettings.Replace("\\n", "\n");
