[INTEGRATE] Change 2417634 by [REDACTED_PERSON][REDACTED_HANDLE] on 2015/01/23 22:05:59

	Fix mono warnings about unused variables.

diff --git a/Engine/Source/Programs/UnrealBuildTool/IOS/UEDeployIOS.cs b/Engine/Source/Programs/UnrealBuildTool/IOS/UEDeployIOS.cs
--- a/Engine/Source/Programs/UnrealBuildTool/IOS/UEDeployIOS.cs
+++ b/Engine/Source/Programs/UnrealBuildTool/IOS/UEDeployIOS.cs
@@ -599,11 +599,11 @@ namespace UnrealBuildTool.IOS
 			}
 
 			// compile the launch .xib
-			string LaunchXib = InEngineDir + "/Build/IOS/Resources/Interface/LaunchScreen.xib";
-			if (File.Exists(BuildDirectory + "/Resources/Interface/LaunchScreen.xib"))
-			{
-				LaunchXib = BuildDirectory + "/Resources/Interface/LaunchScreen.xib";
-			}
+//			string LaunchXib = InEngineDir + "/Build/IOS/Resources/Interface/LaunchScreen.xib";
+//			if (File.Exists(BuildDirectory + "/Resources/Interface/LaunchScreen.xib"))
+//			{
+//				LaunchXib = BuildDirectory + "/Resources/Interface/LaunchScreen.xib";
+//			}
 
 			bool bSkipDefaultPNGs = GeneratePList(InProjectDirectory, bIsUE4Game, GameName, InProjectName, InEngineDir, AppDirectory);
 
diff --git a/Engine/Source/Programs/UnrealBuildTool/Mac/MacToolChain.cs b/Engine/Source/Programs/UnrealBuildTool/Mac/MacToolChain.cs
--- a/Engine/Source/Programs/UnrealBuildTool/Mac/MacToolChain.cs
+++ b/Engine/Source/Programs/UnrealBuildTool/Mac/MacToolChain.cs
@@ -581,7 +581,7 @@ namespace UnrealBuildTool
 		{
 			string LibraryDir = Path.GetDirectoryName(Library);
 			string ExeDir = Path.GetDirectoryName(ExeAbsolutePath);
-			string AbsPath = Path.GetFullPath(Library);
+			//string AbsPath = Path.GetFullPath(Library);
 			if ((Library.Contains("/Plugins/") || Library.Contains("/Binaries/ThirdParty/")) && (Library.EndsWith("dylib") || Library.EndsWith(".framework")) && LibraryDir != ExeDir)
 			{
 				string RelativePath = Utils.MakePathRelativeTo(LibraryDir, ExeDir).Replace("\\", "/");
