- Moved OBBInAPK into ProjectSettings, and renamed it [UE-8359]
- Removed all the OBBInAPK hackery in UBT
- Updated the apk install scripts with some more information as to what it's doing [UE-8258]

diff --git a/Engine/Build/Android/Java/src/com/epicgames/ue4/GameActivity.java b/Engine/Build/Android/Java/src/com/epicgames/ue4/GameActivity.java
--- a/Engine/Build/Android/Java/src/com/epicgames/ue4/GameActivity.java
+++ b/Engine/Build/Android/Java/src/com/epicgames/ue4/GameActivity.java
@@ -65,7 +65,6 @@ import com.epicgames.ue4.GooglePlayLicensing;
 
 // TODO: use the resources from the UE4 lib project once we've got the packager up and running
 //import com.epicgames.ue4.R;
-import com.epicgames.ue4.JavaBuildSettings;
 
 //Extending NativeActivity so that this Java class is instantiated
 //from the beginning of the program.  This will allow the user
@@ -95,6 +94,9 @@ public class GameActivity extends NativeActivity
 	AlertDialog virtualKeyboardAlert;
 	EditText virtualKeyboardInputBox;
 
+	// default the PackageDataInsideApk to an invalid value to make sure we don't get it too early
+	private static int PackageDataInsideApkValue = -1;
+	
 	/** AssetManger reference - populated on start up and used when the OBB is packed into the APK */
 	private AssetManager			AssetManagerReference;
 	
@@ -273,12 +275,29 @@ public class GameActivity extends NativeActivity
 			{
 				DepthBufferPreference = bundle.getInt("com.epicgames.ue4.GameActivity.DepthBufferPreference");
 				Log.debug( "Found DepthBufferPreference = " + DepthBufferPreference);
-			} else {
+			}
+			else
+			{
 				Log.debug( "Did not find DepthBufferPreference, using default.");
 			}
-		} catch (NameNotFoundException e) {
+
+			if (bundle.containsKey("com.epicgames.ue4.GameActivity.bPackageDataInsideApk"))
+			{
+				PackageDataInsideApkValue = bundle.getBoolean("com.epicgames.ue4.GameActivity.bPackageDataInsideApk") ? 1 : 0;
+				Log.debug( "Found bPackageDataInsideApk = " + PackageDataInsideApkValue);
+			}
+			else
+			{
+				PackageDataInsideApkValue = 0;
+				Log.debug( "Did not find bPackageDataInsideApk, using default.");
+			}
+		}
+		catch (NameNotFoundException e)
+		{
 			Log.debug( "Failed to load meta-data: NameNotFound: " + e.getMessage());
-		} catch (NullPointerException e) {
+		}
+		catch (NullPointerException e)
+		{
 			Log.debug( "Failed to load meta-data: NullPointer: " + e.getMessage());
 		}
 
@@ -761,7 +780,8 @@ public class GameActivity extends NativeActivity
 
 	public static boolean isOBBInAPK()
 	{
-		return JavaBuildSettings.PackageType.AMAZON == JavaBuildSettings.PACKAGING;
+		Log.debug("Asking if osOBBInAPK? " + (PackageDataInsideApkValue == 1));
+		return PackageDataInsideApkValue == 1;
 	}
 
 	public void AndroidThunkJava_Minimize()
diff --git a/Engine/Config/BaseEngine.ini b/Engine/Config/BaseEngine.ini
--- a/Engine/Config/BaseEngine.ini
+++ b/Engine/Config/BaseEngine.ini
@@ -2143,6 +2143,7 @@ bBuildForES2=true
 bBuildForES31=false
 bBuildForGL4=false
 bSplitIntoSeparateApks=false
+bPackageDataInsideApk=false
 Orientation=landscape
 DepthBufferPreference=0
 PackageName=com.YourCompany.[PROJECT]
diff --git a/Engine/Config/BaseGame.ini b/Engine/Config/BaseGame.ini
--- a/Engine/Config/BaseGame.ini
+++ b/Engine/Config/BaseGame.ini
@@ -38,7 +38,6 @@ FullRebuild=True
 UsePakFile=True
 bGenerateChunks=False
 IncludePrerequisites=True
-UseOBB_InAPK=False
 IncludeCrashReporter=True
 InternationalizationPreset=English
 +CulturesToStage=en
diff --git a/Engine/Source/Editor/UnrealEd/Classes/Settings/ProjectPackagingSettings.h b/Engine/Source/Editor/UnrealEd/Classes/Settings/ProjectPackagingSettings.h
--- a/Engine/Source/Editor/UnrealEd/Classes/Settings/ProjectPackagingSettings.h
+++ b/Engine/Source/Editor/UnrealEd/Classes/Settings/ProjectPackagingSettings.h
@@ -93,10 +93,6 @@ public:
 	UPROPERTY(config, EditAnywhere, Category=Packaging)
 	bool IncludePrerequisites;
 
-	/** If enabled, on Android platforms, .pak files are placed inside the APK. */
-	UPROPERTY(config, EditAnywhere, Category = Packaging, meta = (DisplayName = "Package .pak in APK"))
-	bool UseOBB_InAPK;
-
 	/**
 	 * Specifies whether to include the crash reporter in the packaged project. 
 	 * This is included by default for Blueprint based projects, but can optionally be disabled.
diff --git a/Engine/Source/Programs/AutomationTool/Android/AndroidPlatform.Automation.cs b/Engine/Source/Programs/AutomationTool/Android/AndroidPlatform.Automation.cs
--- a/Engine/Source/Programs/AutomationTool/Android/AndroidPlatform.Automation.cs
+++ b/Engine/Source/Programs/AutomationTool/Android/AndroidPlatform.Automation.cs
@@ -189,9 +189,7 @@ public class AndroidPlatform : Platform
 		string[] Architectures = UnrealBuildTool.AndroidToolChain.GetAllArchitectures();
 		string[] GPUArchitectures = UnrealBuildTool.AndroidToolChain.GetAllGPUArchitectures();
 		bool bMakeSeparateApks = UnrealBuildTool.Android.UEDeployAndroid.ShouldMakeSeparateApks();
-
-		// Make sure this setting is sync'd pre-build
-		UEBuildConfiguration.bOBBinAPK = Params.OBBinAPK;
+		bool bPackageDataInsideApk = UnrealBuildTool.Android.UEDeployAndroid.PackageDataInsideApk();
 
 		var Deploy = UEBuildDeploy.GetBuildDeploy(UnrealTargetPlatform.Android);
 
@@ -248,7 +246,7 @@ public class AndroidPlatform : Platform
 			    // Create APK specific OBB in case we have a detached OBB.
 			    string DeviceObbName = "";
 			    string ObbName = "";
-			    if (!Params.OBBinAPK)
+				if (!bPackageDataInsideApk)
 			    {
 				    DeviceObbName = GetDeviceObbName(ApkName);
 				    ObbName = GetFinalObbName(ApkName);
@@ -262,7 +260,7 @@ public class AndroidPlatform : Platform
 				string[] BatchLines;
 				if (Utils.IsRunningOnMono)
 				{
-					Log("Writing shell script for install with {0}", Params.OBBinAPK ? "OBB in APK" : "OBB separate");
+					Log("Writing shell script for install with {0}", bPackageDataInsideApk ? "data in APK" : "separate obb");
 					BatchLines = new string[] {
 						"#!/bin/sh",
 						"cd \"`dirname \"$0\"`\"",
@@ -270,15 +268,24 @@ public class AndroidPlatform : Platform
 						"if [ \"$ANDROID_HOME\" != \"\" ]; then ADB=$ANDROID_HOME/platform-tools/adb; else ADB=" +Environment.GetEnvironmentVariable("ANDROID_HOME") + "; fi",
 						"DEVICE=",
 						"if [ \"$1\" != \"\" ]; then DEVICE=\"-s $1\"; fi",
+						"echo",
+						"echo Uninstalling existing application. Failures here can almost always be ignored.",
 						"$ADB $DEVICE uninstall " + PackageName,
+						"echo",
+						"echo Installing existing application. Failures here indicate a problem with the device \\(connection or storage permissions\\) and are fatal.",
 						"$ADB $DEVICE install " + Path.GetFileName(ApkName),
 						"if [ $? -eq 0 ]; then",
+						"\techo",
+						"\techo Removing old data. Failures here are usually fine - indicating the files were not on the device.",
 						"\t$ADB $DEVICE shell 'rm -r $EXTERNAL_STORAGE/" + Params.ShortProjectName + "'",
 						"\t$ADB $DEVICE shell 'rm -r $EXTERNAL_STORAGE/UE4Game/UE4CommandLine.txt" + "'",
 						"\t$ADB $DEVICE shell 'rm -r $EXTERNAL_STORAGE/obb/" + PackageName + "'",
-						"\tSTORAGE=$(echo \"`$ADB $DEVICE shell 'echo $EXTERNAL_STORAGE'`\" | cat -v | tr -d '^M')",
-						"\t$ADB $DEVICE push " + Path.GetFileName(ObbName) + " $STORAGE/" + DeviceObbName,
-						"\tif [ $? -eq 0 ]; then",
+						bPackageDataInsideApk ? "" : "\techo",
+						bPackageDataInsideApk ? "" : "\techo Installing new data. Failures here indicate storage problems \\(missing SD card or bad permissions\\) and are fatal.",
+						bPackageDataInsideApk ? "" : "\tSTORAGE=$(echo \"`$ADB $DEVICE shell 'echo $EXTERNAL_STORAGE'`\" | cat -v | tr -d '^M')",
+						bPackageDataInsideApk ? "" : "\t$ADB $DEVICE push " + Path.GetFileName(ObbName) + " $STORAGE/" + DeviceObbName,
+						bPackageDataInsideApk ? "if [ 1 ]; then" : "\tif [ $? -eq 0 ]; then",
+						"\t\techo",
 						"\t\techo Installation successful",
 						"\t\texit 0",
 						"\tfi",
@@ -287,7 +294,7 @@ public class AndroidPlatform : Platform
 						"echo There was an error installing the game or the obb file. Look above for more info.",
 						"echo",
 						"echo Things to try:",
-						"echo Check that the device (and only the device) is listed with \"$ADB devices\" from a command prompt.",
+						"echo Check that the device (and only the device) is listed with \\\"$ADB devices\\\" from a command prompt.",
 						"echo Make sure all Developer options look normal on the device",
 						"echo Check that the device has an SD card.",
 						"exit 1"
@@ -295,7 +302,7 @@ public class AndroidPlatform : Platform
 				}
 				else
 				{
-					Log("Writing bat for install with {0}", Params.OBBinAPK ? "OBB in APK" : "OBB separate");
+					Log("Writing bat for install with {0}", bPackageDataInsideApk ? "data in APK" : "separate OBB");
 					BatchLines = new string[] {
 						"setlocal",
                         "set ANDROIDHOME=%ANDROID_HOME%",
@@ -304,14 +311,22 @@ public class AndroidPlatform : Platform
 						"set DEVICE=",
                         "if not \"%1\"==\"\" set DEVICE=-s %1",
                         "for /f \"delims=\" %%A in ('%ADB% %DEVICE% " + GetStorageQueryCommand() +"') do @set STORAGE=%%A",
+						"@echo.",
+						"@echo Uninstalling existing application. Failures here can almost always be ignored.",
 						"%ADB% %DEVICE% uninstall " + PackageName,
+						"@echo.",
+						"@echo Installing existing application. Failures here indicate a problem with the device \\(connection or storage permissions\\) and are fatal.",
 						"%ADB% %DEVICE% install " + Path.GetFileName(ApkName),
 						"@if \"%ERRORLEVEL%\" NEQ \"0\" goto Error",
 						"%ADB% %DEVICE% shell rm -r %STORAGE%/" + Params.ShortProjectName,
 						"%ADB% %DEVICE% shell rm -r %STORAGE%/UE4Game/UE4CommandLine.txt", // we need to delete the commandline in UE4Game or it will mess up loading
 						"%ADB% %DEVICE% shell rm -r %STORAGE%/obb/" + PackageName,
-						Params.OBBinAPK ? "" : "%ADB% %DEVICE% push " + Path.GetFileName(ObbName) + " %STORAGE%/" + DeviceObbName,
-						Params.OBBinAPK ? "" : "if \"%ERRORLEVEL%\" NEQ \"0\" goto Error",
+						bPackageDataInsideApk ? "" : "@echo.",
+						bPackageDataInsideApk ? "" : "@echo Installing new data. Failures here indicate storage problems \\(missing SD card or bad permissions\\) and are fatal.",
+						bPackageDataInsideApk ? "" : "%ADB% %DEVICE% push " + Path.GetFileName(ObbName) + " %STORAGE%/" + DeviceObbName,
+						bPackageDataInsideApk ? "" : "if \"%ERRORLEVEL%\" NEQ \"0\" goto Error",
+						"@echo.",
+						"@echo Installation successful",
 						"goto:eof",
 						":Error",
 						"@echo.",
@@ -348,6 +363,7 @@ public class AndroidPlatform : Platform
 		string[] Architectures = UnrealBuildTool.AndroidToolChain.GetAllArchitectures();
 		string[] GPUArchitectures = UnrealBuildTool.AndroidToolChain.GetAllGPUArchitectures();
 		bool bMakeSeparateApks = UnrealBuildTool.Android.UEDeployAndroid.ShouldMakeSeparateApks();
+		bool bPackageDataInsideApk = UnrealBuildTool.Android.UEDeployAndroid.PackageDataInsideApk();
 
 		bool bAddedOBB = false;
 		foreach (string Architecture in Architectures)
@@ -365,7 +381,7 @@ public class AndroidPlatform : Platform
 					ErrorReporter.Error(ErrorString, (int)ErrorCodes.Error_AppNotFound);
 					throw new AutomationException(ErrorString);
 				}
-				if (!Params.OBBinAPK && !FileExists(ObbName))
+				if (!bPackageDataInsideApk && !FileExists(ObbName))
 				{
 					string ErrorString = String.Format("ARCHIVE FAILED - {0} was not found", ObbName);
 					ErrorReporter.Error(ErrorString, (int)ErrorCodes.Error_ObbNotFound);
@@ -373,7 +389,7 @@ public class AndroidPlatform : Platform
 				}
 
 				SC.ArchiveFiles(Path.GetDirectoryName(ApkName), Path.GetFileName(ApkName));
-				if (!Params.OBBinAPK && !bAddedOBB)
+				if (!bPackageDataInsideApk && !bAddedOBB)
 				{
 					bAddedOBB = true;
 					SC.ArchiveFiles(Path.GetDirectoryName(ObbName), Path.GetFileName(ObbName));
diff --git a/Engine/Source/Programs/AutomationTool/ProjectParams.cs b/Engine/Source/Programs/AutomationTool/ProjectParams.cs
--- a/Engine/Source/Programs/AutomationTool/ProjectParams.cs
+++ b/Engine/Source/Programs/AutomationTool/ProjectParams.cs
@@ -282,7 +282,6 @@ namespace AutomationTool
 			this.Distribution = InParams.Distribution;
 			this.Prereqs = InParams.Prereqs;
 			this.NoBootstrapExe = InParams.NoBootstrapExe;
-            this.OBBinAPK = InParams.OBBinAPK;
             this.Prebuilt = InParams.Prebuilt;
             this.RunTimeoutSeconds = InParams.RunTimeoutSeconds;
 		}
@@ -374,7 +373,6 @@ namespace AutomationTool
 			bool? ArchiveMetaData = null,
 			ParamList<string> ProgramTargets = null,
 			bool? Distribution = null,
-            bool? OBBinAPK = null,
             bool? Prebuilt = null,
             int? RunTimeoutSeconds = null,
 			string OverrideMinimumOS = null,
@@ -508,7 +506,6 @@ namespace AutomationTool
 			this.Distribution = GetParamValueIfNotSpecified(Command, Distribution, this.Distribution, "distribution");
 			this.Prereqs = GetParamValueIfNotSpecified(Command, Prereqs, this.Prereqs, "prereqs");
 			this.NoBootstrapExe = GetParamValueIfNotSpecified(Command, NoBootstrapExe, this.NoBootstrapExe, "nobootstrapexe");
-            this.OBBinAPK = GetParamValueIfNotSpecified(Command, OBBinAPK, this.OBBinAPK, "obbinapk");
             this.Prebuilt = GetParamValueIfNotSpecified(Command, Prebuilt, this.Prebuilt, "prebuilt");
             if (this.Prebuilt)
             {
@@ -1254,9 +1251,6 @@ namespace AutomationTool
 		[Help("prereqs", "stage prerequisites along with the project")]
 		public bool Prereqs { get; set; }
 
-		[Help("obbinapk", "package with OBB data in APK assets directory")]
-        public bool OBBinAPK {get; set; }
-
         [Help("Prebuilt", "this is a prebuilt cooked and packaged build")]
         public bool Prebuilt { get; private set; }
 
@@ -1928,7 +1922,6 @@ namespace AutomationTool
 				CommandUtils.Log("ProjectGameExeFilename={0}", ProjectGameExeFilename);
 				CommandUtils.Log("ProjectGameExePath={0}", ProjectGameExePath);
 				CommandUtils.Log("Distribution={0}", Distribution);
-                CommandUtils.Log("OBBinAPK={0}", OBBinAPK);
                 CommandUtils.Log("Prebuilt={0}", Prebuilt);
 				CommandUtils.Log("Prereqs={0}", Prereqs);
 				CommandUtils.Log("NoBootstrapExe={0}", NoBootstrapExe);
diff --git a/Engine/Source/Programs/AutomationTool/Scripts/BuildCookRun.Automation.cs b/Engine/Source/Programs/AutomationTool/Scripts/BuildCookRun.Automation.cs
--- a/Engine/Source/Programs/AutomationTool/Scripts/BuildCookRun.Automation.cs
+++ b/Engine/Source/Programs/AutomationTool/Scripts/BuildCookRun.Automation.cs
@@ -133,11 +133,6 @@ public class BuildCookRun : BuildCommand
 			}
 		}
 
-        var gameConfigCache = new UnrealBuildTool.ConfigCacheIni("Game", Path.GetDirectoryName(ProjectPath));
-        bool obbInAPK = false;
-        gameConfigCache.GetBool("/Script/UnrealEd.ProjectPackagingSettings", "UseOBB_InAPK", out obbInAPK);
-        Params.OBBinAPK = obbInAPK;         
-
 		Params.ValidateAndLog();
 		return Params;
 	}
diff --git a/Engine/Source/Programs/UnrealBuildTool/Android/UEDeployAndroid.cs b/Engine/Source/Programs/UnrealBuildTool/Android/UEDeployAndroid.cs
--- a/Engine/Source/Programs/UnrealBuildTool/Android/UEDeployAndroid.cs
+++ b/Engine/Source/Programs/UnrealBuildTool/Android/UEDeployAndroid.cs
@@ -104,6 +104,21 @@ namespace UnrealBuildTool.Android
 			return CachedSDKLevel;
 		}
 
+		public static bool PackageDataInsideApk(ConfigCacheIni Ini=null)
+		{
+			// make a new one if one wasn't passed in
+			if (Ini == null)
+			{
+				Ini = new ConfigCacheIni(UnrealTargetPlatform.Android, "Engine", UnrealBuildTool.GetUProjectPath());
+			}
+
+			// we check this a lot, so make it easy 
+			bool bPackageDataInsideApk;
+			Ini.GetBool("/Script/AndroidRuntimeSettings.AndroidRuntimeSettings", "bPackageDataInsideApk", out bPackageDataInsideApk);
+
+			return bPackageDataInsideApk;
+		}
+
 		private static string GetAntPath()
 		{
 			// look up an ANT_HOME env var
@@ -253,43 +268,6 @@ namespace UnrealBuildTool.Android
             return Path.GetFullPath(Path.Combine(EngineDirectory, "Build/Android/Java"));
         }
 
-        public string GetUE4JavaFilePath(String EngineDirectory)
-        {
-            return Path.GetFullPath(Path.Combine(GetUE4BuildFilePath(EngineDirectory), "src/com/epicgames/ue4"));
-        }
-
-        public string GetUE4JavaBuildSettingsFileName(String EngineDirectory)
-        {
-            return Path.Combine(GetUE4JavaFilePath(EngineDirectory), "JavaBuildSettings.java");
-        }
-
-        public void WriteJavaBuildSettingsFile(string FileName, bool OBBinAPK)
-        {
-             // (!UEBuildConfiguration.bOBBinAPK ? "PackageType.AMAZON" : /*bPackageForGoogle ? "PackageType.GOOGLE" :*/ "PackageType.DEVELOPMENT") + ";\n");
-            string Setting = OBBinAPK ? "AMAZON" : "DEVELOPMENT";
-            if (!File.Exists(FileName) || ShouldWriteJavaBuildSettingsFile(FileName, Setting))
-            {
-				Log.TraceInformation("\n===={0}====WRITING JAVABUILDSETTINGS.JAVA========================================================", DateTime.Now.ToString());
-				StringBuilder BuildSettings = new StringBuilder("package com.epicgames.ue4;\npublic class JavaBuildSettings\n{\n");
-                BuildSettings.Append("\tpublic enum PackageType {AMAZON, GOOGLE, DEVELOPMENT};\n");
-                BuildSettings.Append("\tpublic static final PackageType PACKAGING = PackageType." + Setting + ";\n");
-                BuildSettings.Append("}\n");
-                File.WriteAllText(FileName, BuildSettings.ToString());
-            }
-        }
-
-        public bool ShouldWriteJavaBuildSettingsFile(string FileName, string setting)
-        {
-            var fileContent = File.ReadAllLines(FileName);
-            if (fileContent.Length < 5)
-                return true;
-            var packageLine = fileContent[4]; // We know this to be true... because we write it below...
-            int location = packageLine.IndexOf("PACKAGING") + 12 + 12; // + ("PACKAGING = ") + ("PackageType.")
-            if (location == -1)
-                return true;
-            return String.Compare(setting, packageLine.Substring(location, Math.Min(packageLine.Length - location - 1, setting.Length))) != 0;
-        }
-
 		private static string GetNDKArch(string UE4Arch)
 		{
 			switch (UE4Arch)
@@ -460,7 +438,7 @@ namespace UnrealBuildTool.Android
 		}
 
 
-		private string GetAllBuildSettings(string BuildPath, bool bForDistribution, bool bMakeSeparateApks, bool bOBBinApk)
+		private string GetAllBuildSettings(string BuildPath, bool bForDistribution, bool bMakeSeparateApks)
 		{
 			// make the settings string - this will be char by char compared against last time
 			StringBuilder CurrentSettings = new StringBuilder();
@@ -471,7 +449,6 @@ namespace UnrealBuildTool.Android
 			CurrentSettings.AppendLine(string.Format("SDKVersion={0}", GetSdkApiLevel()));
 			CurrentSettings.AppendLine(string.Format("bForDistribution={0}", bForDistribution));
 			CurrentSettings.AppendLine(string.Format("bMakeSeparateApks={0}", bMakeSeparateApks));
-			CurrentSettings.AppendLine(string.Format("bOBBinApk={0}", bOBBinApk));
 
 			// all AndroidRuntimeSettings ini settings in here
 			ConfigCacheIni Ini = new ConfigCacheIni(UnrealTargetPlatform.Android, "Engine", UnrealBuildTool.GetUProjectPath());
@@ -516,7 +493,7 @@ namespace UnrealBuildTool.Android
 			return CurrentSettings.ToString();
 		}
 
-		private bool CheckDependencies(string ProjectName, string ProjectDirectory, string UE4BuildFilesPath, string GameBuildFilesPath, string EngineDirectory, string JavaSettingsFile, string CookFlavor, string OutputPath, string UE4BuildPath, bool bMakeSeparateApks)
+		private bool CheckDependencies(string ProjectName, string ProjectDirectory, string UE4BuildFilesPath, string GameBuildFilesPath, string EngineDirectory, string CookFlavor, string OutputPath, string UE4BuildPath, bool bMakeSeparateApks)
 		{
 			string[] Arches = AndroidToolChain.GetAllArchitectures();
 			string[] GPUArchitectures = AndroidToolChain.GetAllGPUArchitectures();
@@ -547,11 +524,8 @@ namespace UnrealBuildTool.Android
 						InputFiles.AddRange(Directory.EnumerateFiles(GameBuildFilesPath, "*.*", SearchOption.AllDirectories));
 					}
 
-					// make sure changed java settings will rebuild apk
-					InputFiles.Add(JavaSettingsFile);
-
-					// rebuild if .pak files exist for OBB in APK case
-					if (UEBuildConfiguration.bOBBinAPK)
+					// rebuild if .pak files exist for  in APK case
+					if (PackageDataInsideApk())
 					{
 						string PAKFileLocation = ProjectDirectory + "/Saved/StagedBuilds/Android" + CookFlavor + "/" + ProjectName + "/Content/Paks";
 						if (Directory.Exists(PAKFileLocation))
@@ -692,6 +666,7 @@ namespace UnrealBuildTool.Android
 			}
 			Text.AppendLine("\t\t</activity>");
 			Text.AppendLine(string.Format("\t\t<meta-data android:name=\"com.epicgames.ue4.GameActivity.DepthBufferPreference\" android:value=\"{0}\"/>", ConvertDepthBufferIniValue(DepthBufferPreference)));
+			Text.AppendLine(string.Format("\t\t<meta-data android:name=\"com.epicgames.ue4.GameActivity.bPackageDataInsideApk\" android:value=\"{0}\"/>", PackageDataInsideApk(Ini) ? "true" : "false"));
 			Text.AppendLine("\t\t<meta-data android:name=\"com.google.android.gms.games.APP_ID\"");
 			Text.AppendLine("\t\t           android:value=\"@string/app_id\" />");
 			Text.AppendLine("\t\t<meta-data android:name=\"com.google.android.gms.version\"");
@@ -752,13 +727,8 @@ namespace UnrealBuildTool.Android
 			string[] GPUArchitectures = AndroidToolChain.GetAllGPUArchitectures();
 //			int NumArches = Arches.Length * GPUArchitectures.Length;
 
-
-            // See if we need to create a 'default' Java Build settings file if one doesn't exist (if it does exist we have to assume it has been setup correctly)
-            string UE4JavaBuildSettingsFileName = GetUE4JavaBuildSettingsFileName(EngineDirectory);
-            WriteJavaBuildSettingsFile(UE4JavaBuildSettingsFileName, UEBuildConfiguration.bOBBinAPK);
-
 			// check to see if any "meta information" is newer than last time we build
-			string CurrentBuildSettings = GetAllBuildSettings(UE4BuildPath, bForDistribution, bMakeSeparateApks, UEBuildConfiguration.bOBBinAPK);
+			string CurrentBuildSettings = GetAllBuildSettings(UE4BuildPath, bForDistribution, bMakeSeparateApks);
 			string BuildSettingsCacheFile = Path.Combine(UE4BuildPath, "UEBuildSettings.txt");
 
 			// do we match previous build settings?
@@ -795,7 +765,7 @@ namespace UnrealBuildTool.Android
 			{
 				// check if so's are up to date against various inputs
 				bool bAllInputsCurrent = CheckDependencies(ProjectName, ProjectDirectory, UE4BuildFilesPath, GameBuildFilesPath, 
-					EngineDirectory, UE4JavaBuildSettingsFileName, CookFlavor, OutputPath, UE4BuildPath, bMakeSeparateApks);
+					EngineDirectory, CookFlavor, OutputPath, UE4BuildPath, bMakeSeparateApks);
 
 				if (bAllInputsCurrent)
 				{
@@ -825,9 +795,10 @@ namespace UnrealBuildTool.Android
 				DeleteDirectory(UE4BuildPath, "dexedLibs");
 			}
 
-			// If we are packaging for Amazon then we need to copy the OBB file to the correct location
-			Log.TraceInformation("UEBuildConfiguration.bOBBinAPK = {0}", UEBuildConfiguration.bOBBinAPK);
-			if (UEBuildConfiguration.bOBBinAPK)
+			// If we are packaging for Amazon then we need to copy the  file to the correct location
+			bool bPackageDataInsideApk = PackageDataInsideApk();
+			Log.TraceInformation("bPackageDataInsideApk = {0}", bPackageDataInsideApk);
+			if (bPackageDataInsideApk)
 			{
 				string ObbFileLocation = ProjectDirectory + "/Saved/StagedBuilds/Android" + CookFlavor + ".obb";
 				Console.WriteLine("Obb location {0}", ObbFileLocation);
diff --git a/Engine/Source/Programs/UnrealBuildTool/Configuration/UEBuildConfiguration.cs b/Engine/Source/Programs/UnrealBuildTool/Configuration/UEBuildConfiguration.cs
--- a/Engine/Source/Programs/UnrealBuildTool/Configuration/UEBuildConfiguration.cs
+++ b/Engine/Source/Programs/UnrealBuildTool/Configuration/UEBuildConfiguration.cs
@@ -160,10 +160,6 @@ namespace UnrealBuildTool
 		[XmlConfig]
 		public static bool bUseLoggingInShipping;
 
-        /** True if we need to package up Android with the OBB in the APK file */
-		[XmlConfig]
-        public static bool bOBBinAPK;
-
 		/** True if we need PhysX vehicle support */
 		[XmlConfig]
 		public static bool bCompilePhysXVehicle;
@@ -223,7 +219,6 @@ namespace UnrealBuildTool
             bUseLoggingInShipping = false;
 			bCompileSteamOSS = true;
 			bCompileMcpOSS = true;
-            bOBBinAPK = false;
 			bCompilePhysXVehicle = true;
 			bCompileFreeType = true;
 			bCompileForSize = false;
diff --git a/Engine/Source/Programs/UnrealBuildTool/System/UnrealBuildTool.cs b/Engine/Source/Programs/UnrealBuildTool/System/UnrealBuildTool.cs
--- a/Engine/Source/Programs/UnrealBuildTool/System/UnrealBuildTool.cs
+++ b/Engine/Source/Programs/UnrealBuildTool/System/UnrealBuildTool.cs
@@ -863,10 +863,6 @@ namespace UnrealBuildTool
                         {
                             //ConfigName = Arg;
                         }
-                        else if (LowercaseArg == "-obbinapk")
-                        {
-                            UEBuildConfiguration.bOBBinAPK = true;
-                        }
                         else if (LowercaseArg == "-modulewithsuffix")
                         {
                             bSpecificModulesOnly = true;
diff --git a/Engine/Source/Runtime/Android/AndroidRuntimeSettings/Classes/AndroidRuntimeSettings.h b/Engine/Source/Runtime/Android/AndroidRuntimeSettings/Classes/AndroidRuntimeSettings.h
--- a/Engine/Source/Runtime/Android/AndroidRuntimeSettings/Classes/AndroidRuntimeSettings.h
+++ b/Engine/Source/Runtime/Android/AndroidRuntimeSettings/Classes/AndroidRuntimeSettings.h
@@ -120,6 +120,10 @@ public:
 	UPROPERTY(GlobalConfig, EditAnywhere, Category = APKPackaging, Meta = (DisplayName = "Minimum SDK Version (9=Gingerbread, 14=Ice Cream Sandwich, 21=Lollipop)"))
 	int32 MinSDKVersion;
 	
+	// Should the data be placed into the .apk file instead of a separate .obb file. Amazon requires this to be enabled, but Google Play Store will not allow .apk files larger than 50MB, so only small games will work with this enabled.
+	UPROPERTY(GlobalConfig, EditAnywhere, Category = APKPackaging, Meta = (DisplayName = "Package game data inside .apk?"))
+	bool bPackageDataInsideApk;
+	
 	// The permitted orientation or orientations of the application on the device
 	UPROPERTY(GlobalConfig, EditAnywhere, Category = APKPackaging)
 	TEnumAsByte<EAndroidScreenOrientation::Type> Orientation;
@@ -177,7 +181,7 @@ public:
 	bool bBuildForES2;
 
 	// Enable ES31 support? [CURRENTLY FOR FULL SOURCE GAMES ONLY]
-	UPROPERTY(GlobalConfig, EditAnywhere, Category = Build, meta = (DisplayName = "Support OpenGL ES31 + AEP"))
+	UPROPERTY(GlobalConfig, EditAnywhere, Category = Build, meta = (DisplayName = "Support OpenGL ES31 + AEP (restart may be needed - see tooltip)"))
 	bool bBuildForES31;
 
 	// If selected, the checked architectures will be split into separate .apk files [CURRENTLY FOR FULL SOURCE GAMES ONLY]
diff --git a/Engine/Source/Runtime/Launch/Private/Android/AndroidJNI.cpp b/Engine/Source/Runtime/Launch/Private/Android/AndroidJNI.cpp
--- a/Engine/Source/Runtime/Launch/Private/Android/AndroidJNI.cpp
+++ b/Engine/Source/Runtime/Launch/Private/Android/AndroidJNI.cpp
@@ -510,10 +510,6 @@ JNIEXPORT jint JNI_OnLoad(JavaVM* InJavaVM, void* InReserved)
 	Env->ReleaseStringUTFChars(pathString, nativePathString);
 	FPlatformMisc::LowLevelOutputDebugStringf(TEXT("Path found as '%s'\n"), *GFilePathBase);
 
-	// Next we check to see if the OBB file is in the APK
-	jmethodID isOBBInAPKMethod = Env->GetStaticMethodID(FJavaWrapper::GameActivityClassID, "isOBBInAPK", "()Z");
-	GOBBinAPK = (bool)Env->CallStaticBooleanMethod(FJavaWrapper::GameActivityClassID, isOBBInAPKMethod, nullptr);
-
 	// Get the system font directory
 	jstring fontPath = (jstring)Env->CallStaticObjectMethod(FJavaWrapper::GameActivityClassID, FJavaWrapper::AndroidThunkJava_GetFontDirectory);
 	const char * nativeFontPathString = Env->GetStringUTFChars(fontPath, 0);
@@ -550,6 +546,10 @@ extern "C" void Java_com_epicgames_ue4_GameActivity_nativeSetGlobalActivity(JNIE
 		// @todo split GooglePlay, this needs to be passed in to this function
 		FJavaWrapper::GoogleServicesThis = FJavaWrapper::GameActivityThis;
 		// FJavaWrapper::GoogleServicesThis = jenv->NewGlobalRef(googleServices);
+
+		// Next we check to see if the OBB file is in the APK
+		jmethodID isOBBInAPKMethod = jenv->GetStaticMethodID(FJavaWrapper::GameActivityClassID, "isOBBInAPK", "()Z");
+		GOBBinAPK = (bool)jenv->CallStaticBooleanMethod(FJavaWrapper::GameActivityClassID, isOBBInAPKMethod, nullptr);
 	}
 }
 
