Simplified dylibs setting in Mac toolchain, so we use rpaths for almost everything.

diff --git a/Engine/Source/Programs/UnrealBuildTool/Mac/MacToolChain.cs b/Engine/Source/Programs/UnrealBuildTool/Mac/MacToolChain.cs
--- a/Engine/Source/Programs/UnrealBuildTool/Mac/MacToolChain.cs
+++ b/Engine/Source/Programs/UnrealBuildTool/Mac/MacToolChain.cs
@@ -581,8 +581,7 @@ namespace UnrealBuildTool
 		{
 			string LibraryDir = Path.GetDirectoryName(Library);
 			string ExeDir = Path.GetDirectoryName(ExeAbsolutePath);
-			//string AbsPath = Path.GetFullPath(Library);
-			if ((Library.Contains("/Plugins/") || Library.Contains("/Binaries/ThirdParty/")) && (Library.EndsWith("dylib") || Library.EndsWith(".framework")) && LibraryDir != ExeDir)
+			if (!Library.Contains("/Engine/Binaries/Mac/") && (Library.EndsWith("dylib") || Library.EndsWith(".framework")) && LibraryDir != ExeDir)
 			{
 				string RelativePath = Utils.MakePathRelativeTo(LibraryDir, ExeDir).Replace("\\", "/");
 				if (!RelativePath.Contains(LibraryDir) && !RPaths.Contains(RelativePath))
@@ -642,10 +641,6 @@ namespace UnrealBuildTool
 			string DylibsPath = "@rpath";
 
 			string AbsolutePath = OutputFile.AbsolutePath.Replace("\\", "/");
-			if (!AbsolutePath.Contains("/Engine/Binaries/Mac/"))
-			{
-				DylibsPath = AbsolutePath.Contains("/Plugins/") ? "@rpath" : "@loader_path";
-			}
 			if (!bIsBuildingLibrary)
 			{
 				LinkCommand += " -rpath @loader_path/ -rpath @executable_path/";
