UBT: fix compiling for Linux (UE-8505).

diff --git a/Engine/Source/Programs/UnrealBuildTool/Configuration/UEBuildTarget.cs b/Engine/Source/Programs/UnrealBuildTool/Configuration/UEBuildTarget.cs
--- a/Engine/Source/Programs/UnrealBuildTool/Configuration/UEBuildTarget.cs
+++ b/Engine/Source/Programs/UnrealBuildTool/Configuration/UEBuildTarget.cs
@@ -2385,7 +2385,8 @@ namespace UnrealBuildTool
 			{
 				Prefix = "lib";
 			}
-			else if (LocalConfig == UnrealTargetConfiguration.Development || bForceNameAsForDevelopment)
+
+			if (LocalConfig == UnrealTargetConfiguration.Development || bForceNameAsForDevelopment)
 			{
 				OutBinaryPath = Path.Combine(BaseDirectory, String.Format("{3}{0}{1}{2}", BinaryName, BinarySuffix, BinaryExtension, Prefix));
 			}
