Comment out unused android permission (DISABLE_KEYGUARD)
#android
#ue4

diff --git a/Engine/Source/Programs/UnrealBuildTool/Android/UEDeployAndroid.cs b/Engine/Source/Programs/UnrealBuildTool/Android/UEDeployAndroid.cs
--- a/Engine/Source/Programs/UnrealBuildTool/Android/UEDeployAndroid.cs
+++ b/Engine/Source/Programs/UnrealBuildTool/Android/UEDeployAndroid.cs
@@ -698,7 +698,7 @@ namespace UnrealBuildTool.Android
 			Text.AppendLine("\t<uses-permission android:name=\"android.permission.MODIFY_AUDIO_SETTINGS\"/>");
 			Text.AppendLine("\t<uses-permission android:name=\"android.permission.GET_ACCOUNTS\"/>");
 			Text.AppendLine("\t<uses-permission android:name=\"com.android.vending.BILLING\"/>");
-			Text.AppendLine("\t<uses-permission android:name=\"android.permission.DISABLE_KEYGUARD\"/>");
+//			Text.AppendLine("\t<uses-permission android:name=\"android.permission.DISABLE_KEYGUARD\"/>");
 			if (ExtraPermissions != null)
 			{
 				foreach (string Permission in ExtraPermissions)
