- Fixing Mac env var settings in UBT to always read from .ini [UE-8542]

diff --git a/Engine/Source/Programs/UnrealBuildTool/Android/UEBuildAndroid.cs b/Engine/Source/Programs/UnrealBuildTool/Android/UEBuildAndroid.cs
--- a/Engine/Source/Programs/UnrealBuildTool/Android/UEBuildAndroid.cs
+++ b/Engine/Source/Programs/UnrealBuildTool/Android/UEBuildAndroid.cs
@@ -57,10 +57,12 @@ namespace UnrealBuildTool
         private bool HasAnySDK()
         {
             string NDKPath = Environment.GetEnvironmentVariable("NDKROOT");
-			bool bNeedsNDKPath = string.IsNullOrEmpty(NDKPath);
-			bool bNeedsAndroidHome = string.IsNullOrEmpty(Environment.GetEnvironmentVariable("ANDROID_HOME"));
-			bool bNeedsAntHome = string.IsNullOrEmpty(Environment.GetEnvironmentVariable("ANT_HOME"));
-            if((bNeedsNDKPath || bNeedsAndroidHome || bNeedsAntHome))
+/* Don't check for existence of env vars, always set them from the .ini values if they exist
+//			bool bNeedsNDKPath = string.IsNullOrEmpty(NDKPath);
+//			bool bNeedsAndroidHome = string.IsNullOrEmpty(Environment.GetEnvironmentVariable("ANDROID_HOME"));
+//			bool bNeedsAntHome = string.IsNullOrEmpty(Environment.GetEnvironmentVariable("ANT_HOME"));
+ //           if((bNeedsNDKPath || bNeedsAndroidHome || bNeedsAntHome))
+ */
             {
                 var configCacheIni = new ConfigCacheIni("Engine", null);
                 var AndroidEnv = new Dictionary<string, string>();
