Fix broken formatting.

diff --git a/Engine/Source/Programs/UnrealBuildTool/Windows/VCToolChain.cs b/Engine/Source/Programs/UnrealBuildTool/Windows/VCToolChain.cs
--- a/Engine/Source/Programs/UnrealBuildTool/Windows/VCToolChain.cs
+++ b/Engine/Source/Programs/UnrealBuildTool/Windows/VCToolChain.cs
@@ -68,7 +68,7 @@ namespace UnrealBuildTool
 				String.Append( " /I " + IncludePath );
 			}
 			else
-		{
+			{
 				String.Append( " -I" + IncludePath );
 			}
 		}
@@ -83,7 +83,7 @@ namespace UnrealBuildTool
 			}
 			else
 			{
-			// Arguments.Append( " /showIncludes" );
+				// Arguments.Append( " /showIncludes" );
 			}
 
 			if( WindowsPlatform.bCompileWithClang )
@@ -141,10 +141,10 @@ namespace UnrealBuildTool
 
 			if( WindowsPlatform.bUseVCCompilerArgs )
 			{ 
-			// Prevents the compiler from displaying its logo for each invocation.
-			Arguments.Append(" /nologo");
+				// Prevents the compiler from displaying its logo for each invocation.
+				Arguments.Append(" /nologo");
 
-			// Enable intrinsic functions.
+				// Enable intrinsic functions.
 				Arguments.Append(" /Oi");	// @todo clang: No Clang equivalent to this?
 			}
 
@@ -174,29 +174,29 @@ namespace UnrealBuildTool
 
 			if( WindowsPlatform.bUseVCCompilerArgs )
 			{
-			// Separate functions for linker.
-			Arguments.Append(" /Gy");
+				// Separate functions for linker.
+				Arguments.Append(" /Gy");
 
-			// Allow 800% of the default memory allocation limit.
-			Arguments.Append(" /Zm800");
+				// Allow 800% of the default memory allocation limit.
+				Arguments.Append(" /Zm800");
 
-			// Disable "The file contains a character that cannot be represented in the current code page" warning for non-US windows.
-			Arguments.Append(" /wd4819");
+				// Disable "The file contains a character that cannot be represented in the current code page" warning for non-US windows.
+				Arguments.Append(" /wd4819");
 			}
 
 			if( BuildConfiguration.bUseSharedPCHs )
 			{
 				if( WindowsPlatform.bUseVCCompilerArgs )
 				{ 
-				// @todo SharedPCH: Disable warning about PCH defines not matching .cpp defines.  We "cheat" these defines a little
-				// bit to make shared PCHs work.  But it's totally safe.  Trust us.
-				Arguments.Append(" /wd4651");
+					// @todo SharedPCH: Disable warning about PCH defines not matching .cpp defines.  We "cheat" these defines a little
+					// bit to make shared PCHs work.  But it's totally safe.  Trust us.
+					Arguments.Append(" /wd4651");
 
-				// @todo SharedPCH: Disable warning about redefining *API macros.  The PCH header is compiled with various DLLIMPORTs, but
-				// when a module that uses that PCH header *IS* one of those imports, that module is compiled with EXPORTS, so the macro
-				// is redefined on the command-line.  We need to clobber those defines to make shared PCHs work properly!
-				Arguments.Append(" /wd4005");
-			}
+					// @todo SharedPCH: Disable warning about redefining *API macros.  The PCH header is compiled with various DLLIMPORTs, but
+					// when a module that uses that PCH header *IS* one of those imports, that module is compiled with EXPORTS, so the macro
+					// is redefined on the command-line.  We need to clobber those defines to make shared PCHs work properly!
+					Arguments.Append(" /wd4005");
+				}
 			}
 
 			// If compiling as a DLL, set the relevant defines
@@ -207,7 +207,7 @@ namespace UnrealBuildTool
 
 			// When targeting Windows XP with Visual Studio 2012+, we need to tell the compiler to use the older Windows SDK that works
 			// with Windows XP (http://blogs.msdn.com/b/vcblog/archive/2012/10/08/10357555.aspx)
-			if (WindowsPlatform.IsWindowsXPSupported())
+			if( WindowsPlatform.SupportWindowsXP )
 			{
 				AddDefinition( ref Arguments, "_USING_V110_SDK71_");
 			}
@@ -229,29 +229,29 @@ namespace UnrealBuildTool
 			{
 				if( WindowsPlatform.bUseVCCompilerArgs )
 				{ 
-				// Disable compiler optimization.
-				Arguments.Append(" /Od");
+					// Disable compiler optimization.
+					Arguments.Append(" /Od");
 
-				// Favor code size (especially useful for embedded platforms).
-				Arguments.Append(" /Os");
+					// Favor code size (especially useful for embedded platforms).
+					Arguments.Append(" /Os");
 
-				// Allow inline method expansion unless E&C support is requested
-				if( !BuildConfiguration.bSupportEditAndContinue )
-				{
-					Arguments.Append(" /Ob2");
-				}
+					// Allow inline method expansion unless E&C support is requested
+					if( !BuildConfiguration.bSupportEditAndContinue )
+					{
+						Arguments.Append(" /Ob2");
+					}
 
-				if ((CompileEnvironment.Config.Target.Platform == CPPTargetPlatform.Win32) || 
-					(CompileEnvironment.Config.Target.Platform == CPPTargetPlatform.Win64))
-				{
-					// Runtime stack checks are not allowed when compiling for CLR
-					if (CompileEnvironment.Config.CLRMode == CPPCLRMode.CLRDisabled)
+					if ((CompileEnvironment.Config.Target.Platform == CPPTargetPlatform.Win32) || 
+						(CompileEnvironment.Config.Target.Platform == CPPTargetPlatform.Win64))
 					{
-						Arguments.Append(" /RTCs");
+						// Runtime stack checks are not allowed when compiling for CLR
+						if (CompileEnvironment.Config.CLRMode == CPPCLRMode.CLRDisabled)
+						{
+							Arguments.Append(" /RTCs");
+						}
 					}
 				}
 			}
-			}
 			//
 			//	Development and LTCG
 			//
@@ -259,38 +259,38 @@ namespace UnrealBuildTool
 			{
 				if( WindowsPlatform.bUseVCCompilerArgs )
 				{ 
-				// Maximum optimizations if desired.
-				if( CompileEnvironment.Config.OptimizeCode >= ModuleRules.CodeOptimization.InNonDebugBuilds )
-				{
-					Arguments.Append(" /Ox");
-				}
+					// Maximum optimizations if desired.
+					if( CompileEnvironment.Config.OptimizeCode >= ModuleRules.CodeOptimization.InNonDebugBuilds )
+					{
+						Arguments.Append(" /Ox");
+					}
 				
-				// Favor code speed.
-				Arguments.Append(" /Ot");
+					// Favor code speed.
+					Arguments.Append(" /Ot");
 
-				// Only omit frame pointers on the PC (which is implied by /Ox) if wanted.
-				if ( BuildConfiguration.bOmitFramePointers == false
-				&& ((CompileEnvironment.Config.Target.Platform == CPPTargetPlatform.Win32) ||
-					(CompileEnvironment.Config.Target.Platform == CPPTargetPlatform.Win64)))
-				{
-					Arguments.Append(" /Oy-");
-				}
+					// Only omit frame pointers on the PC (which is implied by /Ox) if wanted.
+					if ( BuildConfiguration.bOmitFramePointers == false
+					&& ((CompileEnvironment.Config.Target.Platform == CPPTargetPlatform.Win32) ||
+						(CompileEnvironment.Config.Target.Platform == CPPTargetPlatform.Win64)))
+					{
+						Arguments.Append(" /Oy-");
+					}
 
-				// Allow inline method expansion
-				Arguments.Append(" /Ob2");
+					// Allow inline method expansion
+					Arguments.Append(" /Ob2");
 
-				//
-				// LTCG
-				//
-				if (CompileEnvironment.Config.Target.Configuration == CPPTargetConfiguration.Shipping)
-				{
-					if( BuildConfiguration.bAllowLTCG )
+					//
+					// LTCG
+					//
+					if (CompileEnvironment.Config.Target.Configuration == CPPTargetConfiguration.Shipping)
 					{
-						// Enable link-time code generation.
-						Arguments.Append(" /GL");
+						if( BuildConfiguration.bAllowLTCG )
+						{
+							// Enable link-time code generation.
+							Arguments.Append(" /GL");
+						}
 					}
 				}
-			}
 				else
 				{
 					// Maximum optimizations if desired.
@@ -314,15 +314,15 @@ namespace UnrealBuildTool
 				{
 					if( WindowsPlatform.bUseVCCompilerArgs )
 					{ 
-					// Allow the compiler to generate SSE2 instructions.
-					Arguments.Append(" /arch:SSE2");
-				}
+						// Allow the compiler to generate SSE2 instructions.
+						Arguments.Append(" /arch:SSE2");
+					}
 				}
 
 				if( WindowsPlatform.bUseVCCompilerArgs )
 				{ 
-				// Prompt the user before reporting internal errors to Microsoft.
-				Arguments.Append(" /errorReport:prompt");
+					// Prompt the user before reporting internal errors to Microsoft.
+					Arguments.Append(" /errorReport:prompt");
 				}
 
 				if (CompileEnvironment.Config.CLRMode == CPPCLRMode.CLRDisabled)
@@ -361,36 +361,36 @@ namespace UnrealBuildTool
 			{
 				if( WindowsPlatform.bUseVCCompilerArgs )
 				{ 
-				// Store debug info in .pdb files.
-				// @todo clang: PDB files are emited from Clang but do not fully work with Visual Studio yet (breakpoints won't hit due to "symbol read error")
-				if( BuildConfiguration.bUsePDBFiles )
-				{
-					// Create debug info suitable for E&C if wanted.
-					if( BuildConfiguration.bSupportEditAndContinue &&
-						// We only need to do this in debug as that's the only configuration that supports E&C.
-						CompileEnvironment.Config.Target.Configuration == CPPTargetConfiguration.Debug)
+					// Store debug info in .pdb files.
+					// @todo clang: PDB files are emited from Clang but do not fully work with Visual Studio yet (breakpoints won't hit due to "symbol read error")
+					if( BuildConfiguration.bUsePDBFiles )
 					{
-						Arguments.Append(" /ZI");
+						// Create debug info suitable for E&C if wanted.
+						if( BuildConfiguration.bSupportEditAndContinue &&
+							// We only need to do this in debug as that's the only configuration that supports E&C.
+							CompileEnvironment.Config.Target.Configuration == CPPTargetConfiguration.Debug)
+						{
+							Arguments.Append(" /ZI");
+						}
+						// Regular PDB debug information.
+						else
+						{
+							Arguments.Append(" /Zi");
+						}
+						// We need to add this so VS won't lock the PDB file and prevent synchronous updates. This forces serialization through MSPDBSRV.exe.
+						// See http://msdn.microsoft.com/en-us/library/dn502518.aspx for deeper discussion of /FS switch.
+						if (BuildConfiguration.bUseIncrementalLinking && WindowsPlatform.Compiler == WindowsCompiler.VisualStudio2013)
+						{
+							Arguments.Append(" /FS");
+						}
 					}
-					// Regular PDB debug information.
+					// Store C7-format debug info in the .obj files, which is faster.
 					else
 					{
-						Arguments.Append(" /Zi");
+						Arguments.Append(" /Z7");
 					}
-                    // We need to add this so VS won't lock the PDB file and prevent synchronous updates. This forces serialization through MSPDBSRV.exe.
-                    // See http://msdn.microsoft.com/en-us/library/dn502518.aspx for deeper discussion of /FS switch.
-                    if (BuildConfiguration.bUseIncrementalLinking && WindowsPlatform.Compiler == WindowsCompiler.VisualStudio2013)
-                    {
-                        Arguments.Append(" /FS");
-                    }
-				}
-				// Store C7-format debug info in the .obj files, which is faster.
-				else
-				{
-					Arguments.Append(" /Z7");
 				}
 			}
-			}
 
 			// Specify the appropriate runtime library based on the platform and config.
 			if( CompileEnvironment.Config.bUseStaticCRT )
@@ -399,10 +399,10 @@ namespace UnrealBuildTool
 				{
 					if( WindowsPlatform.bUseVCCompilerArgs )
 					{ 
-					Arguments.Append(" /MTd");
-				}
-				else
-				{
+						Arguments.Append(" /MTd");
+					}
+					else
+					{
 						AddDefinition( ref Arguments, "_MT" );
 						AddDefinition( ref Arguments, "_DEBUG" );
 					}
@@ -411,8 +411,8 @@ namespace UnrealBuildTool
 				{
 					if( WindowsPlatform.bUseVCCompilerArgs )
 					{ 
-					Arguments.Append(" /MT");
-				}
+						Arguments.Append(" /MT");
+					}
 					else
 					{
 						AddDefinition( ref Arguments, "_MT" );
@@ -425,10 +425,10 @@ namespace UnrealBuildTool
 				{
 					if( WindowsPlatform.bUseVCCompilerArgs )
 					{ 
-					Arguments.Append(" /MDd");
-				}
-				else
-				{
+						Arguments.Append(" /MDd");
+					}
+					else
+					{
 						AddDefinition( ref Arguments, "_MT" );
 						AddDefinition( ref Arguments, "_DEBUG" );
 						AddDefinition( ref Arguments, "_DLL" );
@@ -438,8 +438,8 @@ namespace UnrealBuildTool
 				{
 					if( WindowsPlatform.bUseVCCompilerArgs )
 					{ 
-					Arguments.Append(" /MD");
-				}
+						Arguments.Append(" /MD");
+					}
 					else
 					{
 						AddDefinition( ref Arguments, "_MT" );
@@ -452,12 +452,12 @@ namespace UnrealBuildTool
 			{
 				// @todo clang: Not supported in clang-cl yet
 				if( !WindowsPlatform.bCompileWithClang )
-			{
-				// Allow large object files to avoid hitting the 2^16 section limit when running with -StressTestUnity.
-				Arguments.Append(" /bigobj");
+				{
+					// Allow large object files to avoid hitting the 2^16 section limit when running with -StressTestUnity.
+					Arguments.Append(" /bigobj");
 
-				// Relaxes floating point precision semantics to allow more optimization.
-				Arguments.Append(" /fp:fast");
+					// Relaxes floating point precision semantics to allow more optimization.
+					Arguments.Append(" /fp:fast");
 				}
 
 				if (CompileEnvironment.Config.OptimizeCode >= ModuleRules.CodeOptimization.InNonDebugBuilds)
@@ -486,8 +486,8 @@ namespace UnrealBuildTool
 					// Pack struct members on 8-byte boundaries.
 					if( WindowsPlatform.bUseVCCompilerArgs )
 					{ 
-					Arguments.Append(" /Zp8");
-				}
+						Arguments.Append(" /Zp8");
+					}
 					else
 					{
 						Arguments.Append(" -fpack-struct=8" );
@@ -498,8 +498,8 @@ namespace UnrealBuildTool
 					// Pack struct members on 4-byte boundaries.
 					if( WindowsPlatform.bUseVCCompilerArgs )
 					{ 
-					Arguments.Append(" /Zp4");
-				}
+						Arguments.Append(" /Zp4");
+					}
 					else
 					{
 						Arguments.Append(" -fpack-struct=4" );
@@ -520,8 +520,8 @@ namespace UnrealBuildTool
 		{
 			if( WindowsPlatform.bUseVCCompilerArgs )
 			{
-			// Explicitly compile the file as C++.
-			Arguments.Append(" /TP");
+				// Explicitly compile the file as C++.
+				Arguments.Append(" /TP");
 			}
 			else
 			{
@@ -541,11 +541,11 @@ namespace UnrealBuildTool
 			{
 				if( WindowsPlatform.bUseVCCompilerArgs )
 				{ 
-				// This will disable buffer security checks (which are enabled by default) that the MS compiler adds around arrays on the stack,
-				// Which can add some performance overhead, especially in performance intensive code
-				// Only disable this if you know what you are doing, because it will be disabled for the entire module!
-				Arguments.Append(" /GS-");
-			}
+					// This will disable buffer security checks (which are enabled by default) that the MS compiler adds around arrays on the stack,
+					// Which can add some performance overhead, especially in performance intensive code
+					// Only disable this if you know what you are doing, because it will be disabled for the entire module!
+					Arguments.Append(" /GS-");
+				}
 			}
 
 			// C++/CLI requires that RTTI is left enabled
@@ -556,10 +556,10 @@ namespace UnrealBuildTool
 					// Enable C++ RTTI.
 					if( WindowsPlatform.bUseVCCompilerArgs )
 					{ 
-					Arguments.Append(" /GR");
-				}
-				else
-				{
+						Arguments.Append(" /GR");
+					}
+					else
+					{
 						Arguments.Append(" -Xclang -frtti-data" );
 					}
 				}
@@ -568,8 +568,8 @@ namespace UnrealBuildTool
 					// Disable C++ RTTI.
 					if( WindowsPlatform.bUseVCCompilerArgs )
 					{ 
-					Arguments.Append(" /GR-");
-				}
+						Arguments.Append(" /GR-");
+					}
 					else
 					{
 						Arguments.Append(" -Xclang -fno-rtti-data" );
@@ -580,16 +580,16 @@ namespace UnrealBuildTool
 			// Set warning level.
 			if( WindowsPlatform.bUseVCCompilerArgs )
 			{ 
-			if (!BuildConfiguration.bRunUnrealCodeAnalyzer)
-			{
-				// Restrictive during regular compilation.
-				Arguments.Append(" /W4");
-			}
-			else
-			{
-				// If we had /W4 with clang on windows we would be flooded with warnings. This will be fixed incrementally.
-				Arguments.Append(" /W0");
-			}
+				if (!BuildConfiguration.bRunUnrealCodeAnalyzer)
+				{
+					// Restrictive during regular compilation.
+					Arguments.Append(" /W4");
+				}
+				else
+				{
+					// If we had /W4 with clang on windows we would be flooded with warnings. This will be fixed incrementally.
+					Arguments.Append(" /W0");
+				}
 			}
 			else
 			{
@@ -648,15 +648,15 @@ namespace UnrealBuildTool
 			}
 			else
 			{ 
-			// Explicitly compile the file as C.
-			Arguments.Append(" /TC");
+				// Explicitly compile the file as C.
+				Arguments.Append(" /TC");
 			}
 		
 			if( WindowsPlatform.bUseVCCompilerArgs )
 			{ 
-			// Level 0 warnings.  Needed for external C projects that produce warnings at higher warning levels.
-			Arguments.Append(" /W0");
-		}
+				// Level 0 warnings.  Needed for external C projects that produce warnings at higher warning levels.
+				Arguments.Append(" /W0");
+			}
 		}
 
 		static void AppendLinkArguments(LinkEnvironment LinkEnvironment, StringBuilder Arguments)
@@ -722,7 +722,7 @@ namespace UnrealBuildTool
 
 					// When targeting Windows XP in Visual Studio 2012+, we need to tell the linker we are going to support execution
 					// on that older platform.  The compiler defaults to version 6.0+.  We'll modify the SUBSYSTEM parameter here.
-					if (WindowsPlatform.IsWindowsXPSupported())
+					if( WindowsPlatform.SupportWindowsXP )
 					{
 						Arguments.Append(LinkEnvironment.Config.Target.Platform == CPPTargetPlatform.Win64 ? ",5.02" : ",5.01");
 					}
@@ -886,7 +886,7 @@ namespace UnrealBuildTool
 
 					// When targeting Windows XP in Visual Studio 2012+, we need to tell the linker we are going to support execution
 					// on that older platform.  The compiler defaults to version 6.0+.  We'll modify the SUBSYSTEM parameter here.
-					if (WindowsPlatform.IsWindowsXPSupported())
+					if( WindowsPlatform.SupportWindowsXP )
 					{
 						Arguments.Append(LinkEnvironment.Config.Target.Platform == CPPTargetPlatform.Win64 ? ",5.02" : ",5.01");
 					}
@@ -1314,7 +1314,7 @@ namespace UnrealBuildTool
 
 				// When targeting Windows XP with Visual Studio 2012+, we need to tell the compiler to use the older Windows SDK that works
 				// with Windows XP (http://blogs.msdn.com/b/vcblog/archive/2012/10/08/10357555.aspx)
-				if (WindowsPlatform.IsWindowsXPSupported())
+				if (WindowsPlatform.SupportWindowsXP)
 				{
 					AddDefinition( ref Arguments, "_USING_V110_SDK71_" );
 				}
