[INTEGRATE] Change 2420404 by [REDACTED_PERSON][REDACTED_HANDLE] on 2015/01/27 10:50:14

	Fixed UnrealBuildTool Clean not deleting all project-specific intermediate files
	- Files in a project's target-specific intermediate sub-folder were never cleaned!
	- (e.g. "MyProject/Intermediate/Build/<Platform>/MyProjectEditor/<Configuration>/*")

diff --git a/Engine/Source/Programs/UnrealBuildTool/Configuration/UEBuildTarget.cs b/Engine/Source/Programs/UnrealBuildTool/Configuration/UEBuildTarget.cs
--- a/Engine/Source/Programs/UnrealBuildTool/Configuration/UEBuildTarget.cs
+++ b/Engine/Source/Programs/UnrealBuildTool/Configuration/UEBuildTarget.cs
@@ -970,23 +970,11 @@ namespace UnrealBuildTool
 		protected void CleanTarget(List<UEBuildBinary> Binaries, CPPTargetPlatform Platform, BuildManifest Manifest)
 		{
 			{
-				var TargetFilename = RulesCompiler.GetTargetFilename(GameName);
 				var LocalTargetName = (TargetType == TargetRules.TargetType.Program) ? AppName : GameName;
-
 				Log.TraceVerbose("Cleaning target {0} - AppName {1}", LocalTargetName, AppName);
-				Log.TraceVerbose("\tTargetFilename {0}", TargetFilename);
 
-				var TargetFolder = "";
-				if (String.IsNullOrEmpty(TargetFilename) == false)
-				{
-					var TargetInfo = new FileInfo(TargetFilename);
-					TargetFolder = TargetInfo.Directory.FullName;
-					var SourceIdx = TargetFolder.LastIndexOf("\\Source");
-					if (SourceIdx != -1)
-					{
-						TargetFolder = TargetFolder.Substring(0, SourceIdx + 1);
-					}
-				}
+				var TargetFilename = RulesCompiler.GetTargetFilename(GameName);
+				Log.TraceVerbose("\tTargetFilename {0}", TargetFilename);
 
 				// Collect all files to delete.
 				var AdditionalFileExtensions = new string[] { ".lib", ".exp", ".dll.response" };
@@ -1063,6 +1051,7 @@ namespace UnrealBuildTool
 					}
 				}
 
+
 				//
 				{
 					var AppEnginePath = Path.Combine(PlatformEngineBuildDataFolder, LocalTargetName, Configuration.ToString());
@@ -1073,6 +1062,13 @@ namespace UnrealBuildTool
 				}
 
 				// Clean the intermediate directory
+				if( !String.IsNullOrEmpty( ProjectIntermediateDirectory ) )
+				{
+					if (Directory.Exists(ProjectIntermediateDirectory))
+					{
+						CleanDirectory(ProjectIntermediateDirectory);
+					}
+				}
 				if (!UnrealBuildTool.RunningRocket())
 				{
 					// This is always under Rocket installation folder
