[INTEGRATE] Change 2420813 by [REDACTED_PERSON][REDACTED_HANDLE] on 2015/01/27 14:28:27

	Fix missing engine folder in Rocket projects, and clarify message about building UE4Editor et al.

diff --git a/Engine/Source/Programs/UnrealBuildTool/System/ProjectFileGenerator.cs b/Engine/Source/Programs/UnrealBuildTool/System/ProjectFileGenerator.cs
--- a/Engine/Source/Programs/UnrealBuildTool/System/ProjectFileGenerator.cs
+++ b/Engine/Source/Programs/UnrealBuildTool/System/ProjectFileGenerator.cs
@@ -741,23 +741,7 @@ namespace UnrealBuildTool
 			}
 
 
-			if( bGeneratingGameProjectFiles )
-			{
-				GameProjectName = Path.GetFileNameWithoutExtension(UnrealBuildTool.GetUProjectFile());
-				if (String.IsNullOrEmpty(GameProjectName))
-				{
-					throw new BuildException("A valid game project was not found in the specified location (" + UnrealBuildTool.GetUProjectPath() + ")");
-				}
-
-				IncludeEngineSource = bAlwaysIncludeEngineModules;
-				IncludeDocumentation = false;
-				IncludeBuildSystemFiles = false;
-				IncludeShaderSource = true;
-				IncludeTemplateFiles = false;
-				IncludeConfigFiles = true;
-				IncludeEnginePrograms = bAlwaysIncludeEngineModules;
-			}
-			else if( bGeneratingRocketProjectFiles )
+			if( bGeneratingRocketProjectFiles )
 			{
 				// We expected a project path to be passed in
 				if (!UnrealBuildTool.HasUProjectFile())
@@ -780,6 +764,22 @@ namespace UnrealBuildTool
 				IncludeConfigFiles = true;
 				IncludeEnginePrograms = false;
 			}
+			else if( bGeneratingGameProjectFiles )
+			{
+				GameProjectName = Path.GetFileNameWithoutExtension(UnrealBuildTool.GetUProjectFile());
+				if (String.IsNullOrEmpty(GameProjectName))
+				{
+					throw new BuildException("A valid game project was not found in the specified location (" + UnrealBuildTool.GetUProjectPath() + ")");
+				}
+
+				IncludeEngineSource = bAlwaysIncludeEngineModules;
+				IncludeDocumentation = false;
+				IncludeBuildSystemFiles = false;
+				IncludeShaderSource = true;
+				IncludeTemplateFiles = false;
+				IncludeConfigFiles = true;
+				IncludeEnginePrograms = bAlwaysIncludeEngineModules;
+			}
 			else
 			{
 				// At least one extra argument was specified, but we weren't expected it.  Ignored.
diff --git a/Engine/Source/Programs/UnrealBuildTool/System/VCProject.cs b/Engine/Source/Programs/UnrealBuildTool/System/VCProject.cs
--- a/Engine/Source/Programs/UnrealBuildTool/System/VCProject.cs
+++ b/Engine/Source/Programs/UnrealBuildTool/System/VCProject.cs
@@ -1016,8 +1016,8 @@ namespace UnrealBuildTool
 					if(UnrealBuildTool.RunningRocket() && Utils.IsFileUnderDirectory(TargetFilePath, EngineRootDirectory))
 					{
 						// Don't build any engine targets for Rocket. We can still start the program through the debugger though...
-						VCProjectFileContent.AppendFormat("\t\t<NMakeBuildCommandLine>@echo {0} is up to date.</NMakeBuildCommandLine>" + ProjectFileGenerator.NewLine, TargetName);
-						VCProjectFileContent.AppendFormat("\t\t<NMakeReBuildCommandLine>@echo {0} is up to date.</NMakeReBuildCommandLine>" + ProjectFileGenerator.NewLine, TargetName);
+						VCProjectFileContent.AppendFormat("\t\t<NMakeBuildCommandLine>@echo Using installed engine. {0} is up to date.</NMakeBuildCommandLine>" + ProjectFileGenerator.NewLine, TargetName);
+						VCProjectFileContent.AppendFormat("\t\t<NMakeReBuildCommandLine>@echo Using installed engine. {0} is up to date.</NMakeReBuildCommandLine>" + ProjectFileGenerator.NewLine, TargetName);
 						VCProjectFileContent.AppendFormat("\t\t<NMakeCleanCommandLine>@echo Ignoring clean for {0}.</NMakeCleanCommandLine>" + ProjectFileGenerator.NewLine, TargetName);
 					}
 					else
