[INTEGRATE] Change 2422587 by [REDACTED_PERSON][REDACTED_HANDLE] on 2015/01/28 14:18:12

	UnrealBuildTool: Minor progress text clean-ups for consistency with editor

diff --git a/Engine/Source/Programs/UnrealBuildTool/System/ExternalExecution.cs b/Engine/Source/Programs/UnrealBuildTool/System/ExternalExecution.cs
--- a/Engine/Source/Programs/UnrealBuildTool/System/ExternalExecution.cs
+++ b/Engine/Source/Programs/UnrealBuildTool/System/ExternalExecution.cs
@@ -531,7 +531,7 @@ namespace UnrealBuildTool
 		 */
 		public static bool ExecuteHeaderToolIfNecessary( UEBuildTarget Target, CPPEnvironment GlobalCompileEnvironment, List<UHTModuleInfo> UObjectModules, string ModuleInfoFileName, ref ECompilationResult UHTResult )
 		{
-			using (ProgressWriter Progress = new ProgressWriter("Generating headers...", false))
+			using (ProgressWriter Progress = new ProgressWriter("Generating code...", false))
 			{
 				// We never want to try to execute the header tool when we're already trying to build it!
 				var bIsBuildingUHT = Target.GetTargetName().Equals( "UnrealHeaderTool", StringComparison.InvariantCultureIgnoreCase );
diff --git a/Engine/Source/Programs/UnrealBuildTool/System/LocalExecutor.cs b/Engine/Source/Programs/UnrealBuildTool/System/LocalExecutor.cs
--- a/Engine/Source/Programs/UnrealBuildTool/System/LocalExecutor.cs
+++ b/Engine/Source/Programs/UnrealBuildTool/System/LocalExecutor.cs
@@ -390,7 +390,7 @@ namespace UnrealBuildTool
 
 			Dictionary<Action, ActionThread> ActionThreadDictionary = new Dictionary<Action, ActionThread>();
             int JobNumber = 1;
-			using(ProgressWriter ProgressWriter = new ProgressWriter("Compiling source code...", false))
+			using(ProgressWriter ProgressWriter = new ProgressWriter("Compiling C++ source code...", false))
 			{
 				int ProgressValue = 0;
 				while(true)
diff --git a/Engine/Source/Programs/UnrealBuildTool/System/XGE.cs b/Engine/Source/Programs/UnrealBuildTool/System/XGE.cs
--- a/Engine/Source/Programs/UnrealBuildTool/System/XGE.cs
+++ b/Engine/Source/Programs/UnrealBuildTool/System/XGE.cs
@@ -478,7 +478,7 @@ namespace UnrealBuildTool
 		 */
 		public static ExecutionResult ExecuteTaskFileWithProgressMarkup(string TaskFilePath, int NumActions, DataReceivedEventHandler OutputEventHandler)
 		{
-			using (ProgressWriter Writer = new ProgressWriter("Compiling source files...", false))
+			using (ProgressWriter Writer = new ProgressWriter("Compiling C++ source files...", false))
 			{
 				int NumCompletedActions = 0;
 
