[INTEGRATE] Change 2423747 by [REDACTED_PERSON][REDACTED_HANDLE] on 2015/01/29 10:14:07

	Move anonymous function for writing pre-default properties into a normal function.

diff --git a/Engine/Source/Programs/UnrealBuildTool/System/VCProject.cs b/Engine/Source/Programs/UnrealBuildTool/System/VCProject.cs
--- a/Engine/Source/Programs/UnrealBuildTool/System/VCProject.cs
+++ b/Engine/Source/Programs/UnrealBuildTool/System/VCProject.cs
@@ -592,51 +592,13 @@ namespace UnrealBuildTool
 					"	</PropertyGroup>" + ProjectFileGenerator.NewLine);
 			}
 
-			// Anonymous function that writes pre-Default.props configuration data
-			System.Action<UnrealTargetConfiguration, UnrealTargetPlatform, string, TargetRules> PreDefaultPropsWriterFunc = ( Configuration, Platform, TargetFilePath, TargetRulesObject ) => 
-				{
-					UEPlatformProjectGenerator ProjGenerator = UEPlatformProjectGenerator.GetPlatformProjectGenerator(Platform, true);
-					if (((ProjGenerator == null) && (Platform != UnrealTargetPlatform.Unknown)))
-					{
-						return;
-					}
-
-					string ProjectConfigName = StubProjectConfigurationName;
-					string ProjectPlatformName = StubProjectPlatformName;
-					MakeProjectPlatformAndConfigurationNames( Platform, Configuration, TargetRulesObject.ConfigurationName, out ProjectPlatformName, out ProjectConfigName );
-					var ConfigAndPlatformName = ProjectConfigName + "|" + ProjectPlatformName;
-					string ConditionString = "Condition=\"'$(Configuration)|$(Platform)'=='" + ConfigAndPlatformName + "'\"";
-
-					string PlatformToolsetString = (ProjGenerator != null) ? ProjGenerator.GetVisualStudioPlatformToolsetString(Platform, Configuration, this) : "";
-					if( String.IsNullOrEmpty( PlatformToolsetString ) )
-					{
-						if( VCProjectFileGenerator.ProjectFileFormat == VCProjectFileGenerator.VCProjectFileFormat.VisualStudio2013 )
-						{
-							PlatformToolsetString = "		<PlatformToolset>v120</PlatformToolset>" + ProjectFileGenerator.NewLine;
-						}
-						else if( VCProjectFileGenerator.ProjectFileFormat == VCProjectFileGenerator.VCProjectFileFormat.VisualStudio2012 )
-						{
-							PlatformToolsetString = "		<PlatformToolset>v110</PlatformToolset>" + ProjectFileGenerator.NewLine;
-						}
-					}
-					
-					string PlatformConfigurationType = ProjGenerator.GetVisualStudioPlatformConfigurationType(Platform);	
-					VCProjectFileContent.Append(
-						"	<PropertyGroup " + ConditionString + " Label=\"Configuration\">" + ProjectFileGenerator.NewLine +
-						"		<ConfigurationType>" + PlatformConfigurationType + "</ConfigurationType>" + ProjectFileGenerator.NewLine +
-								PlatformToolsetString +
-						"	</PropertyGroup>" + ProjectFileGenerator.NewLine 
-						);
-				};
-
-
 			// Write each project configuration PreDefaultProps section
 			foreach( var Combination in ProjectConfigAndTargetCombinations )
 			{
 				// Only write config settings for valid permutations
 				if (IsStubProject || IsValidProjectPlatformAndConfiguration( Combination.ProjectTarget, Combination.Platform, Combination.Configuration ))
 				{
-					PreDefaultPropsWriterFunc( Combination.Configuration, Combination.Platform, Combination.ProjectTarget.TargetFilePath, Combination.ProjectTarget.TargetRules );
+					WritePreDefaultPropsConfiguration( Combination.Configuration, Combination.Platform, Combination.ProjectTarget.TargetFilePath, Combination.ProjectTarget.TargetRules, VCProjectFileContent );
 				}
 			}
 
@@ -884,6 +846,43 @@ namespace UnrealBuildTool
 			}
 		}
 
+		// Anonymous function that writes pre-Default.props configuration data
+		private void WritePreDefaultPropsConfiguration(UnrealTargetConfiguration Configuration, UnrealTargetPlatform Platform, string TargetFilePath, TargetRules TargetRulesObject, StringBuilder VCProjectFileContent)
+		{
+			UEPlatformProjectGenerator ProjGenerator = UEPlatformProjectGenerator.GetPlatformProjectGenerator(Platform, true);
+			if (((ProjGenerator == null) && (Platform != UnrealTargetPlatform.Unknown)))
+			{
+				return;
+			}
+
+			string ProjectConfigName = StubProjectConfigurationName;
+			string ProjectPlatformName = StubProjectPlatformName;
+			MakeProjectPlatformAndConfigurationNames( Platform, Configuration, TargetRulesObject.ConfigurationName, out ProjectPlatformName, out ProjectConfigName );
+			var ConfigAndPlatformName = ProjectConfigName + "|" + ProjectPlatformName;
+			string ConditionString = "Condition=\"'$(Configuration)|$(Platform)'=='" + ConfigAndPlatformName + "'\"";
+
+			string PlatformToolsetString = (ProjGenerator != null) ? ProjGenerator.GetVisualStudioPlatformToolsetString(Platform, Configuration, this) : "";
+			if( String.IsNullOrEmpty( PlatformToolsetString ) )
+			{
+				if( VCProjectFileGenerator.ProjectFileFormat == VCProjectFileGenerator.VCProjectFileFormat.VisualStudio2013 )
+				{
+					PlatformToolsetString = "		<PlatformToolset>v120</PlatformToolset>" + ProjectFileGenerator.NewLine;
+				}
+				else if( VCProjectFileGenerator.ProjectFileFormat == VCProjectFileGenerator.VCProjectFileFormat.VisualStudio2012 )
+				{
+					PlatformToolsetString = "		<PlatformToolset>v110</PlatformToolset>" + ProjectFileGenerator.NewLine;
+				}
+			}
+					
+			string PlatformConfigurationType = ProjGenerator.GetVisualStudioPlatformConfigurationType(Platform);	
+			VCProjectFileContent.Append(
+				"	<PropertyGroup " + ConditionString + " Label=\"Configuration\">" + ProjectFileGenerator.NewLine +
+				"		<ConfigurationType>" + PlatformConfigurationType + "</ConfigurationType>" + ProjectFileGenerator.NewLine +
+						PlatformToolsetString +
+				"	</PropertyGroup>" + ProjectFileGenerator.NewLine 
+				);
+		}
+
 		// Anonymous function that writes project configuration data
 		private void WriteConfiguration(string ProjectName, UnrealTargetConfiguration Configuration, UnrealTargetPlatform Platform, string TargetFilePath, TargetRules TargetRulesObject, StringBuilder VCProjectFileContent, StringBuilder VCUserFileContent)
 		{
