[INTEGRATE] Change 2424081 by [REDACTED_PERSON][REDACTED_HANDLE] on 2015/01/29 13:29:21

	Fix for generating engine project files under Rocket, when we don't have any target files. Generate a stub project instead.

diff --git a/Engine/Source/Programs/UnrealBuildTool/System/ProjectFileGenerator.cs b/Engine/Source/Programs/UnrealBuildTool/System/ProjectFileGenerator.cs
--- a/Engine/Source/Programs/UnrealBuildTool/System/ProjectFileGenerator.cs
+++ b/Engine/Source/Programs/UnrealBuildTool/System/ProjectFileGenerator.cs
@@ -390,6 +390,19 @@ namespace UnrealBuildTool
 				// Place projects into root level solution folders
 				if( IncludeEngineSource )
 				{
+					// If we're still missing an engine project because we don't have any targets for it, make one up.
+					if( EngineProject == null )
+					{
+						string ProjectFilePath = Path.Combine(IntermediateProjectFilesPath, "UE4" + ProjectFileExtension);
+
+						bool bAlreadyExisted;
+						EngineProject = FindOrAddProject(Utils.MakePathRelativeTo(ProjectFilePath, MasterProjectRelativePath), true, out bAlreadyExisted);
+
+						EngineProject.IsForeignProject = false;
+						EngineProject.IsGeneratedProject = true;
+						EngineProject.IsStubProject = true;
+					}
+
 					if( EngineProject != null )
 					{
 						RootFolder.AddSubFolder( "Engine" ).ChildProjects.Add( EngineProject );
diff --git a/Engine/Source/Programs/UnrealBuildTool/System/VCProject.cs b/Engine/Source/Programs/UnrealBuildTool/System/VCProject.cs
--- a/Engine/Source/Programs/UnrealBuildTool/System/VCProject.cs
+++ b/Engine/Source/Programs/UnrealBuildTool/System/VCProject.cs
@@ -301,12 +301,28 @@ namespace UnrealBuildTool
 		}
 
 
-		struct ProjectConfigAndTargetCombination
+		class ProjectConfigAndTargetCombination
 		{
 			public UnrealTargetPlatform Platform;
 			public UnrealTargetConfiguration Configuration;
+			public string ProjectPlatformName;
+			public string ProjectConfigurationName;
 			public ProjectTarget ProjectTarget;
 
+			public ProjectConfigAndTargetCombination(UnrealTargetPlatform InPlatform, UnrealTargetConfiguration InConfiguration, string InProjectPlatformName, string InProjectConfigurationName, ProjectTarget InProjectTarget)
+			{
+				Platform = InPlatform;
+				Configuration = InConfiguration;
+				ProjectPlatformName = InProjectPlatformName;
+				ProjectConfigurationName = InProjectConfigurationName;
+				ProjectTarget = InProjectTarget;
+			}
+
+			public string ProjectConfigurationAndPlatformName
+			{
+				get { return (ProjectPlatformName == null)? null : (ProjectConfigurationName + "|" + ProjectPlatformName); }
+			}
+
 			public override string ToString()
 			{
 				return String.Format("{0} {1} {2}", ProjectTarget, Platform, Configuration);
@@ -396,7 +412,8 @@ namespace UnrealBuildTool
 			// If this is a "stub" project, then only add a single configuration to the project
 			if( IsStubProject )
 			{
-				ProjectConfigAndTargetCombinations.Add( new ProjectConfigAndTargetCombination() { Platform = UnrealTargetPlatform.Unknown, Configuration = UnrealTargetConfiguration.Unknown, ProjectTarget = null } );
+				ProjectConfigAndTargetCombination StubCombination = new ProjectConfigAndTargetCombination(UnrealTargetPlatform.Unknown, UnrealTargetConfiguration.Unknown, StubProjectPlatformName, StubProjectConfigurationName, null);
+				ProjectConfigAndTargetCombinations.Add(StubCombination);
 			}
 			else
 			{
@@ -425,12 +442,14 @@ namespace UnrealBuildTool
 
 							foreach( var ProjectTarget in ProjectTargets )
 							{
-								ProjectConfigAndTargetCombinations.Add( new ProjectConfigAndTargetCombination() 
-									{ 
-										Platform = Platform, 
-										Configuration = Configuration, 
-										ProjectTarget = ProjectTarget
-									} );
+								if(IsValidProjectPlatformAndConfiguration( ProjectTarget, Platform, Configuration ))
+								{
+									string ProjectPlatformName, ProjectConfigurationName;
+									MakeProjectPlatformAndConfigurationNames(Platform, Configuration, ProjectTarget.TargetRules.ConfigurationName, out ProjectPlatformName, out ProjectConfigurationName);
+
+									ProjectConfigAndTargetCombination Combination = new ProjectConfigAndTargetCombination(Platform, Configuration, ProjectPlatformName, ProjectConfigurationName, ProjectTarget);
+									ProjectConfigAndTargetCombinations.Add( Combination );
+								}
 							}
 						}
 					}
@@ -446,30 +465,17 @@ namespace UnrealBuildTool
 			var ProjectConfigNames = new List<string>();
 			foreach( var Combination in ProjectConfigAndTargetCombinations )
 			{
-				var ProjectPlatform = Combination.Platform;
-				var ProjectConfig = Combination.Configuration;
-				var TargetRulesObject = Combination.ProjectTarget.TargetRules;
-
-				// Only write config settings for valid permutations
-				if (IsStubProject || IsValidProjectPlatformAndConfiguration( Combination.ProjectTarget, Combination.Platform, Combination.Configuration ))
+				if( !ProjectPlatforms.Contains( Combination.Platform ) )
 				{
-					string ProjectPlatformName;
-					string ProjectConfigName;
-					MakeProjectPlatformAndConfigurationNames(ProjectPlatform, ProjectConfig, TargetRulesObject.ConfigurationName, out ProjectPlatformName,
-															 out ProjectConfigName);
-
-					if( !ProjectPlatforms.Contains( Combination.Platform ) )
-					{
-						ProjectPlatforms.Add( Combination.Platform );
-					}
-					if( !ProjectPlatformNames.Contains( ProjectPlatformName ) )
-					{
-						ProjectPlatformNames.Add(ProjectPlatformName);
-					}
-					if( !ProjectConfigNames.Contains( ProjectConfigName ) )
-					{
-						ProjectConfigNames.Add(ProjectConfigName);
-					}
+					ProjectPlatforms.Add( Combination.Platform );
+				}
+				if( !ProjectPlatformNames.Contains( Combination.ProjectPlatformName ) )
+				{
+					ProjectPlatformNames.Add(Combination.ProjectPlatformName);
+				}
+				if( !ProjectConfigNames.Contains( Combination.ProjectConfigurationName ) )
+				{
+					ProjectConfigNames.Add(Combination.ProjectConfigurationName);
 				}
 			}
 
@@ -595,11 +601,7 @@ namespace UnrealBuildTool
 			// Write each project configuration PreDefaultProps section
 			foreach( var Combination in ProjectConfigAndTargetCombinations )
 			{
-				// Only write config settings for valid permutations
-				if (IsStubProject || IsValidProjectPlatformAndConfiguration( Combination.ProjectTarget, Combination.Platform, Combination.Configuration ))
-				{
-					WritePreDefaultPropsConfiguration( Combination.Configuration, Combination.Platform, Combination.ProjectTarget.TargetFilePath, Combination.ProjectTarget.TargetRules, VCProjectFileContent );
-				}
+				WritePreDefaultPropsConfiguration( Combination, VCProjectFileContent );
 			}
 
 			VCProjectFileContent.Append(
@@ -612,11 +614,7 @@ namespace UnrealBuildTool
 			// Write each project configuration
 			foreach( var Combination in ProjectConfigAndTargetCombinations )
 			{
-				// Only write config settings for valid permutations
-				if (IsStubProject || IsValidProjectPlatformAndConfiguration( Combination.ProjectTarget, Combination.Platform, Combination.Configuration ))
-				{
-					WriteConfiguration( ProjectName, Combination.Configuration, Combination.Platform, Combination.ProjectTarget.TargetFilePath, Combination.ProjectTarget.TargetRules, VCProjectFileContent, bGenerateUserFileContent? VCUserFileContent : null );
-				}
+				WriteConfiguration( ProjectName, Combination, VCProjectFileContent, bGenerateUserFileContent? VCUserFileContent : null );
 			}
 
 			// Write IntelliSense info
@@ -663,16 +661,12 @@ namespace UnrealBuildTool
 						string HackTargetFilePath = null;
 						foreach( var Combination in ProjectConfigAndTargetCombinations )
 						{
-							// Only write config settings for valid permutations
-							if (IsStubProject || IsValidProjectPlatformAndConfiguration( Combination.ProjectTarget, Combination.Platform, Combination.Configuration ))
+							if( Combination.Platform == Platform &&
+								Combination.ProjectTarget.TargetRules != null && 
+								Combination.ProjectTarget.TargetRules.Type == HackTargetType )
 							{
-								if( Combination.Platform == Platform &&
-									Combination.ProjectTarget.TargetRules != null && 
-									Combination.ProjectTarget.TargetRules.Type == HackTargetType )
-								{
-									HackTargetFilePath = Combination.ProjectTarget.TargetFilePath;// ProjectConfigAndTargetCombinations[0].ProjectTarget.TargetFilePath;
-									break;										
-								}
+								HackTargetFilePath = Combination.ProjectTarget.TargetFilePath;// ProjectConfigAndTargetCombinations[0].ProjectTarget.TargetFilePath;
+								break;										
 							}
 						}
 
@@ -847,21 +841,17 @@ namespace UnrealBuildTool
 		}
 
 		// Anonymous function that writes pre-Default.props configuration data
-		private void WritePreDefaultPropsConfiguration(UnrealTargetConfiguration Configuration, UnrealTargetPlatform Platform, string TargetFilePath, TargetRules TargetRulesObject, StringBuilder VCProjectFileContent)
+		private void WritePreDefaultPropsConfiguration(ProjectConfigAndTargetCombination Combination, StringBuilder VCProjectFileContent)
 		{
-			UEPlatformProjectGenerator ProjGenerator = UEPlatformProjectGenerator.GetPlatformProjectGenerator(Platform, true);
-			if (((ProjGenerator == null) && (Platform != UnrealTargetPlatform.Unknown)))
+			UEPlatformProjectGenerator ProjGenerator = UEPlatformProjectGenerator.GetPlatformProjectGenerator(Combination.Platform, true);
+			if (((ProjGenerator == null) && (Combination.Platform != UnrealTargetPlatform.Unknown)))
 			{
 				return;
 			}
 
-			string ProjectConfigName = StubProjectConfigurationName;
-			string ProjectPlatformName = StubProjectPlatformName;
-			MakeProjectPlatformAndConfigurationNames( Platform, Configuration, TargetRulesObject.ConfigurationName, out ProjectPlatformName, out ProjectConfigName );
-			var ConfigAndPlatformName = ProjectConfigName + "|" + ProjectPlatformName;
-			string ConditionString = "Condition=\"'$(Configuration)|$(Platform)'=='" + ConfigAndPlatformName + "'\"";
+			string ConditionString = "Condition=\"'$(Configuration)|$(Platform)'=='" + Combination.ProjectConfigurationAndPlatformName + "'\"";
 
-			string PlatformToolsetString = (ProjGenerator != null) ? ProjGenerator.GetVisualStudioPlatformToolsetString(Platform, Configuration, this) : "";
+			string PlatformToolsetString = (ProjGenerator != null) ? ProjGenerator.GetVisualStudioPlatformToolsetString(Combination.Platform, Combination.Configuration, this) : "";
 			if( String.IsNullOrEmpty( PlatformToolsetString ) )
 			{
 				if( VCProjectFileGenerator.ProjectFileFormat == VCProjectFileGenerator.VCProjectFileFormat.VisualStudio2013 )
@@ -873,8 +863,8 @@ namespace UnrealBuildTool
 					PlatformToolsetString = "		<PlatformToolset>v110</PlatformToolset>" + ProjectFileGenerator.NewLine;
 				}
 			}
-					
-			string PlatformConfigurationType = ProjGenerator.GetVisualStudioPlatformConfigurationType(Platform);	
+
+			string PlatformConfigurationType = (ProjGenerator == null)? "Makefile" : ProjGenerator.GetVisualStudioPlatformConfigurationType(Combination.Platform);	
 			VCProjectFileContent.Append(
 				"	<PropertyGroup " + ConditionString + " Label=\"Configuration\">" + ProjectFileGenerator.NewLine +
 				"		<ConfigurationType>" + PlatformConfigurationType + "</ConfigurationType>" + ProjectFileGenerator.NewLine +
@@ -884,8 +874,11 @@ namespace UnrealBuildTool
 		}
 
 		// Anonymous function that writes project configuration data
-		private void WriteConfiguration(string ProjectName, UnrealTargetConfiguration Configuration, UnrealTargetPlatform Platform, string TargetFilePath, TargetRules TargetRulesObject, StringBuilder VCProjectFileContent, StringBuilder VCUserFileContent)
+		private void WriteConfiguration(string ProjectName, ProjectConfigAndTargetCombination Combination, StringBuilder VCProjectFileContent, StringBuilder VCUserFileContent)
 		{
+			UnrealTargetPlatform Platform = Combination.Platform;
+			UnrealTargetConfiguration Configuration = Combination.Configuration;
+
 			UEPlatformProjectGenerator ProjGenerator = UEPlatformProjectGenerator.GetPlatformProjectGenerator(Platform, true);
 			if (((ProjGenerator == null) && (Platform != UnrealTargetPlatform.Unknown)))
 			{
@@ -898,11 +891,7 @@ namespace UnrealBuildTool
 				UProjectPath = "\"$(SolutionDir)$(ProjectName).uproject\"";
 			}
 
-			string ProjectConfigName = StubProjectConfigurationName;
-			string ProjectPlatformName = StubProjectPlatformName;
-			MakeProjectPlatformAndConfigurationNames(Platform, Configuration, TargetRulesObject.ConfigurationName, out ProjectPlatformName, out ProjectConfigName);
-			var ConfigAndPlatformName = ProjectConfigName + "|" + ProjectPlatformName;
-			string ConditionString = "Condition=\"'$(Configuration)|$(Platform)'=='" + ConfigAndPlatformName + "'\"";
+			string ConditionString = "Condition=\"'$(Configuration)|$(Platform)'=='" + Combination.ProjectConfigurationAndPlatformName + "'\"";
 
 			{
 				VCProjectFileContent.Append(
@@ -918,15 +907,17 @@ namespace UnrealBuildTool
 					string ProjectRelativeUnusedDirectory = NormalizeProjectPath(Path.Combine(ProjectFileGenerator.EngineRelativePath, BuildConfiguration.BaseIntermediateFolder, "Unused"));
 
 					VCProjectFileContent.Append(
-						"		<OutDir></OutDir>" + ProjectFileGenerator.NewLine +
+						"		<OutDir>" + ProjectRelativeUnusedDirectory + Path.DirectorySeparatorChar + "</OutDir>" + ProjectFileGenerator.NewLine +
 						"		<IntDir>" + ProjectRelativeUnusedDirectory + Path.DirectorySeparatorChar + "</IntDir>" + ProjectFileGenerator.NewLine +
-						"		<NMakeBuildCommandLine/>" + ProjectFileGenerator.NewLine +
-						"		<NMakeReBuildCommandLine/>" + ProjectFileGenerator.NewLine +
-						"		<NMakeCleanCommandLine/>" + ProjectFileGenerator.NewLine +
+						"		<NMakeBuildCommandLine>@rem Nothing to do.</NMakeBuildCommandLine>" + ProjectFileGenerator.NewLine +
+						"		<NMakeReBuildCommandLine>@rem Nothing to do.</NMakeReBuildCommandLine>" + ProjectFileGenerator.NewLine +
+						"		<NMakeCleanCommandLine>@rem Nothing to do.</NMakeCleanCommandLine>" + ProjectFileGenerator.NewLine +
 						"		<NMakeOutput/>" + ProjectFileGenerator.NewLine);
 				}
 				else
 				{
+					TargetRules TargetRulesObject = Combination.ProjectTarget.TargetRules;
+					string TargetFilePath = Combination.ProjectTarget.TargetFilePath;
 					string TargetName = Utils.GetFilenameWithoutAnyExtensions(TargetFilePath);
 					var UBTPlatformName = IsStubProject ? StubProjectPlatformName : Platform.ToString();
 					var UBTConfigurationName = IsStubProject ? StubProjectConfigurationName : Configuration.ToString();
@@ -1012,49 +1003,39 @@ namespace UnrealBuildTool
 						UEPlatformProjectGenerator.GenerateGamePlatformSpecificProperties(Platform, Configuration, TargetRulesObject.Type, VCProjectFileContent, RootDirectory, TargetFilePath);
 					}
 
-					if(UnrealBuildTool.RunningRocket() && Utils.IsFileUnderDirectory(TargetFilePath, EngineRootDirectory))
-					{
-						// Don't build any engine targets for Rocket. We can still start the program through the debugger though...
-						VCProjectFileContent.AppendFormat("\t\t<NMakeBuildCommandLine>@echo Using installed engine. {0} is up to date.</NMakeBuildCommandLine>" + ProjectFileGenerator.NewLine, TargetName);
-						VCProjectFileContent.AppendFormat("\t\t<NMakeReBuildCommandLine>@echo Using installed engine. {0} is up to date.</NMakeReBuildCommandLine>" + ProjectFileGenerator.NewLine, TargetName);
-						VCProjectFileContent.AppendFormat("\t\t<NMakeCleanCommandLine>@echo Ignoring clean for {0}.</NMakeCleanCommandLine>" + ProjectFileGenerator.NewLine, TargetName);
-					}
-					else
-					{
-						// This is the standard UE4 based project NMake build line:
-						//	..\..\Build\BatchFiles\Build.bat <TARGETNAME> <PLATFORM> <CONFIGURATION>
-						//	ie ..\..\Build\BatchFiles\Build.bat BlankProgram Win64 Debug
+					// This is the standard UE4 based project NMake build line:
+					//	..\..\Build\BatchFiles\Build.bat <TARGETNAME> <PLATFORM> <CONFIGURATION>
+					//	ie ..\..\Build\BatchFiles\Build.bat BlankProgram Win64 Debug
 
-						string ProjectPlatformConfiguration = " " + TargetName + " " + UBTPlatformName + " " + UBTConfigurationName;
-						string BatchFilesDirectoryName = Path.Combine(ProjectFileGenerator.EngineRelativePath, "Build", "BatchFiles");
+					string ProjectPlatformConfiguration = " " + TargetName + " " + UBTPlatformName + " " + UBTConfigurationName;
+					string BatchFilesDirectoryName = Path.Combine(ProjectFileGenerator.EngineRelativePath, "Build", "BatchFiles");
 
-						// NMake Build command line
-						VCProjectFileContent.Append("		<NMakeBuildCommandLine>");
-						VCProjectFileContent.Append(EscapePath(NormalizeProjectPath(Path.Combine(BatchFilesDirectoryName, "Build.bat"))) + ProjectPlatformConfiguration.ToString());
-						if (IsForeignProject)
-						{
-							VCProjectFileContent.Append(" " + UProjectPath + (UnrealBuildTool.RunningRocket() ? " -rocket" : ""));
-						}
-						VCProjectFileContent.Append("</NMakeBuildCommandLine>" + ProjectFileGenerator.NewLine);
+					// NMake Build command line
+					VCProjectFileContent.Append("		<NMakeBuildCommandLine>");
+					VCProjectFileContent.Append(EscapePath(NormalizeProjectPath(Path.Combine(BatchFilesDirectoryName, "Build.bat"))) + ProjectPlatformConfiguration.ToString());
+					if (IsForeignProject)
+					{
+						VCProjectFileContent.Append(" " + UProjectPath + (UnrealBuildTool.RunningRocket() ? " -rocket" : ""));
+					}
+					VCProjectFileContent.Append("</NMakeBuildCommandLine>" + ProjectFileGenerator.NewLine);
 
-						// NMake ReBuild command line
-						VCProjectFileContent.Append("		<NMakeReBuildCommandLine>");
-						VCProjectFileContent.Append(EscapePath(NormalizeProjectPath(Path.Combine(BatchFilesDirectoryName, "Rebuild.bat"))) + ProjectPlatformConfiguration.ToString());
-						if (IsForeignProject)
-						{
-							VCProjectFileContent.Append(" " + UProjectPath + (UnrealBuildTool.RunningRocket() ? " -rocket" : ""));
-						}
-						VCProjectFileContent.Append("</NMakeReBuildCommandLine>" + ProjectFileGenerator.NewLine);
+					// NMake ReBuild command line
+					VCProjectFileContent.Append("		<NMakeReBuildCommandLine>");
+					VCProjectFileContent.Append(EscapePath(NormalizeProjectPath(Path.Combine(BatchFilesDirectoryName, "Rebuild.bat"))) + ProjectPlatformConfiguration.ToString());
+					if (IsForeignProject)
+					{
+						VCProjectFileContent.Append(" " + UProjectPath + (UnrealBuildTool.RunningRocket() ? " -rocket" : ""));
+					}
+					VCProjectFileContent.Append("</NMakeReBuildCommandLine>" + ProjectFileGenerator.NewLine);
 
-						// NMake Clean command line
-						VCProjectFileContent.Append("		<NMakeCleanCommandLine>");
-						VCProjectFileContent.Append(EscapePath(NormalizeProjectPath(Path.Combine(BatchFilesDirectoryName, "Clean.bat"))) + ProjectPlatformConfiguration.ToString());
-						if (IsForeignProject)
-						{
-							VCProjectFileContent.Append(" " + UProjectPath + (UnrealBuildTool.RunningRocket() ? " -rocket" : ""));
-						}
-						VCProjectFileContent.Append("</NMakeCleanCommandLine>" + ProjectFileGenerator.NewLine);
+					// NMake Clean command line
+					VCProjectFileContent.Append("		<NMakeCleanCommandLine>");
+					VCProjectFileContent.Append(EscapePath(NormalizeProjectPath(Path.Combine(BatchFilesDirectoryName, "Clean.bat"))) + ProjectPlatformConfiguration.ToString());
+					if (IsForeignProject)
+					{
+						VCProjectFileContent.Append(" " + UProjectPath + (UnrealBuildTool.RunningRocket() ? " -rocket" : ""));
 					}
+					VCProjectFileContent.Append("</NMakeCleanCommandLine>" + ProjectFileGenerator.NewLine);
 
 					VCProjectFileContent.Append("		<NMakeOutput>");
 					VCProjectFileContent.Append(NormalizeProjectPath(NMakePath));
@@ -1062,8 +1043,10 @@ namespace UnrealBuildTool
 				}
 				VCProjectFileContent.Append("	</PropertyGroup>" + ProjectFileGenerator.NewLine);
 
-				if (VCUserFileContent != null)
+				if (VCUserFileContent != null && Combination.ProjectTarget != null)
 				{
+					TargetRules TargetRulesObject = Combination.ProjectTarget.TargetRules;
+
 					if ((Platform == UnrealTargetPlatform.Win32) || (Platform == UnrealTargetPlatform.Win64))
 					{
 						VCUserFileContent.Append(
@@ -1102,12 +1085,12 @@ namespace UnrealBuildTool
 							);
 					}
 
-					string PlatformUserFileStrings = (ProjGenerator != null) ? ProjGenerator.GetVisualStudioUserFileStrings(Platform, Configuration, ConditionString, TargetRulesObject, TargetFilePath, ProjectFilePath) : "";
+					string PlatformUserFileStrings = (ProjGenerator != null) ? ProjGenerator.GetVisualStudioUserFileStrings(Platform, Configuration, ConditionString, TargetRulesObject, Combination.ProjectTarget.TargetFilePath, ProjectFilePath) : "";
 					VCUserFileContent.Append(PlatformUserFileStrings);
 				}
 			}
 
-			string LayoutDirString = (ProjGenerator != null) ? ProjGenerator.GetVisualStudioLayoutDirSection(Platform, Configuration, ConditionString, TargetRulesObject.Type, TargetFilePath, ProjectFilePath) : "";
+			string LayoutDirString = (ProjGenerator != null) ? ProjGenerator.GetVisualStudioLayoutDirSection(Platform, Configuration, ConditionString, Combination.ProjectTarget.TargetRules.Type, Combination.ProjectTarget.TargetFilePath, ProjectFilePath) : "";
 			VCProjectFileContent.Append(LayoutDirString);
 		}
 	}
