- Added android packaging option for generating a GearVR specific AndroidManifest.xml

diff --git a/Engine/Source/Programs/UnrealBuildTool/Android/UEDeployAndroid.cs b/Engine/Source/Programs/UnrealBuildTool/Android/UEDeployAndroid.cs
--- a/Engine/Source/Programs/UnrealBuildTool/Android/UEDeployAndroid.cs
+++ b/Engine/Source/Programs/UnrealBuildTool/Android/UEDeployAndroid.cs
@@ -628,6 +628,8 @@ namespace UnrealBuildTool.Android
 			Ini.GetString("/Script/AndroidRuntimeSettings.AndroidRuntimeSettings", "ExtraApplicationSettings", out ExtraApplicationSettings);
 			List<string> ExtraPermissions;
 			Ini.GetArray("/Script/AndroidRuntimeSettings.AndroidRuntimeSettings", "ExtraPermissions", out ExtraPermissions);
+			bool bPackageForGearVR;
+			Ini.GetBool("/Script/AndroidRuntimeSettings.AndroidRuntimeSettings", "bPackageForGearVR", out bPackageForGearVR);
 
 			// replace some variables
 			PackageName = PackageName.Replace("[PROJECT]", ProjectName);
@@ -646,15 +648,29 @@ namespace UnrealBuildTool.Android
 			Text.AppendLine("\t             android:hasCode=\"true\">");
 			Text.AppendLine("\t\t<activity android:name=\"com.epicgames.ue4.GameActivity\"");
 			Text.AppendLine("\t\t          android:label=\"@string/app_name\"");
-			Text.AppendLine("\t\t          android:theme=\"@android:style/Theme.NoTitleBar.Fullscreen\"");
+			if (!bPackageForGearVR)
+			{
+				// normal application settings
+				Text.AppendLine("\t\t          android:theme=\"@android:style/Theme.NoTitleBar.Fullscreen\"");
+				Text.AppendLine("\t\t          android:configChanges=\"orientation|keyboardHidden\"");
+			}
+			else
+			{
+				// GearVR
+				Text.AppendLine("\t\t          android:theme=\"@android:style/Theme.Black.NoTitleBar.Fullscreen\"");
+				Text.AppendLine("\t\t          android:configChanges=\"screenSize|orientation|keyboardHidden|keyboard\"");
+				if (bIsForDistribution)
+				{
+					Text.AppendLine("\t\t          android:excludeFromRecents=\"true\"");
+				}
+			}
 			Text.AppendLine("\t\t          android:launchMode=\"singleTask\"");
 			Text.AppendLine(string.Format("\t\t          android:screenOrientation=\"{0}\"", ConvertOrientationIniValue(Orientation)));
-			Text.AppendLine(string.Format("\t\t          android:debuggable=\"{0}\"", bIsForDistribution ? "false" : "true"));
-			Text.AppendLine("\t\t          android:configChanges=\"orientation|keyboardHidden\">");
+			Text.AppendLine(string.Format("\t\t          android:debuggable=\"{0}\">", bIsForDistribution ? "false" : "true"));
 			Text.AppendLine("\t\t\t<meta-data android:name=\"android.app.lib_name\" android:value=\"UE4\"/>");
 			Text.AppendLine("\t\t\t<intent-filter>");
 			Text.AppendLine("\t\t\t\t<action android:name=\"android.intent.action.MAIN\" />");
-			Text.AppendLine("\t\t\t\t<category android:name=\"android.intent.category.LAUNCHER\" />");
+			Text.AppendLine(string.Format("\t\t\t\t<category android:name=\"android.intent.category.{0}\" />", (bPackageForGearVR && bIsForDistribution) ? "INFO" : "LAUNCHER"));
 			Text.AppendLine("\t\t\t</intent-filter>");
 			if (!string.IsNullOrEmpty(ExtraActivitySettings))
 			{
@@ -673,6 +689,16 @@ namespace UnrealBuildTool.Android
 			Text.AppendLine("\t\t           android:value=\"@integer/google_play_services_version\" />");
 			Text.AppendLine("\t\t<activity android:name=\"com.google.android.gms.ads.AdActivity\"");
 			Text.AppendLine("\t\t          android:configChanges=\"keyboard|keyboardHidden|orientation|screenLayout|uiMode|screenSize|smallestScreenSize\"/>");
+			if (bPackageForGearVR)
+			{
+				Text.AppendLine("\t\t<meta-data android:name=\"com.samsung.android.vr.application.mode\"");
+				Text.AppendLine("\t\t           android:value=\"vr_only\" />");
+				Text.AppendLine("\t\t<activity android:name=\"com.oculusvr.vrlib.PlatformActivity\"");
+				Text.AppendLine("\t\t          android:theme=\"@android:style/Theme.Black.NoTitleBar.Fullscreen\"");
+				Text.AppendLine("\t\t          android:launchMode=\"singleTask\"");
+				Text.AppendLine("\t\t          android:screenOrientation=\"landscape\"");
+				Text.AppendLine("\t\t          android:configChanges=\"screenSize|orientation|keyboardHidden|keyboard\"/>");
+			}
 			if (!string.IsNullOrEmpty(ExtraApplicationSettings))
 			{
 				ExtraApplicationSettings = ExtraApplicationSettings.Replace("\\n", "\n");
@@ -699,6 +725,12 @@ namespace UnrealBuildTool.Android
 			Text.AppendLine("\t<uses-permission android:name=\"android.permission.GET_ACCOUNTS\"/>");
 			Text.AppendLine("\t<uses-permission android:name=\"com.android.vending.BILLING\"/>");
 //			Text.AppendLine("\t<uses-permission android:name=\"android.permission.DISABLE_KEYGUARD\"/>");
+			if (bPackageForGearVR)
+			{
+				Text.AppendLine("\t<uses-permission android:name=\"android.permission.CAMERA\"/>");
+				Text.AppendLine("\t<uses-feature android:name=\"android.hardware.camera\"/>");
+				Text.AppendLine("\t<uses-feature android:name=\"android.hardware.usb.host\"/>");
+			}
 			if (ExtraPermissions != null)
 			{
 				foreach (string Permission in ExtraPermissions)
diff --git a/Engine/Source/Runtime/Android/AndroidRuntimeSettings/Classes/AndroidRuntimeSettings.h b/Engine/Source/Runtime/Android/AndroidRuntimeSettings/Classes/AndroidRuntimeSettings.h
--- a/Engine/Source/Runtime/Android/AndroidRuntimeSettings/Classes/AndroidRuntimeSettings.h
+++ b/Engine/Source/Runtime/Android/AndroidRuntimeSettings/Classes/AndroidRuntimeSettings.h
@@ -143,7 +143,11 @@ public:
 	// Any extra permissions your app needs
 	UPROPERTY(GlobalConfig, EditAnywhere, Category = AdvancedAPKPackaging, Meta = (DisplayName = "Extra Permissions (e.g. 'android.permission.INTERNET')"))
 	TArray<FString> ExtraPermissions;
-	
+
+	// Configure AndroidManifest.xml for GearVR
+	UPROPERTY(GlobalConfig, EditAnywhere, Category = AdvancedAPKPackaging, Meta = (DisplayName = "Configure the AndroidManifest for deployment to GearVR"))
+	bool bPackageForGearVR;
+
 	// This is the file that keytool outputs, specified with the -keystore parameter (file should be in <Project>/Build/Android)
 	UPROPERTY(GlobalConfig, EditAnywhere, Category = DistributionSigning, Meta = (DisplayName = "Key Store (output of keytool, placed in <Project>/Build/Android)"))
 	FString KeyStore;
