PR #776: Enable aligned memory access for Linux ARM builds.

- Contributed by [REDACTED_PERSON].

diff --git a/Engine/Source/Programs/UnrealBuildTool/Linux/UEBuildLinux.cs b/Engine/Source/Programs/UnrealBuildTool/Linux/UEBuildLinux.cs
--- a/Engine/Source/Programs/UnrealBuildTool/Linux/UEBuildLinux.cs
+++ b/Engine/Source/Programs/UnrealBuildTool/Linux/UEBuildLinux.cs
@@ -348,6 +348,11 @@ namespace UnrealBuildTool
 
             InBuildTarget.GlobalCompileEnvironment.Config.Definitions.Add("WITH_DATABASE_SUPPORT=0");		//@todo linux: valid?
 
+            if (GetActiveArchitecture().StartsWith("arm"))
+            {
+                InBuildTarget.GlobalCompileEnvironment.Config.Definitions.Add("REQUIRES_ALIGNED_INT_ACCESS");
+            }
+
             // link with Linux libraries.
             InBuildTarget.GlobalLinkEnvironment.Config.AdditionalLibraries.Add("pthread");
 
diff --git a/Engine/Source/Runtime/Core/Public/Linux/LinuxPlatform.h b/Engine/Source/Runtime/Core/Public/Linux/LinuxPlatform.h
--- a/Engine/Source/Runtime/Core/Public/Linux/LinuxPlatform.h
+++ b/Engine/Source/Runtime/Core/Public/Linux/LinuxPlatform.h
@@ -76,6 +76,9 @@ typedef FLinuxPlatformTypes FPlatformTypes;
 // Alignment.
 #define GCC_PACK(n)			__attribute__((packed,aligned(n)))
 #define GCC_ALIGN(n)		__attribute__((aligned(n)))
+#if defined(__arm__)
+	#define REQUIRES_ALIGNED_ACCESS					1
+#endif
 
 // operator new/delete operators
 // As of 10.9 we need to use _NOEXCEPT & cxx_noexcept compatible definitions
