UBT: Enabling PrivateIncludePaths in arbitrary project sub-directories

Merging using UE4-To-UE4-LauncherDev, CL# 2429973

diff --git a/Engine/Source/Programs/UnrealBuildTool/System/RulesCompiler.cs b/Engine/Source/Programs/UnrealBuildTool/System/RulesCompiler.cs
--- a/Engine/Source/Programs/UnrealBuildTool/System/RulesCompiler.cs
+++ b/Engine/Source/Programs/UnrealBuildTool/System/RulesCompiler.cs
@@ -1506,15 +1506,14 @@ namespace UnrealBuildTool
 				//		../<Game>/Source/<Module>
 				// or
 				//		c:/PATH/<Game>/Source/<Module>
-				string SourceSlashModule = Path.Combine("Source", InModuleName);
-				SourceSlashModule = SourceSlashModule.Replace("\\", "/");
-				Int32 SourceSlashIdx = ModuleDirectoryRelativeToEngineSourceDirectory.IndexOf(SourceSlashModule);
+                string SourceDirName = "Source";
+                Int32 SourceSlashIdx = ModuleDirectoryRelativeToEngineSourceDirectory.IndexOf(SourceDirName);
 				string SourceDirectoryPath = null;
 				if(SourceSlashIdx != -1)
 				{
 					try
 					{
-						SourceDirectoryPath = Path.GetFullPath(ModuleDirectoryRelativeToEngineSourceDirectory.Substring(0, SourceSlashIdx + "Source".Length));
+                        SourceDirectoryPath = Path.GetFullPath(ModuleDirectoryRelativeToEngineSourceDirectory.Substring(0, SourceSlashIdx + SourceDirName.Length));
 					}
 					catch(Exception Exc)
 					{
@@ -1530,7 +1529,6 @@ namespace UnrealBuildTool
 				if(SourceDirectoryPath != null)
 				{
 					List<string> ResolvedPrivatePaths = new List<string>();
-					ResolvedPrivatePaths.Add(SourceDirectoryPath);
 					foreach(var PrivatePath in RulesObject.PrivateIncludePaths)
 					{
 						try
