UE-3580

XP support was not able to be set from Bootstrap packaged game. Which meant the subsystem version wasnt valid on xp for this.

diff --git a/Engine/Source/Programs/UnrealBuildTool/Windows/UEBuildWindows.cs b/Engine/Source/Programs/UnrealBuildTool/Windows/UEBuildWindows.cs
--- a/Engine/Source/Programs/UnrealBuildTool/Windows/UEBuildWindows.cs
+++ b/Engine/Source/Programs/UnrealBuildTool/Windows/UEBuildWindows.cs
@@ -440,6 +440,8 @@ namespace UnrealBuildTool
         {
             InBuildTarget.GlobalCompileEnvironment.Config.Definitions.Add("WIN32=1");
 
+			// Win32 XP is only supported at this time.
+			SupportWindowsXP = SupportWindowsXPIfAvailable && (GetCPPTargetPlatform(InBuildTarget.Platform) == CPPTargetPlatform.Win32);
 			if (IsWindowsXPSupported())
             {
                 // Windows XP SP3 or higher required
