UBT: Adding the ability to reset the current project file to fix command chaining in UAT

diff --git a/Engine/Source/Programs/UnrealBuildTool/System/UnrealBuildTool.cs b/Engine/Source/Programs/UnrealBuildTool/System/UnrealBuildTool.cs
--- a/Engine/Source/Programs/UnrealBuildTool/System/UnrealBuildTool.cs
+++ b/Engine/Source/Programs/UnrealBuildTool/System/UnrealBuildTool.cs
@@ -506,10 +506,10 @@ namespace UnrealBuildTool
 
         public static void SetProjectFile(string InProjectFile)
         {
-            if (HasUProjectFile() == false)
+            if (HasUProjectFile() == false || String.IsNullOrEmpty(InProjectFile))
             {
                 UProjectFile = InProjectFile;
-                if (UProjectFile.Length > 0)
+				if (!String.IsNullOrEmpty(UProjectFile))
                 {
                     bHasUProjectFile = true;
                     UProjectPath = Path.GetDirectoryName(UProjectFile);
@@ -517,6 +517,7 @@ namespace UnrealBuildTool
                 else
                 {
                     bHasUProjectFile = false;
+					UProjectPath = String.Empty;
                 }
             }
             else
@@ -2419,6 +2420,8 @@ namespace UnrealBuildTool
 		/// <param name="UProjectDir"></param>
 		public static void SetupUBTFromUAT(string UProjectFile)
 		{
+			// Reset project file
+			SetProjectFile(String.Empty);
 			// when running UAT, the working directory is the root UE4 dir
 			BuildConfiguration.RelativeEnginePath = "Engine";
 			SetProjectFile(UProjectFile);
