- Fixing OBB In APK with Launch On [UE-9150]

diff --git a/Engine/Source/Programs/AutomationTool/Android/AndroidPlatform.Automation.cs b/Engine/Source/Programs/AutomationTool/Android/AndroidPlatform.Automation.cs
--- a/Engine/Source/Programs/AutomationTool/Android/AndroidPlatform.Automation.cs
+++ b/Engine/Source/Programs/AutomationTool/Android/AndroidPlatform.Automation.cs
@@ -189,7 +189,7 @@ public class AndroidPlatform : Platform
 		string[] Architectures = UnrealBuildTool.AndroidToolChain.GetAllArchitectures();
 		string[] GPUArchitectures = UnrealBuildTool.AndroidToolChain.GetAllGPUArchitectures();
 		bool bMakeSeparateApks = UnrealBuildTool.Android.UEDeployAndroid.ShouldMakeSeparateApks();
-		bool bPackageDataInsideApk = UnrealBuildTool.Android.UEDeployAndroid.PackageDataInsideApk();
+		bool bPackageDataInsideApk = UnrealBuildTool.Android.UEDeployAndroid.PackageDataInsideApk(false);
 
 		var Deploy = UEBuildDeploy.GetBuildDeploy(UnrealTargetPlatform.Android);
 
@@ -251,7 +251,7 @@ public class AndroidPlatform : Platform
 				{
 					string CookFlavor = SC.FinalCookPlatform.IndexOf("_") > 0 ? SC.FinalCookPlatform.Substring(SC.FinalCookPlatform.IndexOf("_")) : "";
 					string SOName = GetSONameWithoutArchitecture(Params, SC.StageExecutables[0]);
-					Deploy.PrepForUATPackageOrDeploy(Params.ShortProjectName, SC.ProjectRoot, SOName, SC.LocalRoot + "/Engine", Params.Distribution, CookFlavor);
+					Deploy.PrepForUATPackageOrDeploy(Params.ShortProjectName, SC.ProjectRoot, SOName, SC.LocalRoot + "/Engine", Params.Distribution, CookFlavor, false);
 				}
 
 			    // Create APK specific OBB in case we have a detached OBB.
@@ -374,7 +374,7 @@ public class AndroidPlatform : Platform
 		string[] Architectures = UnrealBuildTool.AndroidToolChain.GetAllArchitectures();
 		string[] GPUArchitectures = UnrealBuildTool.AndroidToolChain.GetAllGPUArchitectures();
 		bool bMakeSeparateApks = UnrealBuildTool.Android.UEDeployAndroid.ShouldMakeSeparateApks();
-		bool bPackageDataInsideApk = UnrealBuildTool.Android.UEDeployAndroid.PackageDataInsideApk();
+		bool bPackageDataInsideApk = UnrealBuildTool.Android.UEDeployAndroid.PackageDataInsideApk(false);
 
 		bool bAddedOBB = false;
 		foreach (string Architecture in Architectures)
@@ -527,7 +527,7 @@ public class AndroidPlatform : Platform
 		{
 			string CookFlavor = SC.FinalCookPlatform.IndexOf("_") > 0 ? SC.FinalCookPlatform.Substring(SC.FinalCookPlatform.IndexOf("_")) : "";
 			string SOName = GetSONameWithoutArchitecture(Params, SC.StageExecutables[0]);
-			Deploy.PrepForUATPackageOrDeploy(Params.ShortProjectName, SC.ProjectRoot, SOName, SC.LocalRoot + "/Engine", Params.Distribution, CookFlavor);
+			Deploy.PrepForUATPackageOrDeploy(Params.ShortProjectName, SC.ProjectRoot, SOName, SC.LocalRoot + "/Engine", Params.Distribution, CookFlavor, true);
 		}
 
 		// now we can use the apk to get more info
diff --git a/Engine/Source/Programs/AutomationTool/IOS/IOSPlatform.Automation.cs b/Engine/Source/Programs/AutomationTool/IOS/IOSPlatform.Automation.cs
--- a/Engine/Source/Programs/AutomationTool/IOS/IOSPlatform.Automation.cs
+++ b/Engine/Source/Programs/AutomationTool/IOS/IOSPlatform.Automation.cs
@@ -153,7 +153,9 @@ public class IOSPlatform : Platform
 				Path.GetDirectoryName(Params.RawProjectPath),
 				CombinePaths(Path.GetDirectoryName(Params.ProjectGameExeFilename), SC.StageExecutables[0]),
 				CombinePaths(SC.LocalRoot, "Engine"),
-				Params.Distribution, "");
+				Params.Distribution, 
+				"",
+				false);
 
 			// figure out where to pop in the staged files
 			string AppDirectory = string.Format("{0}/Payload/{1}.app",
@@ -404,7 +406,7 @@ public class IOSPlatform : Platform
 			Path.GetDirectoryName(RawProjectPath),
 			CombinePaths(BaseDirectory, GameName),
 			CombinePaths(LocalRoot, "Engine"),
-			Distribution, "");
+			Distribution, "", false);
 
 		string Arguments = "UBT_NO_POST_DEPLOY=true";
 		Arguments += " /usr/bin/xcrun xcodebuild build -project \"" + XcodeProj + "\"";
diff --git a/Engine/Source/Programs/AutomationTool/ProjectParams.cs b/Engine/Source/Programs/AutomationTool/ProjectParams.cs
--- a/Engine/Source/Programs/AutomationTool/ProjectParams.cs
+++ b/Engine/Source/Programs/AutomationTool/ProjectParams.cs
@@ -1956,7 +1956,8 @@ namespace AutomationTool
 				CommandUtils.Log("IterativeCooking={0}", IterativeCooking);
                 CommandUtils.Log("CookAll={0}", CookAll);
                 CommandUtils.Log("CookMapsOnly={0}", CookMapsOnly);
-                CommandUtils.Log("IterativeDeploy={0}", IterativeDeploy);
+                CommandUtils.Log("Deploy={0}", Deploy);
+				CommandUtils.Log("IterativeDeploy={0}", IterativeDeploy);
 				CommandUtils.Log("LogWindow={0}", LogWindow);
 				CommandUtils.Log("Manifests={0}", Manifests);
 				CommandUtils.Log("MapToRun={0}", MapToRun);
@@ -1967,6 +1968,7 @@ namespace AutomationTool
 				CommandUtils.Log("NoXGE={0}", NoXGE);
 				CommandUtils.Log("MapsToCook={0}", MapsToCook.ToString());
 				CommandUtils.Log("Pak={0}", Pak);
+				CommandUtils.Log("Package={0}", Package);
 				CommandUtils.Log("NullRHI={0}", NullRHI);
 				CommandUtils.Log("FakeClient={0}", FakeClient);
                 CommandUtils.Log("EditorTest={0}", EditorTest);
diff --git a/Engine/Source/Programs/UnrealBuildTool/Android/UEDeployAndroid.cs b/Engine/Source/Programs/UnrealBuildTool/Android/UEDeployAndroid.cs
--- a/Engine/Source/Programs/UnrealBuildTool/Android/UEDeployAndroid.cs
+++ b/Engine/Source/Programs/UnrealBuildTool/Android/UEDeployAndroid.cs
@@ -104,8 +104,13 @@ namespace UnrealBuildTool.Android
 			return CachedSDKLevel;
 		}
 
-		public static bool PackageDataInsideApk(ConfigCacheIni Ini=null)
+		public static bool PackageDataInsideApk(bool bDisallowPackagingDataInApk, ConfigCacheIni Ini=null)
 		{
+			if (bDisallowPackagingDataInApk)
+			{
+				return false;
+			}
+
 			// make a new one if one wasn't passed in
 			if (Ini == null)
 			{
@@ -438,7 +443,7 @@ namespace UnrealBuildTool.Android
 		}
 
 
-		private string GetAllBuildSettings(string BuildPath, bool bForDistribution, bool bMakeSeparateApks)
+		private string GetAllBuildSettings(string BuildPath, bool bForDistribution, bool bMakeSeparateApks, bool bPackageDataInsideApk)
 		{
 			// make the settings string - this will be char by char compared against last time
 			StringBuilder CurrentSettings = new StringBuilder();
@@ -449,6 +454,7 @@ namespace UnrealBuildTool.Android
 			CurrentSettings.AppendLine(string.Format("SDKVersion={0}", GetSdkApiLevel()));
 			CurrentSettings.AppendLine(string.Format("bForDistribution={0}", bForDistribution));
 			CurrentSettings.AppendLine(string.Format("bMakeSeparateApks={0}", bMakeSeparateApks));
+			CurrentSettings.AppendLine(string.Format("bPackageDataInsideApk={0}", bPackageDataInsideApk));
 
 			// all AndroidRuntimeSettings ini settings in here
 			ConfigCacheIni Ini = new ConfigCacheIni(UnrealTargetPlatform.Android, "Engine", UnrealBuildTool.GetUProjectPath());
@@ -493,7 +499,8 @@ namespace UnrealBuildTool.Android
 			return CurrentSettings.ToString();
 		}
 
-		private bool CheckDependencies(string ProjectName, string ProjectDirectory, string UE4BuildFilesPath, string GameBuildFilesPath, string EngineDirectory, string CookFlavor, string OutputPath, string UE4BuildPath, bool bMakeSeparateApks)
+		private bool CheckDependencies(string ProjectName, string ProjectDirectory, string UE4BuildFilesPath, string GameBuildFilesPath, string EngineDirectory, 
+			string CookFlavor, string OutputPath, string UE4BuildPath, bool bMakeSeparateApks, bool bPackageDataInsideApk)
 		{
 			string[] Arches = AndroidToolChain.GetAllArchitectures();
 			string[] GPUArchitectures = AndroidToolChain.GetAllGPUArchitectures();
@@ -525,7 +532,7 @@ namespace UnrealBuildTool.Android
 					}
 
 					// rebuild if .pak files exist for  in APK case
-					if (PackageDataInsideApk())
+					if (bPackageDataInsideApk)
 					{
 						string PAKFileLocation = ProjectDirectory + "/Saved/StagedBuilds/Android" + CookFlavor + "/" + ProjectName + "/Content/Paks";
 						if (Directory.Exists(PAKFileLocation))
@@ -604,7 +611,7 @@ namespace UnrealBuildTool.Android
 			}
 		}
 
-		private string GenerateManifest(string ProjectName, bool bIsForDistribution)
+		private string GenerateManifest(string ProjectName, bool bIsForDistribution, bool bPackageDataInsideApk)
 		{
 			// ini file to get settings from
 			ConfigCacheIni Ini = new ConfigCacheIni(UnrealTargetPlatform.Android, "Engine", UnrealBuildTool.GetUProjectPath());
@@ -682,7 +689,7 @@ namespace UnrealBuildTool.Android
 			}
 			Text.AppendLine("\t\t</activity>");
 			Text.AppendLine(string.Format("\t\t<meta-data android:name=\"com.epicgames.ue4.GameActivity.DepthBufferPreference\" android:value=\"{0}\"/>", ConvertDepthBufferIniValue(DepthBufferPreference)));
-			Text.AppendLine(string.Format("\t\t<meta-data android:name=\"com.epicgames.ue4.GameActivity.bPackageDataInsideApk\" android:value=\"{0}\"/>", PackageDataInsideApk(Ini) ? "true" : "false"));
+			Text.AppendLine(string.Format("\t\t<meta-data android:name=\"com.epicgames.ue4.GameActivity.bPackageDataInsideApk\" android:value=\"{0}\"/>", bPackageDataInsideApk ? "true" : "false"));
 			Text.AppendLine("\t\t<meta-data android:name=\"com.google.android.gms.games.APP_ID\"");
 			Text.AppendLine("\t\t           android:value=\"@string/app_id\" />");
 			Text.AppendLine("\t\t<meta-data android:name=\"com.google.android.gms.version\"");
@@ -742,7 +749,7 @@ namespace UnrealBuildTool.Android
 			return Text.ToString();
 		}
 
-		private void MakeApk(string ProjectName, string ProjectDirectory, string OutputPath, string EngineDirectory, bool bForDistribution, string CookFlavor, bool bMakeSeparateApks, bool bIncrementalPackage)
+		private void MakeApk(string ProjectName, string ProjectDirectory, string OutputPath, string EngineDirectory, bool bForDistribution, string CookFlavor, bool bMakeSeparateApks, bool bIncrementalPackage, bool bDisallowPackagingDataInApk)
 		{
 			Log.TraceInformation("\n===={0}====PREPARING TO MAKE APK=================================================================", DateTime.Now.ToString());
 
@@ -759,15 +766,18 @@ namespace UnrealBuildTool.Android
 			string[] GPUArchitectures = AndroidToolChain.GetAllGPUArchitectures();
 //			int NumArches = Arches.Length * GPUArchitectures.Length;
 
+			// cache if we want data in the Apk
+			bool bPackageDataInsideApk = PackageDataInsideApk(bDisallowPackagingDataInApk);
+
 			// check to see if any "meta information" is newer than last time we build
-			string CurrentBuildSettings = GetAllBuildSettings(UE4BuildPath, bForDistribution, bMakeSeparateApks);
+			string CurrentBuildSettings = GetAllBuildSettings(UE4BuildPath, bForDistribution, bMakeSeparateApks, bPackageDataInsideApk);
 			string BuildSettingsCacheFile = Path.Combine(UE4BuildPath, "UEBuildSettings.txt");
 
 			// do we match previous build settings?
 			bool bBuildSettingsMatch = false;
 
 			string ManifestFile = Path.Combine(UE4BuildPath, "AndroidManifest.xml");
-			string NewManifest = GenerateManifest(ProjectName, bForDistribution);
+			string NewManifest = GenerateManifest(ProjectName, bForDistribution, bDisallowPackagingDataInApk);
 			string OldManifest = File.Exists(ManifestFile) ? File.ReadAllText(ManifestFile) : "";
 			if (NewManifest == OldManifest) 
 			{
@@ -797,7 +807,7 @@ namespace UnrealBuildTool.Android
 			{
 				// check if so's are up to date against various inputs
 				bool bAllInputsCurrent = CheckDependencies(ProjectName, ProjectDirectory, UE4BuildFilesPath, GameBuildFilesPath, 
-					EngineDirectory, CookFlavor, OutputPath, UE4BuildPath, bMakeSeparateApks);
+					EngineDirectory, CookFlavor, OutputPath, UE4BuildPath, bMakeSeparateApks, bPackageDataInsideApk);
 
 				if (bAllInputsCurrent)
 				{
@@ -828,7 +838,6 @@ namespace UnrealBuildTool.Android
 			}
 
 			// If we are packaging for Amazon then we need to copy the  file to the correct location
-			bool bPackageDataInsideApk = PackageDataInsideApk();
 			Log.TraceInformation("bPackageDataInsideApk = {0}", bPackageDataInsideApk);
 			if (bPackageDataInsideApk)
 			{
@@ -1033,7 +1042,8 @@ namespace UnrealBuildTool.Android
 			string BaseSoName = AndroidToolChain.RemoveArchName(InTarget.OutputPaths[0]);
 
 			// make an apk at the end of compiling, so that we can run without packaging (debugger, cook on the fly, etc)
-			MakeApk(InTarget.AppName, InTarget.ProjectDirectory, BaseSoName, BuildConfiguration.RelativeEnginePath, bForDistribution: false, CookFlavor: "", bMakeSeparateApks:ShouldMakeSeparateApks(), bIncrementalPackage:true);
+			MakeApk(InTarget.AppName, InTarget.ProjectDirectory, BaseSoName, BuildConfiguration.RelativeEnginePath, bForDistribution: false, CookFlavor: "", 
+				bMakeSeparateApks:ShouldMakeSeparateApks(), bIncrementalPackage:true, bDisallowPackagingDataInApk:false);
 
 			// if we made any non-standard .apk files, the generated debugger settings may be wrong
 			if (ShouldMakeSeparateApks() && (InTarget.OutputPaths.Length > 1 || !InTarget.OutputPaths[0].Contains("-armv7-es2")))
@@ -1061,9 +1071,12 @@ namespace UnrealBuildTool.Android
 // 			return bSeparateApks;
 		}
 
-		public override bool PrepForUATPackageOrDeploy(string ProjectName, string ProjectDirectory, string ExecutablePath, string EngineDirectory, bool bForDistribution, string CookFlavor)
+		public override bool PrepForUATPackageOrDeploy(string ProjectName, string ProjectDirectory, string ExecutablePath, string EngineDirectory, bool bForDistribution, string CookFlavor, bool bIsDataDeploy)
 		{
-			MakeApk(ProjectName, ProjectDirectory, ExecutablePath, EngineDirectory, bForDistribution, CookFlavor, ShouldMakeSeparateApks(), bIncrementalPackage:false);
+			// note that we cannot allow the data packaged into the APK if we are doing something like Launch On that will not make an obb
+			// file and instead pushes files directly via deploy
+			MakeApk(ProjectName, ProjectDirectory, ExecutablePath, EngineDirectory, bForDistribution:bForDistribution, CookFlavor:CookFlavor, 
+				bMakeSeparateApks:ShouldMakeSeparateApks(), bIncrementalPackage:false, bDisallowPackagingDataInApk:bIsDataDeploy);
 			return true;
 		}
 
diff --git a/Engine/Source/Programs/UnrealBuildTool/Configuration/UEBuildDeploy.cs b/Engine/Source/Programs/UnrealBuildTool/Configuration/UEBuildDeploy.cs
--- a/Engine/Source/Programs/UnrealBuildTool/Configuration/UEBuildDeploy.cs
+++ b/Engine/Source/Programs/UnrealBuildTool/Configuration/UEBuildDeploy.cs
@@ -27,7 +27,7 @@ namespace UnrealBuildTool
 		/**
 		 * Prepare the target for deployment.
 		 */
-		bool PrepForUATPackageOrDeploy(string ProjectName, string ProjectDirectory, string ExecutablePath, string EngineDirectory, bool bForDistribution, string CookFlavor);
+		bool PrepForUATPackageOrDeploy(string ProjectName, string ProjectDirectory, string ExecutablePath, string EngineDirectory, bool bForDistribution, string CookFlavor, bool bIsDataDeploy);
 	}
 
 	/// <summary>
@@ -94,7 +94,7 @@ namespace UnrealBuildTool
 		/**
 		 * Prepare the target for deployment
 		 */
-		public virtual bool PrepForUATPackageOrDeploy(string ProjectName, string ProjectDirectory, string ExecutablePath, string EngineDirectory, bool bForDistribution, string CookFlavor)
+		public virtual bool PrepForUATPackageOrDeploy(string ProjectName, string ProjectDirectory, string ExecutablePath, string EngineDirectory, bool bForDistribution, string CookFlavor, bool bIsDataDeploy)
 		{
 			return true;
 		}
diff --git a/Engine/Source/Programs/UnrealBuildTool/IOS/UEDeployIOS.cs b/Engine/Source/Programs/UnrealBuildTool/IOS/UEDeployIOS.cs
--- a/Engine/Source/Programs/UnrealBuildTool/IOS/UEDeployIOS.cs
+++ b/Engine/Source/Programs/UnrealBuildTool/IOS/UEDeployIOS.cs
@@ -450,7 +450,7 @@ namespace UnrealBuildTool.IOS
 			return bSkipDefaultPNGs;
 		}
 
-		public override bool PrepForUATPackageOrDeploy(string InProjectName, string InProjectDirectory, string InExecutablePath, string InEngineDir, bool bForDistribution, string CookFlavor)
+		public override bool PrepForUATPackageOrDeploy(string InProjectName, string InProjectDirectory, string InExecutablePath, string InEngineDir, bool bForDistribution, string CookFlavor, bool bIsDataDeploy)
 		{
 			if (BuildHostPlatform.Current.Platform != UnrealTargetPlatform.Mac)
 			{
@@ -675,7 +675,7 @@ namespace UnrealBuildTool.IOS
 
 			if (BuildHostPlatform.Current.Platform == UnrealTargetPlatform.Mac && Environment.GetEnvironmentVariable("UBT_NO_POST_DEPLOY") != "true")
 			{
-				return PrepForUATPackageOrDeploy(GameName, ProjectDirectory, BuildPath + "/" + DecoratedGameName, "../../Engine", false, "");
+				return PrepForUATPackageOrDeploy(GameName, ProjectDirectory, BuildPath + "/" + DecoratedGameName, "../../Engine", false, "", false);
 			}
 			else
 			{
