[INTEGRATE] Change 2426198 by [REDACTED_PERSON][REDACTED_HANDLE] on 2015/01/30 14:56:31

	Include engine source folders in the Intelisense search paths. Allows debugger to find source in Rocket.

diff --git a/Engine/Source/Programs/UnrealBuildTool/System/VCProject.cs b/Engine/Source/Programs/UnrealBuildTool/System/VCProject.cs
--- a/Engine/Source/Programs/UnrealBuildTool/System/VCProject.cs
+++ b/Engine/Source/Programs/UnrealBuildTool/System/VCProject.cs
@@ -617,6 +617,23 @@ namespace UnrealBuildTool
 				WriteConfiguration( ProjectName, Combination, VCProjectFileContent, bGenerateUserFileContent? VCUserFileContent : null );
 			}
 
+			// For Rocket, include engine source in the source search paths. We never build it locally, so the debugger can't find it.
+			if(UnrealBuildTool.RunningRocket() && !IsStubProject)
+			{
+				VCProjectFileContent.Append("	<PropertyGroup>" + ProjectFileGenerator.NewLine);
+				VCProjectFileContent.Append("		<SourcePath>");
+				foreach(string DirectoryName in Directory.EnumerateDirectories(Path.GetFullPath(Path.Combine(ProjectFileGenerator.EngineRelativePath, "Source")), "*", SearchOption.AllDirectories))
+				{
+					if(Directory.EnumerateFiles(DirectoryName, "*.cpp").Any())
+					{
+						VCProjectFileContent.Append(DirectoryName);
+						VCProjectFileContent.Append(";");
+					}
+				}
+				VCProjectFileContent.Append("</SourcePath>" + ProjectFileGenerator.NewLine);
+				VCProjectFileContent.Append("	</PropertyGroup>" + ProjectFileGenerator.NewLine);
+			}
+
 			// Write IntelliSense info
 			{
 				// @todo projectfiles: Currently we are storing defines/include paths for ALL configurations rather than using ConditionString and storing
