Making sure UAT does not trigger hot-reload from IDE. Reverted changes to UBT return value which resulted in errors in the editor when doing hot-reload on an up-to-date code.

diff --git a/Engine/Source/Programs/AutomationTool/UBTUtils.cs b/Engine/Source/Programs/AutomationTool/UBTUtils.cs
--- a/Engine/Source/Programs/AutomationTool/UBTUtils.cs
+++ b/Engine/Source/Programs/AutomationTool/UBTUtils.cs
@@ -30,6 +30,8 @@ namespace AutomationTool
 			{
 				CommandLine += " -nocreatestub";
 			}
+			CommandLine += " -NoHotReloadFromIDE";
+
 			CommandUtils.RunAndLog(Env, UBTExecutable, CommandLine, LogName, EnvVars: EnvVars);
 		}
 
diff --git a/Engine/Source/Programs/UnrealBuildTool/System/UnrealBuildTool.cs b/Engine/Source/Programs/UnrealBuildTool/System/UnrealBuildTool.cs
--- a/Engine/Source/Programs/UnrealBuildTool/System/UnrealBuildTool.cs
+++ b/Engine/Source/Programs/UnrealBuildTool/System/UnrealBuildTool.cs
@@ -1231,9 +1231,7 @@ namespace UnrealBuildTool
             // Print some performance info
             Log.TraceVerbose("Execution time: {0}", (DateTime.UtcNow - StartTime - MutexWaitTime).TotalSeconds);
 
-            int RetVal = (Result == ECompilationResult.Succeeded || Result == ECompilationResult.UpToDate) ? 0 : (int)Result;
-
-            return RetVal; 
+            return (int)Result; 
         }
 
 
@@ -1891,6 +1889,12 @@ namespace UnrealBuildTool
 		/// <returns></returns>
 		private static bool ShouldDoHotReloadFromIDE(TargetDescriptor TargetDesc)
 		{
+			if (UnrealBuildTool.CommandLineContains("-NoHotReloadFromIDE"))
+			{
+				// Hot reload disabled through command line, possibly running from UAT
+				return false;
+			}
+
 			bool bIsRunning = false;
 			
 			// @todo ubtmake: Kind of cheating here to figure out if an editor target.  At this point we don't have access to the actual target description, and
