Fix hot reload for plugins.

diff --git a/Engine/Source/Programs/UnrealBuildTool/System/UnrealBuildTool.cs b/Engine/Source/Programs/UnrealBuildTool/System/UnrealBuildTool.cs
--- a/Engine/Source/Programs/UnrealBuildTool/System/UnrealBuildTool.cs
+++ b/Engine/Source/Programs/UnrealBuildTool/System/UnrealBuildTool.cs
@@ -2455,24 +2455,27 @@ namespace UnrealBuildTool
 					// Assume that the first produced item (with no extension) is our output file name
 					var OriginalFileNameWithoutExtension = Utils.GetFilenameWithoutAnyExtensions( Action.ProducedItems[0].AbsolutePath );
 
-					// Figure out if we have a numbered suffix at the end of the output file name, which will be present for already-existing
-					// module that was previously hot reloaded.  However, if the editor asked UnrealBuildTool to compile a brand new module
-					// and hot reload it, the output file will not have a numbered suffixed extension.
-					string OriginalFileNameWithoutNumberSuffix = OriginalFileNameWithoutExtension;
+					// Remove the numbered suffix from the end of the file
+					var IndexOfLastHyphen = OriginalFileNameWithoutExtension.LastIndexOf( '-' );
+					var OriginalNumberSuffix = OriginalFileNameWithoutExtension.Substring( IndexOfLastHyphen );
+					int NumberSuffix;
+					bool bHasNumberSuffix = int.TryParse(OriginalNumberSuffix, out NumberSuffix);
+
+					string OriginalFileNameWithoutNumberSuffix = null;
+					string PlatformConfigSuffix = string.Empty;
+					if (bHasNumberSuffix)
 					{
-						int IndexOfNumberedSuffix = OriginalFileNameWithoutExtension.LastIndexOf( '-' );
-						if( IndexOfNumberedSuffix != -1 )
-						{ 
-							var OriginalNumberSuffix = OriginalFileNameWithoutExtension.Substring( IndexOfNumberedSuffix );
-
-							// Make sure the number we pulled off the end of the string was actually a number (e.g. "-3141")
-							int TestParsedInt;
-							if( int.TryParse( OriginalNumberSuffix, out TestParsedInt ) )
-							{ 
-								// Remove the numbered suffix from the end of the file
-								OriginalFileNameWithoutNumberSuffix = OriginalFileNameWithoutExtension.Substring( 0, OriginalFileNameWithoutExtension.Length - OriginalNumberSuffix.Length );	// Remove "-####"
-							}
-						}
+						// Remove "-####" suffix in Development configuration
+						OriginalFileNameWithoutNumberSuffix = OriginalFileNameWithoutExtension.Substring(0, OriginalFileNameWithoutExtension.Length - OriginalNumberSuffix.Length);
+					}
+					else
+					{
+						// Remove "-####-Platform-Configuration" suffix in Debug configuration
+						var IndexOfSecondLastHyphen = OriginalFileNameWithoutExtension.LastIndexOf('-', IndexOfLastHyphen - 1, IndexOfLastHyphen);
+						var SuffixLenght = OriginalFileNameWithoutExtension.Length - IndexOfSecondLastHyphen;
+						var IndexOfThirdLastHyphen = OriginalFileNameWithoutExtension.LastIndexOf('-', IndexOfSecondLastHyphen - 1, IndexOfSecondLastHyphen);
+						OriginalFileNameWithoutNumberSuffix = OriginalFileNameWithoutExtension.Substring(0, IndexOfThirdLastHyphen);
+						PlatformConfigSuffix = OriginalFileNameWithoutExtension.Substring(IndexOfSecondLastHyphen);
 					}
 
 					// Figure out which suffix to use
@@ -2506,7 +2509,7 @@ namespace UnrealBuildTool
 							UniqueSuffix = "-" + (new Random((int)(DateTime.Now.Ticks % Int32.MaxValue)).Next(10000)).ToString();
 						}
 					}
-					var NewFileNameWithoutExtension = OriginalFileNameWithoutNumberSuffix + UniqueSuffix;
+					var NewFileNameWithoutExtension = OriginalFileNameWithoutNumberSuffix + UniqueSuffix + PlatformConfigSuffix;
 
 					// Find the response file in the command line.  We'll need to make a copy of it with our new file name.
 					{
diff --git a/Engine/Source/Runtime/Core/Private/Modules/ModuleManager.cpp b/Engine/Source/Runtime/Core/Private/Modules/ModuleManager.cpp
--- a/Engine/Source/Runtime/Core/Private/Modules/ModuleManager.cpp
+++ b/Engine/Source/Runtime/Core/Private/Modules/ModuleManager.cpp
@@ -825,6 +825,17 @@ void FModuleManager::FindModulePaths(const TCHAR* NamePattern, TMap<FName, FStri
 		FindModulePathsInDirectory(EngineBinariesDirectories[Idx], false, NamePattern, OutModulePaths);
 	}
 
+	// If any directory found so far is a plugin directory, skip next step as it would override whatever we 
+	// found so far. Plugins are stored in their binary directory, not in game binary directory.
+	static const FString Plugin = TEXT("Plugins");
+	for (auto Kvp : OutModulePaths)
+	{
+		if (Kvp.Value.Contains(Plugin))
+		{
+			return;
+		}
+	}
+
 	// Search any game directories
 	for (int Idx = 0; Idx < GameBinariesDirectories.Num(); Idx++)
 	{
