- Fixed packaging TappyChicken [UE-9441]

diff --git a/Engine/Source/Programs/UnrealBuildTool/Android/UEDeployAndroid.cs b/Engine/Source/Programs/UnrealBuildTool/Android/UEDeployAndroid.cs
--- a/Engine/Source/Programs/UnrealBuildTool/Android/UEDeployAndroid.cs
+++ b/Engine/Source/Programs/UnrealBuildTool/Android/UEDeployAndroid.cs
@@ -730,7 +730,6 @@ namespace UnrealBuildTool.Android
 			Text.AppendLine("\t<uses-permission android:name=\"android.permission.ACCESS_WIFI_STATE\"/>");
 			Text.AppendLine("\t<uses-permission android:name=\"android.permission.MODIFY_AUDIO_SETTINGS\"/>");
 			Text.AppendLine("\t<uses-permission android:name=\"android.permission.GET_ACCOUNTS\"/>");
-            Text.AppendLine("\t<uses-permission android:name=\"android.permission.VIBRATE\"/>");
 			Text.AppendLine("\t<uses-permission android:name=\"com.android.vending.BILLING\"/>");
 //			Text.AppendLine("\t<uses-permission android:name=\"android.permission.DISABLE_KEYGUARD\"/>");
 			if (bPackageForGearVR)
@@ -778,7 +777,7 @@ namespace UnrealBuildTool.Android
 			bool bBuildSettingsMatch = false;
 
 			string ManifestFile = Path.Combine(UE4BuildPath, "AndroidManifest.xml");
-			string NewManifest = GenerateManifest(ProjectName, bForDistribution, bDisallowPackagingDataInApk);
+			string NewManifest = GenerateManifest(ProjectName, bForDistribution, bPackageDataInsideApk);
 			string OldManifest = File.Exists(ManifestFile) ? File.ReadAllText(ManifestFile) : "";
 			if (NewManifest == OldManifest) 
 			{
@@ -1074,8 +1073,8 @@ namespace UnrealBuildTool.Android
 
 		public override bool PrepForUATPackageOrDeploy(string ProjectName, string ProjectDirectory, string ExecutablePath, string EngineDirectory, bool bForDistribution, string CookFlavor, bool bIsDataDeploy)
 		{
-			// note that we cannot allow the data packaged into the APK if we are doing something like Launch On that will not make an obb
-			// file and instead pushes files directly via deploy
+			// note that we cannot allow the data packaged into the APK if we are doing something like Launch On that will not make a pak (obb)
+			// file and instead pushes files directly
 			MakeApk(ProjectName, ProjectDirectory, ExecutablePath, EngineDirectory, bForDistribution:bForDistribution, CookFlavor:CookFlavor, 
 				bMakeSeparateApks:ShouldMakeSeparateApks(), bIncrementalPackage:false, bDisallowPackagingDataInApk:bIsDataDeploy);
 			return true;
