[INTEGRATE] Change 2436754 by [REDACTED_PERSON][REDACTED_HANDLE] on 2015/02/07 11:04:49

	Rocket projects now get all engine source files (not just header files)
	- This makes Rocket projects a bit larger and they load slower, but it is much more useful for navigating source
	- To get the old behavior, you can pass the new "-OnlyPublic" command-line option
	- Fixes UE-9390

diff --git a/Engine/Source/Programs/UnrealBuildTool/System/ProjectFileGenerator.cs b/Engine/Source/Programs/UnrealBuildTool/System/ProjectFileGenerator.cs
--- a/Engine/Source/Programs/UnrealBuildTool/System/ProjectFileGenerator.cs
+++ b/Engine/Source/Programs/UnrealBuildTool/System/ProjectFileGenerator.cs
@@ -1492,8 +1492,10 @@ namespace UnrealBuildTool
 					{
 						if( IncludeEngineSource )
 						{
-							// This is an engine module, so strip out private code from the Rocket projects.  Rocket users only need public headers to compile
-							IncludePrivateSourceCode &= !bGeneratingRocketProjectFiles;
+				            // Always include all engine source files unless we were asked to only find public header files via command-line option.  In older releases
+				            // of Unreal Engine, we only included public header files in binary distributions, but this is no longer the case.  You might get
+				            // faster solution load times if you use '-OnlyPublic' though.
+				            IncludePrivateSourceCode &= !UnrealBuildTool.CommandLineContains( "-OnlyPublic" );
 						}
 						else
 						{
