Use relative paths in the response file used for linking on Mac to fix the posix_spawn_failed error in non-unity builds

diff --git a/Engine/Source/Programs/UnrealBuildTool/Mac/MacToolChain.cs b/Engine/Source/Programs/UnrealBuildTool/Mac/MacToolChain.cs
--- a/Engine/Source/Programs/UnrealBuildTool/Mac/MacToolChain.cs
+++ b/Engine/Source/Programs/UnrealBuildTool/Mac/MacToolChain.cs
@@ -794,7 +794,16 @@ namespace UnrealBuildTool
 			List<string> InputFileNames = new List<string>();
 			foreach (FileItem InputFile in LinkEnvironment.InputFiles)
 			{
-				InputFileNames.Add(string.Format("\"{0}\"", InputFile.AbsolutePath));
+				if (bIsBuildingLibrary)
+				{
+					InputFileNames.Add(string.Format("\"{0}\"", InputFile.AbsolutePath));
+				}
+				else
+				{
+					string EnginePath = ConvertPath(Path.GetDirectoryName(Directory.GetCurrentDirectory()));
+					string InputFileRelativePath = InputFile.AbsolutePath.Replace(EnginePath, "..");
+					InputFileNames.Add(string.Format("\"{0}\"", InputFileRelativePath));
+				}
 				LinkAction.PrerequisiteItems.Add(InputFile);
 			}
 
@@ -923,9 +932,11 @@ namespace UnrealBuildTool
 					{
 						EngineAndGameLibrariesString += " \"" + Library + "\"";
 					}
-					string FixDylibLine = string.Format("TIMESTAMP=`stat -n -f \"%Sm\" -t \"%Y%m%d%H%M.%S\" \"{0}\"`; ", RemoteOutputFile.AbsolutePath);
+					string FixDylibLine = "pushd \"" + ConvertPath(Directory.GetCurrentDirectory()) + "\"; ";
+					FixDylibLine += string.Format("TIMESTAMP=`stat -n -f \"%Sm\" -t \"%Y%m%d%H%M.%S\" \"{0}\"`; ", RemoteOutputFile.AbsolutePath);
 					FixDylibLine += LinkCommand.Replace("-undefined dynamic_lookup", EngineAndGameLibrariesString).Replace("$", "\\$");
-					FixDylibLine += string.Format("; touch -t $TIMESTAMP \"{0}\"; if [[ $? -ne 0 ]]; then exit 1; fi", RemoteOutputFile.AbsolutePath);
+					FixDylibLine += string.Format("; touch -t $TIMESTAMP \"{0}\"; if [[ $? -ne 0 ]]; then exit 1; fi; ", RemoteOutputFile.AbsolutePath);
+					FixDylibLine += "popd";
 					AppendMacLine(FixDylibDepsScript, FixDylibLine);
 				}
 
