Disable output from pushd and popd called from FixDylibDependencies script on Mac

diff --git a/Engine/Source/Programs/UnrealBuildTool/Mac/MacToolChain.cs b/Engine/Source/Programs/UnrealBuildTool/Mac/MacToolChain.cs
--- a/Engine/Source/Programs/UnrealBuildTool/Mac/MacToolChain.cs
+++ b/Engine/Source/Programs/UnrealBuildTool/Mac/MacToolChain.cs
@@ -932,11 +932,11 @@ namespace UnrealBuildTool
 					{
 						EngineAndGameLibrariesString += " \"" + Library + "\"";
 					}
-					string FixDylibLine = "pushd \"" + ConvertPath(Directory.GetCurrentDirectory()) + "\"; ";
+					string FixDylibLine = "pushd \"" + ConvertPath(Directory.GetCurrentDirectory()) + "\"  > /dev/null; ";
 					FixDylibLine += string.Format("TIMESTAMP=`stat -n -f \"%Sm\" -t \"%Y%m%d%H%M.%S\" \"{0}\"`; ", RemoteOutputFile.AbsolutePath);
 					FixDylibLine += LinkCommand.Replace("-undefined dynamic_lookup", EngineAndGameLibrariesString).Replace("$", "\\$");
 					FixDylibLine += string.Format("; touch -t $TIMESTAMP \"{0}\"; if [[ $? -ne 0 ]]; then exit 1; fi; ", RemoteOutputFile.AbsolutePath);
-					FixDylibLine += "popd";
+					FixDylibLine += "popd > /dev/null";
 					AppendMacLine(FixDylibDepsScript, FixDylibLine);
 				}
 
