[INTEGRATE] Change 2442289 by [REDACTED_PERSON][REDACTED_HANDLE] on 2015/02/11 16:26:04

	Strip linux executables in Rocket. Saves ~400mb from install footprint.

diff --git a/Engine/Source/Programs/UnrealBuildTool/Linux/LinuxToolChain.cs b/Engine/Source/Programs/UnrealBuildTool/Linux/LinuxToolChain.cs
--- a/Engine/Source/Programs/UnrealBuildTool/Linux/LinuxToolChain.cs
+++ b/Engine/Source/Programs/UnrealBuildTool/Linux/LinuxToolChain.cs
@@ -451,6 +451,12 @@ namespace UnrealBuildTool
                 Result += string.Format(" -Wl,--unresolved-symbols=ignore-in-shared-libs");
             }
 
+            if (UnrealBuildTool.BuildingRocket())
+            {
+                // strip symbols for Rocket in every configuration
+                Result += " -Wl,-s";
+            }
+
             // RPATH for third party libs
             Result += " -Wl,-rpath=${ORIGIN}";
             Result += " -Wl,-rpath-link=${ORIGIN}";
