Reverted changes in some files that shouldn't have been checked in that was trying to build D3D shader support for Mac

diff --git a/Engine/Source/Programs/UnrealBuildTool/Linux/UEBuildLinux.cs b/Engine/Source/Programs/UnrealBuildTool/Linux/UEBuildLinux.cs
--- a/Engine/Source/Programs/UnrealBuildTool/Linux/UEBuildLinux.cs
+++ b/Engine/Source/Programs/UnrealBuildTool/Linux/UEBuildLinux.cs
@@ -327,11 +327,12 @@ namespace UnrealBuildTool
                     InModule.AddDynamicallyLoadedModule("LinuxTargetPlatform");
                     InModule.AddDynamicallyLoadedModule("LinuxNoEditorTargetPlatform");
                     InModule.AddDynamicallyLoadedModule("LinuxServerTargetPlatform");
+					InModule.AddDynamicallyLoadedModule("DesktopTargetPlatform");
                 }
 
                 if (bBuildShaderFormats)
                 {
-					InModule.AddDynamicallyLoadedModule("ShaderFormatD3D");
+					// InModule.AddDynamicallyLoadedModule("ShaderFormatD3D");
                     InModule.AddDynamicallyLoadedModule("ShaderFormatOpenGL");
                 }
             }
diff --git a/Engine/Source/Programs/UnrealBuildTool/Mac/UEBuildMac.cs b/Engine/Source/Programs/UnrealBuildTool/Mac/UEBuildMac.cs
--- a/Engine/Source/Programs/UnrealBuildTool/Mac/UEBuildMac.cs
+++ b/Engine/Source/Programs/UnrealBuildTool/Mac/UEBuildMac.cs
@@ -113,11 +113,12 @@ namespace UnrealBuildTool
 					InModule.AddDynamicallyLoadedModule("MacNoEditorTargetPlatform");
 					InModule.AddDynamicallyLoadedModule("MacClientTargetPlatform");
 					InModule.AddDynamicallyLoadedModule("MacServerTargetPlatform");
+					InModule.AddDynamicallyLoadedModule("DesktopTargetPlatform");
 				}
 
                 if (bBuildShaderFormats)
                 {
-					InModule.AddDynamicallyLoadedModule("ShaderFormatD3D");
+					// InModule.AddDynamicallyLoadedModule("ShaderFormatD3D");
                     InModule.AddDynamicallyLoadedModule("ShaderFormatOpenGL");
                 }
 			}
