[INTEGRATE] Change 2448139 by [REDACTED_PERSON][REDACTED_HANDLE] on 2015/02/16 21:59:45

	Clean up empty directories at UBT startup, to prevent checks for available build features misfiring.

diff --git a/Engine/Source/Programs/UnrealBuildTool/System/UnrealBuildTool.cs b/Engine/Source/Programs/UnrealBuildTool/System/UnrealBuildTool.cs
--- a/Engine/Source/Programs/UnrealBuildTool/System/UnrealBuildTool.cs
+++ b/Engine/Source/Programs/UnrealBuildTool/System/UnrealBuildTool.cs
@@ -673,6 +673,10 @@ namespace UnrealBuildTool
             if (Directory.Exists(EngineSourceDirectory)) // only set the directory if it exists, this should only happen if we are launching the editor from an artist sync
             {
                 Directory.SetCurrentDirectory(EngineSourceDirectory);
+				if(RunningRocket())
+				{
+					RemoveEmptyFolders(EngineSourceDirectory);
+				}
             }
 
             // Build the list of game projects that we know about. When building from the editor (for hot-reload) or for Rocket projects, we require the 
@@ -1234,6 +1238,25 @@ namespace UnrealBuildTool
             return (int)Result; 
         }
 
+		private static void RemoveEmptyFolders(string BaseDirectoryName)
+		{
+			string[] ChildDirectoryNames = Directory.EnumerateDirectories(BaseDirectoryName).ToArray();
+			foreach(string ChildDirectoryName in ChildDirectoryNames)
+			{
+				RemoveEmptyFolders(ChildDirectoryName);
+			}
+
+			string[] ChildFileNames = Directory.EnumerateFiles(BaseDirectoryName).ToArray();
+			if(ChildFileNames.Length == 1 && Path.GetFileName(ChildFileNames[0]) == ".DS_Store")
+			{
+				try { File.Delete(ChildFileNames[0]); } catch(Exception) { }
+			}
+
+			if(!Directory.EnumerateFileSystemEntries(BaseDirectoryName).Any())
+			{
+				try { Directory.Delete(BaseDirectoryName); } catch (Exception) {  }
+			}
+		}
 
         /// <summary>
         /// Generates project files.  Can also be used to generate projects "in memory", to allow for building
