[INTEGRATE] Change 2449602 by [REDACTED_PERSON][REDACTED_HANDLE] on 2015/02/17 21:30:32

	Move code which detects the VS2012 toolchain for Windows. We don't support it, and UAT calling into it can cause attempts to stage the wrong DLLs for BP-only games. Leave the plumbing around for when we introduce VS2015 support. (UE-9991)

diff --git a/Engine/Source/Programs/UnrealBuildTool/Windows/UEBuildWindows.cs b/Engine/Source/Programs/UnrealBuildTool/Windows/UEBuildWindows.cs
--- a/Engine/Source/Programs/UnrealBuildTool/Windows/UEBuildWindows.cs
+++ b/Engine/Source/Programs/UnrealBuildTool/Windows/UEBuildWindows.cs
@@ -40,16 +40,10 @@ namespace UnrealBuildTool
                 }
 
                 // First check if Visual Sudio 2013 is installed.
-                // If both 2012 and 2013 are installed, prefer 2013
                 if (!String.IsNullOrEmpty(WindowsPlatform.GetVSComnToolsPath(WindowsCompiler.VisualStudio2013)))
                 {
                     CachedCompiler = WindowsCompiler.VisualStudio2013;
                 }
-                // Next try Visual Studio 2012
-                else if (!String.IsNullOrEmpty(WindowsPlatform.GetVSComnToolsPath(WindowsCompiler.VisualStudio2012)))
-                {
-                    CachedCompiler = WindowsCompiler.VisualStudio2012;
-                }
                 // Finally assume 2013 is installed to defer errors somewhere else like VCToolChain
                 else
                 {
@@ -182,13 +176,6 @@ namespace UnrealBuildTool
          */
         protected override void RegisterBuildPlatformInternal()
         {
-			// All the plumbing for VS2012 is still in place to support console toolchains that piggy-back on the VS2012 environment, but we require VS2013 for building the editor
-			// and host tools. We always prefer VS2013 over VS2012, so if we detected VS2012 as the only installed compiler, halt now.
-            if (!Utils.IsRunningOnMono && Compiler == WindowsCompiler.VisualStudio2012)
-            {
-				throw new BuildException("Visual Studio 2012 is no longer supported for building the editor and host platform tools. Please also install Visual Studio 2013.");
-            }
-
             // Register this build platform for both Win64 and Win32
             Log.TraceVerbose("        Registering for {0}", UnrealTargetPlatform.Win64.ToString());
             UEBuildPlatform.RegisterBuildPlatform(UnrealTargetPlatform.Win64, this);
