GitHub pull request #719

diff --git a/Engine/Source/Programs/UnrealBuildTool/System/UnrealBuildTool.cs b/Engine/Source/Programs/UnrealBuildTool/System/UnrealBuildTool.cs
--- a/Engine/Source/Programs/UnrealBuildTool/System/UnrealBuildTool.cs
+++ b/Engine/Source/Programs/UnrealBuildTool/System/UnrealBuildTool.cs
@@ -618,7 +618,6 @@ namespace UnrealBuildTool
         private static int Main(string[] Arguments)
         {
             InitLogging();
-            Log.TraceVerbose(Environment.CommandLine);
 
             try
             {
@@ -630,6 +629,16 @@ namespace UnrealBuildTool
                 return (int) ECompilationResult.OtherCompilationError;
             }
 
+			// Because XmlConfigLoader.Init() will reset all parameter by call LoadDefaults, 
+			// so need to parse verbose argument here.
+			int Verbose;
+			if (Utils.ParseCommandLineFlag(Arguments, "-verbose", out Verbose))
+			{
+				BuildConfiguration.bPrintDebugInfo = true;
+			}
+
+			Log.TraceVerbose(Environment.CommandLine);
+
             InitialEnvironment = Environment.GetEnvironmentVariables();
             if (InitialEnvironment.Count < 1)
             {
@@ -695,6 +704,11 @@ namespace UnrealBuildTool
             XmlConfigLoader.Reset<BuildConfiguration>();
             XmlConfigLoader.Reset<UEBuildConfiguration>();
 
+			if (Utils.ParseCommandLineFlag(Arguments, "-verbose", out Verbose))
+			{
+				BuildConfiguration.bPrintDebugInfo = true;
+			}
+
             Log.TraceVerbose("UnrealBuildTool (DEBUG OUTPUT MODE)");
             Log.TraceVerbose("Command-line: {0}", String.Join(" ", Arguments));
 
