HTML5 CIS bodge. Looks like we need to update SDK.

diff --git a/Engine/Source/Programs/UnrealBuildTool/HTML5/HTML5ToolChain.cs b/Engine/Source/Programs/UnrealBuildTool/HTML5/HTML5ToolChain.cs
--- a/Engine/Source/Programs/UnrealBuildTool/HTML5/HTML5ToolChain.cs
+++ b/Engine/Source/Programs/UnrealBuildTool/HTML5/HTML5ToolChain.cs
@@ -13,6 +13,7 @@ namespace UnrealBuildTool
 	{
         static string EMCCPath;
         static string PythonPath;
+		static string EMSDKVersionString;
         // Debug options -- TODO: add these to SDK/Project setup?
         static bool   bEnableTracing = false;
 
@@ -22,7 +23,8 @@ namespace UnrealBuildTool
             if (HTML5SDKInfo.IsSDKInstalled() && HTML5SDKInfo.IsPythonInstalled())
             {
                     EMCCPath = "\"" + HTML5SDKInfo.EmscriptenCompiler() + "\"";
-					PythonPath = HTML5SDKInfo.PythonPath(); 
+					PythonPath = HTML5SDKInfo.PythonPath();
+					EMSDKVersionString = HTML5SDKInfo.EmscriptenVersion().Replace(".", "");
 
 					// set some environment variable we'll need
 					//Environment.SetEnvironmentVariable("EMCC_DEBUG", "cache");
@@ -322,6 +324,8 @@ namespace UnrealBuildTool
 				Arguments += string.Format(" -D__EMSCRIPTEN_TRACING__");
 			}
 
+			Arguments += string.Format(" -D__EMCC_VER__={0}", EMSDKVersionString);
+
         
 			var BuildPlatform = UEBuildPlatform.GetBuildPlatformForCPPTargetPlatform(CompileEnvironment.Config.Target.Platform);
 
diff --git a/Engine/Source/Runtime/OpenGLDrv/Private/HTML5/HTML5OpenGL.cpp b/Engine/Source/Runtime/OpenGLDrv/Private/HTML5/HTML5OpenGL.cpp
--- a/Engine/Source/Runtime/OpenGLDrv/Private/HTML5/HTML5OpenGL.cpp
+++ b/Engine/Source/Runtime/OpenGLDrv/Private/HTML5/HTML5OpenGL.cpp
@@ -354,15 +354,16 @@ extern "C"
 
 	EM_BOOL request_fullscreen_callback(int eventType, const EmscriptenMouseEvent* evt, void* user)
 	{
-		
+#if __EMCC_VER__ >= 1290
 		EmscriptenFullscreenStrategy FSStrat;
 		FMemory::Memzero(FSStrat);
 		FSStrat.scaleMode = EMSCRIPTEN_FULLSCREEN_SCALE_STRETCH;//EMSCRIPTEN_FULLSCREEN_SCALE_ASPECT;// : EMSCRIPTEN_FULLSCREEN_SCALE_STRETCH;
 		FSStrat.canvasResolutionScaleMode = EMSCRIPTEN_FULLSCREEN_CANVAS_SCALE_HIDEF;
 		FSStrat.filteringMode = EMSCRIPTEN_FULLSCREEN_FILTERING_DEFAULT;
 		emscripten_request_fullscreen_strategy("canvas", true, &FSStrat);
-		
-		//emscripten_request_fullscreen("canvas", true);
+#else		
+		emscripten_request_fullscreen("canvas", true);
+#endif
 		return 0;
 	}
 }
