- Fixed JunkDeleter.cs to work with the backslashes in JunkManifest.txt on Mac (and Linux likely)

diff --git a/Engine/Source/Programs/UnrealBuildTool/System/JunkDeleter.cs b/Engine/Source/Programs/UnrealBuildTool/System/JunkDeleter.cs
--- a/Engine/Source/Programs/UnrealBuildTool/System/JunkDeleter.cs
+++ b/Engine/Source/Programs/UnrealBuildTool/System/JunkDeleter.cs
@@ -89,7 +89,10 @@ namespace UnrealBuildTool
 							if (bIsValidJunkLine)
 							{
 								// the entry is always the last element in the token array (after the final :)
-								JunkManifest.Add(Path.Combine(CurrentToRootDir, Tokens[Tokens.Length - 1]));
+								string FixedPath = Path.Combine(CurrentToRootDir, Tokens[Tokens.Length - 1]);
+								FixedPath = FixedPath.Replace('\\', Path.DirectorySeparatorChar);
+								FixedPath = Path.GetFullPath(FixedPath);
+								JunkManifest.Add(FixedPath);
 							}
 						}
 					}
