- Fixed my fix for JunkDeleter (* is not a valid character in Windows, but it is in Mac it seems)

diff --git a/Engine/Source/Programs/UnrealBuildTool/System/JunkDeleter.cs b/Engine/Source/Programs/UnrealBuildTool/System/JunkDeleter.cs
--- a/Engine/Source/Programs/UnrealBuildTool/System/JunkDeleter.cs
+++ b/Engine/Source/Programs/UnrealBuildTool/System/JunkDeleter.cs
@@ -91,7 +91,6 @@ namespace UnrealBuildTool
 								// the entry is always the last element in the token array (after the final :)
 								string FixedPath = Path.Combine(CurrentToRootDir, Tokens[Tokens.Length - 1]);
 								FixedPath = FixedPath.Replace('\\', Path.DirectorySeparatorChar);
-								FixedPath = Path.GetFullPath(FixedPath);
 								JunkManifest.Add(FixedPath);
 							}
 						}
