fixes for building Linux when the SDK path has a space in it

diff --git a/Engine/Source/Programs/UnrealBuildTool/Linux/LinuxToolChain.cs b/Engine/Source/Programs/UnrealBuildTool/Linux/LinuxToolChain.cs
--- a/Engine/Source/Programs/UnrealBuildTool/Linux/LinuxToolChain.cs
+++ b/Engine/Source/Programs/UnrealBuildTool/Linux/LinuxToolChain.cs
@@ -383,7 +383,7 @@ namespace UnrealBuildTool
                 {
                     Result += String.Format(" -target {0}", CompileEnvironment.Config.Target.Architecture);        // Set target triple
                 }
-                Result += String.Format(" --sysroot={0}", BaseLinuxPath);
+                Result += String.Format(" --sysroot=\"{0}\"", BaseLinuxPath);
             }
 
             return Result;
@@ -707,7 +707,7 @@ namespace UnrealBuildTool
             else
             {
                 ArchiveAction.CommandPath = "cmd.exe";
-                ArchiveAction.CommandArguments = "/c ";
+                ArchiveAction.CommandArguments = "/c \"";
             }
 
             // this will produce a final library
@@ -718,7 +718,7 @@ namespace UnrealBuildTool
             ArchiveAction.ProducedItems.Add(OutputFile);
             ArchiveAction.CommandDescription = "Archive";
             ArchiveAction.StatusDescription = Path.GetFileName(OutputFile.AbsolutePath);
-            ArchiveAction.CommandArguments += string.Format("{0} {1} \"{2}\"",  GetArPath(LinkEnvironment.Config.Target.Architecture), GetArchiveArguments(LinkEnvironment), OutputFile.AbsolutePath);
+            ArchiveAction.CommandArguments += string.Format("\"{0}\" {1} \"{2}\"",  GetArPath(LinkEnvironment.Config.Target.Architecture), GetArchiveArguments(LinkEnvironment), OutputFile.AbsolutePath);
 
             // Add the input files to a response file, and pass the response file on the command-line.
             List<string> InputFileNames = new List<string>();
@@ -731,7 +731,7 @@ namespace UnrealBuildTool
             }
 
             // add ranlib
-            ArchiveAction.CommandArguments += string.Format(" && {0} \"{1}\"", GetRanlibPath(LinkEnvironment.Config.Target.Architecture), OutputFile.AbsolutePath);
+            ArchiveAction.CommandArguments += string.Format(" && \"{0}\" \"{1}\"", GetRanlibPath(LinkEnvironment.Config.Target.Architecture), OutputFile.AbsolutePath);
 
             // Add the additional arguments specified by the environment.
             ArchiveAction.CommandArguments += LinkEnvironment.Config.AdditionalArguments;
@@ -741,6 +741,10 @@ namespace UnrealBuildTool
             {
                 ArchiveAction.CommandArguments += "'";
             }
+			else
+			{
+				ArchiveAction.CommandArguments += "\"";
+			}
 
             // Only execute linking on the local PC.
             ArchiveAction.bCanExecuteRemotely = false;
