content-only tool chain changes

diff --git a/Engine/Build/Target.cs.template b/Engine/Build/Target.cs.template
new file mode 100644
--- /dev/null
+++ b/Engine/Build/Target.cs.template
@@ -0,0 +1,132 @@
+// Copyright 1998-2015 Epic Games, Inc. All Rights Reserved.
+
+using UnrealBuildTool;
+using System.Collections.Generic;
+
+public class {GAME_NAME}Target : TargetRules
+{
+    public {GAME_NAME}Target(TargetInfo Target)
+	{
+		Type = TargetType.Game;
+	}
+
+	//
+	// TargetRules interface.
+	//
+
+	public override void SetupBinaries(
+		TargetInfo Target,
+		ref List<UEBuildBinaryConfiguration> OutBuildBinaryConfigurations,
+		ref List<string> OutExtraModuleNames
+		)
+	{
+        OutExtraModuleNames.Add("UE4Game");
+		// this is important - for some reason achievements etc intertwined with the onlinesubsystem and they saved without using a fake OSS. :/
+		if (Target.Platform == UnrealTargetPlatform.HTML5)
+		{
+			OutExtraModuleNames.Add("OnlineSubsystemNull");
+		}
+
+
+		if (UnrealBuildTool.UnrealBuildTool.BuildingRocket())
+		{
+			OutExtraModuleNames.Add("GameMenuBuilder");
+
+			if (Target.Platform == UnrealTargetPlatform.Win32 || Target.Platform == UnrealTargetPlatform.Win64)
+			{
+				OutExtraModuleNames.Add("OnlineSubsystemNull");
+				OutExtraModuleNames.Add("OnlineSubsystemAmazon");
+				if (UEBuildConfiguration.bCompileSteamOSS == true)
+				{
+					OutExtraModuleNames.Add("OnlineSubsystemSteam");
+				}
+				OutExtraModuleNames.Add("OnlineSubsystemFacebook");
+			}
+			else if (Target.Platform == UnrealTargetPlatform.Mac || Target.Platform == UnrealTargetPlatform.Linux)
+			{
+				OutExtraModuleNames.Add("OnlineSubsystemNull");
+				if (UEBuildConfiguration.bCompileSteamOSS == true)
+				{
+					OutExtraModuleNames.Add("OnlineSubsystemSteam");
+				}
+			}
+			else if (Target.Platform == UnrealTargetPlatform.IOS)
+			{
+				OutExtraModuleNames.Add("OnlineSubsystemFacebook");
+				OutExtraModuleNames.Add("OnlineSubsystemIOS");
+				OutExtraModuleNames.Add("IOSAdvertising");
+			}
+			else if (Target.Platform == UnrealTargetPlatform.Android)
+			{
+				// @todo android: Add Android online subsystem
+				OutExtraModuleNames.Add("AndroidAdvertising");
+			}
+			else if (Target.Platform == UnrealTargetPlatform.HTML5)
+			{
+				OutExtraModuleNames.Add("OnlineSubsystemNull");
+			}
+		}
+	}
+
+	public override void SetupGlobalEnvironment(
+		TargetInfo Target,
+		ref LinkEnvironmentConfiguration OutLinkEnvironmentConfiguration,
+		ref CPPEnvironmentConfiguration OutCPPEnvironmentConfiguration
+		)
+	{
+		if (UnrealBuildTool.UnrealBuildTool.BuildingRocket())
+		{
+			UEBuildConfiguration.bCompileLeanAndMeanUE = true;
+
+			// Don't need editor or editor only data
+			UEBuildConfiguration.bBuildEditor = false;
+			UEBuildConfiguration.bBuildWithEditorOnlyData = false;
+
+			UEBuildConfiguration.bCompileAgainstEngine = true;
+
+			// no exports, so no need to verify that a .lib and .exp file was emitted by the linker.
+			OutLinkEnvironmentConfiguration.bHasExports = false;
+		}
+		else
+		{
+			// Tag it as a UE4Game build
+			OutCPPEnvironmentConfiguration.Definitions.Add("UE4GAME=1");
+		}
+	}
+
+	public override List<UnrealTargetPlatform> GUBP_GetPlatforms_MonolithicOnly(UnrealTargetPlatform HostPlatform)
+	{
+		List<UnrealTargetPlatform> Platforms = null;
+
+		switch (HostPlatform)
+		{
+			case UnrealTargetPlatform.Mac:
+				Platforms = new List<UnrealTargetPlatform> { HostPlatform, UnrealTargetPlatform.IOS };
+				break;
+
+			case UnrealTargetPlatform.Linux:
+				Platforms = new List<UnrealTargetPlatform> { HostPlatform };
+				break;
+
+			case UnrealTargetPlatform.Win64:
+				Platforms = new List<UnrealTargetPlatform> { HostPlatform, UnrealTargetPlatform.Win32, UnrealTargetPlatform.IOS, UnrealTargetPlatform.XboxOne, UnrealTargetPlatform.PS4, UnrealTargetPlatform.Android, UnrealTargetPlatform.Linux, UnrealTargetPlatform.HTML5 };
+				break;
+
+			default:
+				Platforms = new List<UnrealTargetPlatform>();
+				break;
+		}
+
+		return Platforms;
+	}
+
+	public override List<UnrealTargetConfiguration> GUBP_GetConfigs_MonolithicOnly(UnrealTargetPlatform HostPlatform, UnrealTargetPlatform Platform)
+	{
+		return new List<UnrealTargetConfiguration> { UnrealTargetConfiguration.Development, UnrealTargetConfiguration.Shipping, UnrealTargetConfiguration.Test };
+	}
+
+	public override bool GUBP_BuildWindowsXPMonolithics()
+	{
+		return true;
+	}
+}
diff --git a/Engine/Config/BaseEngine.ini b/Engine/Config/BaseEngine.ini
--- a/Engine/Config/BaseEngine.ini
+++ b/Engine/Config/BaseEngine.ini
@@ -2131,6 +2131,7 @@ bDevForArmV7S=False
 bShipForArmV7=True
 bShipForArm64=True
 bShipForArmV7S=False
+bGeneratedSYMFile=False
 bUseRSync=False
 BundleDisplayName=[PROJECT_NAME]
 BundleName=[PROJECT_NAME]
diff --git a/Engine/Source/Programs/AutomationTool/ProjectParams.cs b/Engine/Source/Programs/AutomationTool/ProjectParams.cs
--- a/Engine/Source/Programs/AutomationTool/ProjectParams.cs
+++ b/Engine/Source/Programs/AutomationTool/ProjectParams.cs
@@ -1335,7 +1335,7 @@ namespace AutomationTool
 				ProjectGameExePath = null;
 			}
 
-			var Properties = ProjectUtils.GetProjectProperties(RawProjectPath);
+			var Properties = ProjectUtils.GetProjectProperties(RawProjectPath, ClientTargetPlatforms);
 
 			bUsesSteam = Properties.bUsesSteam;
 			bUsesCEF3 = Properties.bUsesCEF3;
@@ -1491,7 +1491,14 @@ namespace AutomationTool
 
 			if (String.IsNullOrEmpty(EditorTarget) && ProjectType != TargetRules.TargetType.Program && CommandUtils.IsNullOrEmpty(EditorTargetsList) && !Rocket)
 			{
-				throw new AutomationException("Editor target not found!");
+				if (Properties.bWasGenerated)
+				{
+					EditorTarget = "UE4Editor";
+				}
+				else
+				{
+					throw new AutomationException("Editor target not found!");
+				}
 			}
 			if (String.IsNullOrEmpty(GameTarget) && Run && !NoClient && (Cook || CookOnTheFly) && CommandUtils.IsNullOrEmpty(ClientCookedTargetsList) && !Rocket)
 			{
@@ -1500,7 +1507,7 @@ namespace AutomationTool
 
 			if (EditorTargetsList == null)
 			{
-				if (!GlobalCommandLine.NoCompile && !GlobalCommandLine.NoCompileEditor && (ProjectType != TargetRules.TargetType.Program))
+				if (!GlobalCommandLine.NoCompile && !GlobalCommandLine.NoCompileEditor && (ProjectType != TargetRules.TargetType.Program) && !String.IsNullOrEmpty(EditorTarget))
 				{
 					EditorTargetsList = new ParamList<string>(EditorTarget);
 				}
diff --git a/Engine/Source/Programs/AutomationTool/ProjectUtils.cs b/Engine/Source/Programs/AutomationTool/ProjectUtils.cs
--- a/Engine/Source/Programs/AutomationTool/ProjectUtils.cs
+++ b/Engine/Source/Programs/AutomationTool/ProjectUtils.cs
@@ -81,6 +81,11 @@ namespace AutomationTool
 		/// </summary>
 		public List<SingleTargetProperties> Programs = new List<SingleTargetProperties>();
 
+		/// <summary>
+		/// Specifies if the target files were generated
+		/// </summary>
+		public bool bWasGenerated = false;
+
 		internal ProjectProperties()
 		{
 		}
@@ -109,7 +114,7 @@ namespace AutomationTool
 		/// </summary>
 		/// <param name="RawProjectPath">Full project path.</param>
 		/// <returns>Properties of the project.</returns>
-		public static ProjectProperties GetProjectProperties(string RawProjectPath)
+		public static ProjectProperties GetProjectProperties(string RawProjectPath, List<UnrealTargetPlatform> ClientTargetPlatforms = null)
 		{
 			string ProjectKey = "UE4";
 			if (!String.IsNullOrEmpty(RawProjectPath))
@@ -119,7 +124,7 @@ namespace AutomationTool
 			ProjectProperties Properties;
 			if (PropertiesCache.TryGetValue(ProjectKey, out Properties) == false)
 			{
-				Properties = DetectProjectProperties(RawProjectPath);
+				Properties = DetectProjectProperties(RawProjectPath, ClientTargetPlatforms);
 				PropertiesCache.Add(ProjectKey, Properties);
 			}
 			return Properties;
@@ -132,7 +137,7 @@ namespace AutomationTool
 		/// <returns>True if the project is a UProject file with source code.</returns>
 		public static bool IsCodeBasedUProjectFile(string RawProjectPath)
 		{
-			return GetProjectProperties(RawProjectPath).bIsCodeBasedProject;
+			return GetProjectProperties(RawProjectPath, null).bIsCodeBasedProject;
 		}
 
 		/// <summary>
@@ -150,19 +155,82 @@ namespace AutomationTool
 			return ProjectClientBinariesPath;
 		}
 
+		private static bool RequiresTempTarget(string RawProjectPath, List<UnrealTargetPlatform> ClientTargetPlatforms)
+		{
+			if (ClientTargetPlatforms == null || ClientTargetPlatforms.Count < 1)
+			{
+				foreach (UnrealTargetPlatform TargetPlatformType in Enum.GetValues(typeof(UnrealTargetPlatform)))
+				{
+					if (TargetPlatformType != UnrealTargetPlatform.Unknown)
+					{
+						IUEBuildPlatform BuildPlat = UEBuildPlatform.GetBuildPlatform(TargetPlatformType, true);
+						if (BuildPlat != null && !(BuildPlat as UEBuildPlatform).HasDefaultBuildConfig(TargetPlatformType, Path.GetDirectoryName(RawProjectPath)))
+						{
+							return true;
+						}
+					}
+				}
+			}
+			else
+			{
+				foreach (UnrealTargetPlatform TargetPlatformType in ClientTargetPlatforms)
+				{
+					IUEBuildPlatform BuildPlat = UEBuildPlatform.GetBuildPlatform(TargetPlatformType, true);
+					if (BuildPlat != null && !(BuildPlat as UEBuildPlatform).HasDefaultBuildConfig(TargetPlatformType, Path.GetDirectoryName(RawProjectPath)))
+					{
+						return true;
+					}
+				}
+			}
+			return false;
+		}
+
+		private static void GenerateTempTarget(string RawProjectPath)
+		{
+			// read in the template target cs file
+			var TempCSFile = CommandUtils.CombinePaths(CommandUtils.CmdEnv.LocalRoot, "Engine", "Build", "Target.cs.template");
+			string TargetCSFile = File.ReadAllText(TempCSFile);
+
+			// replace {GAME_NAME} with the game name
+			TargetCSFile = TargetCSFile.Replace("{GAME_NAME}", Path.GetFileNameWithoutExtension(RawProjectPath));
+
+			// write out the file in a new Source directory
+			string FileName = CommandUtils.CombinePaths(Path.GetDirectoryName(RawProjectPath), "Intermediate", "Source", Path.GetFileNameWithoutExtension(RawProjectPath) + ".Target.cs");
+			if (!Directory.Exists(Path.GetDirectoryName(FileName)))
+			{
+				Directory.CreateDirectory(Path.GetDirectoryName(FileName));
+			}
+
+			File.WriteAllText(FileName, TargetCSFile);
+		}
+
 		/// <summary>
 		/// Attempts to autodetect project properties.
 		/// </summary>
 		/// <param name="RawProjectPath">Full project path.</param>
 		/// <returns>Project properties.</returns>
-		private static ProjectProperties DetectProjectProperties(string RawProjectPath)
+		private static ProjectProperties DetectProjectProperties(string RawProjectPath, List<UnrealTargetPlatform> ClientTargetPlatforms)
 		{
 			var Properties = new ProjectProperties();
 			Properties.RawProjectPath = RawProjectPath;
 
+			// detect if the project is content only, but has non-default build settings
+			List<string> ExtraSearchPaths = null;
+			if (RequiresTempTarget(RawProjectPath, ClientTargetPlatforms))
+			{
+				GenerateTempTarget(RawProjectPath);
+				Properties.bWasGenerated = true;
+				ExtraSearchPaths = new List<string>();
+				ExtraSearchPaths.Add(CommandUtils.CombinePaths(Path.GetDirectoryName(RawProjectPath), "Intermediate", "Source"));
+			}
+			else if (File.Exists(Path.Combine(Path.GetDirectoryName(RawProjectPath), "Intermediate", "Source", Path.GetFileNameWithoutExtension(RawProjectPath) + ".Target.cs")))
+			{
+				File.Delete(Path.Combine(Path.GetDirectoryName(RawProjectPath), "Intermediate", "Source", Path.GetFileNameWithoutExtension(RawProjectPath) + ".Target.cs"));
+			}
+
 			if (CommandUtils.CmdEnv.HasCapabilityToCompile)
 			{
-				DetectTargetsForProject(Properties);
+				DetectTargetsForProject(Properties, ExtraSearchPaths);
 				Properties.bIsCodeBasedProject = !CommandUtils.IsNullOrEmpty(Properties.Targets) || !CommandUtils.IsNullOrEmpty(Properties.Programs);
 			}
 			else
@@ -173,14 +241,14 @@ namespace AutomationTool
 					throw new AutomationException("Cannot dtermine engine targets if we can't compile.");
 				}
 
-				Properties.bIsCodeBasedProject = false;
+				Properties.bIsCodeBasedProject = Properties.bWasGenerated;
 				// if there's a Source directory with source code in it, then mark us as having source code
 				string SourceDir = CommandUtils.CombinePaths(Path.GetDirectoryName(RawProjectPath), "Source");
 				if (Directory.Exists(SourceDir))
 				{
 					string[] CppFiles = Directory.GetFiles(SourceDir, "*.cpp", SearchOption.AllDirectories);
 					string[] HFiles = Directory.GetFiles(SourceDir, "*.h", SearchOption.AllDirectories);
-					Properties.bIsCodeBasedProject = CppFiles.Length > 0 || HFiles.Length > 0;
+					Properties.bIsCodeBasedProject |= (CppFiles.Length > 0 || HFiles.Length > 0);
 				}
 			}
 
diff --git a/Engine/Source/Programs/AutomationTool/Scripts/BuildProjectCommand.Automation.cs b/Engine/Source/Programs/AutomationTool/Scripts/BuildProjectCommand.Automation.cs
--- a/Engine/Source/Programs/AutomationTool/Scripts/BuildProjectCommand.Automation.cs
+++ b/Engine/Source/Programs/AutomationTool/Scripts/BuildProjectCommand.Automation.cs
@@ -73,7 +73,7 @@ public partial class Project : CommandUtils
 				foreach (var ClientPlatform in Params.ClientTargetPlatforms)
 				{
 					CrashReportPlatforms.Add(ClientPlatform);
-					Agenda.AddTargets(Params.ClientCookedTargets.ToArray(), ClientPlatform, BuildConfig, Params.RawProjectPath);
+					Agenda.AddTargets(Params.ClientCookedTargets.ToArray(), ClientPlatform, BuildConfig, Params.CodeBasedUprojectPath);
 				}
 			}
 		}
@@ -84,7 +84,7 @@ public partial class Project : CommandUtils
 				foreach (var ServerPlatform in Params.ServerTargetPlatforms)
 				{
 					CrashReportPlatforms.Add(ServerPlatform);
-					Agenda.AddTargets(Params.ServerCookedTargets.ToArray(), ServerPlatform, BuildConfig, Params.RawProjectPath);
+					Agenda.AddTargets(Params.ServerCookedTargets.ToArray(), ServerPlatform, BuildConfig, Params.CodeBasedUprojectPath);
 				}
 			}
 		}
diff --git a/Engine/Source/Programs/AutomationTool/Scripts/CopyBuildToStagingDirectory.Automation.cs b/Engine/Source/Programs/AutomationTool/Scripts/CopyBuildToStagingDirectory.Automation.cs
--- a/Engine/Source/Programs/AutomationTool/Scripts/CopyBuildToStagingDirectory.Automation.cs
+++ b/Engine/Source/Programs/AutomationTool/Scripts/CopyBuildToStagingDirectory.Automation.cs
@@ -450,7 +450,7 @@ public partial class Project : CommandUtils
 	{
         // const string Iso8601DateTimeFormat = "yyyy-MM-ddTHH:mm:ssZ"; // probably should work
 		// const string Iso8601DateTimeFormat = "o"; // predefined universal Iso standard format (has too many millisecond spaces for our read code in FDateTime.ParseISO8601
-        const string Iso8601DateTimeFormat = "yyyy'-'MM'-'dd'T'HH':'mm':'ss'.'fffzzz";
+        const string Iso8601DateTimeFormat = "yyyy'-'MM'-'dd'T'HH':'mm':'ss'.'fffZ";
 
         
         if (Mapping.Count > 0)
@@ -881,7 +881,7 @@ public partial class Project : CommandUtils
 
 	private static string GetIntermediateCommandlineDir(DeploymentContext SC)
 	{
-		return CombinePaths(SC.LocalRoot, "Engine/Intermediate/UAT", SC.FinalCookPlatform);
+		return CombinePaths(SC.ProjectRoot, "Intermediate/UAT", SC.FinalCookPlatform);
 	}
 
 	public static void WriteStageCommandline(string IntermediateCmdLineFile, ProjectParams Params, DeploymentContext SC)
diff --git a/Engine/Source/Programs/IOS/iPhonePackager/CompileTime.cs b/Engine/Source/Programs/IOS/iPhonePackager/CompileTime.cs
--- a/Engine/Source/Programs/IOS/iPhonePackager/CompileTime.cs
+++ b/Engine/Source/Programs/IOS/iPhonePackager/CompileTime.cs
@@ -251,8 +251,12 @@ namespace iPhonePackager
 			// make sure this .mobileprovision file is newer than any other .mobileprovision file on the Mac (this file gets multiple games named the same file, 
 			// so the time stamp checking can fail when moving between games, a la the buildmachines!)
 			File.SetLastWriteTime(FinalMobileProvisionFilename, DateTime.UtcNow);
-
-			FileOperations.CopyRequiredFile(Config.RootRelativePath + @"Engine\Intermediate\IOS\UE4.xcodeproj\project.pbxproj", Path.Combine(Config.PCXcodeStagingDir, @"project.pbxproj.datecheck"));
+			string ProjectFile = Config.RootRelativePath + @"Engine\Intermediate\IOS\UE4.xcodeproj\project.pbxproj";
+			if (Program.GameName != "UE4Game")
+			{
+				ProjectFile = Config.IntermediateDirectory + @"\" + Program.GameName + @".xcodeproj\project.pbxproj";
+			}
+			FileOperations.CopyRequiredFile(ProjectFile, Path.Combine(Config.PCXcodeStagingDir, @"project.pbxproj.datecheck"));
 			
 			// needs Mac line endings so it can be executed
 			string SrcPath = @"..\..\..\Build\IOS\XcodeSupportFiles\prepackage.sh";
diff --git a/Engine/Source/Programs/IOS/iPhonePackager/Config.cs b/Engine/Source/Programs/IOS/iPhonePackager/Config.cs
--- a/Engine/Source/Programs/IOS/iPhonePackager/Config.cs
+++ b/Engine/Source/Programs/IOS/iPhonePackager/Config.cs
@@ -447,8 +447,8 @@ namespace iPhonePackager
 			}
 
 			// Root directory on PC for staging files to copy to Mac
-			Config.PCStagingRootDir = String.Format(@"{0}\Engine\Intermediate\IOS-Deploy\{1}\{2}{3}\",
-				RootRelativePath,
+			Config.PCStagingRootDir = String.Format(@"{0}-Deploy\{1}\{2}{3}\",
+				IntermediateDirectory,
 				Program.GameName,
 				Program.GameConfiguration,
 				Program.Architecture);
diff --git a/Engine/Source/Programs/UnrealBuildTool/Configuration/UEBuildPlatform.cs b/Engine/Source/Programs/UnrealBuildTool/Configuration/UEBuildPlatform.cs
--- a/Engine/Source/Programs/UnrealBuildTool/Configuration/UEBuildPlatform.cs
+++ b/Engine/Source/Programs/UnrealBuildTool/Configuration/UEBuildPlatform.cs
@@ -189,6 +189,13 @@ namespace UnrealBuildTool
 		 */
 		void SetUpConfigurationEnvironment(UEBuildTarget InBuildTarget);
 
+		/**
+		 * Setup the project environment for building
+		 * 
+		 * @param InBuildTarget The target being built
+		 */
+		void SetUpProjectEnvironment(UnrealTargetPlatform InPlatform);
+
 		/**
 		 * Whether this platform should create debug information or not
 		 * 
@@ -261,6 +268,13 @@ namespace UnrealBuildTool
 		 * @param InBuildTarget The target being built
 		 */
 		void SetupBinaries(UEBuildTarget InBuildTarget);
+
+		/**
+		 * Check for the default configuration
+		 *
+		 * return true if the project uses the default build config
+		 */
+		bool HasDefaultBuildConfig(UnrealTargetPlatform InPlatform, string InProjectPath);
 	}
 
 	public abstract partial class UEBuildPlatform : IUEBuildPlatform
@@ -746,6 +760,93 @@ namespace UnrealBuildTool
 			InBuildTarget.GlobalLinkEnvironment.Config.bCreateDebugInfo = InBuildTarget.GlobalCompileEnvironment.Config.bCreateDebugInfo;
 		}
 
+		/**
+		 *	Setup the project environment for building
+		 *	
+		 *	@param	InBuildTarget		The target being built
+		 */
+		public virtual void SetUpProjectEnvironment(UnrealTargetPlatform InPlatform)
+		{
+			ConfigCacheIni Ini = new ConfigCacheIni(InPlatform, "Engine", UnrealBuildTool.GetUProjectPath());
+			bool bValue = UEBuildConfiguration.bCompileAPEX;
+			if (Ini.GetBool("/Script/BuildSettings.BuildSettings", "bCompileApex", out bValue))
+			{
+				UEBuildConfiguration.bCompileAPEX = bValue;
+			}
+
+			bValue = UEBuildConfiguration.bCompileBox2D;
+			if (Ini.GetBool("/Script/BuildSettings.BuildSettings", "bCompileBox2D", out bValue))
+			{
+				UEBuildConfiguration.bCompileBox2D = bValue;
+			}
+
+			bValue = UEBuildConfiguration.bCompileICU;
+			if (Ini.GetBool("/Script/BuildSettings.BuildSettings", "bCompileICU", out bValue))
+			{
+				UEBuildConfiguration.bCompileICU = bValue;
+			}
+
+			bValue = UEBuildConfiguration.bCompileSimplygon;
+			if (Ini.GetBool("/Script/BuildSettings.BuildSettings", "bCompileSimplygon", out bValue))
+			{
+				UEBuildConfiguration.bCompileSimplygon = bValue;
+			}
+
+			bValue = UEBuildConfiguration.bCompileLeanAndMeanUE;
+			if (Ini.GetBool("/Script/BuildSettings.BuildSettings", "bCompileLeanAndMeanUE", out bValue))
+			{
+				UEBuildConfiguration.bCompileLeanAndMeanUE = bValue;
+			}
+
+			bValue = UEBuildConfiguration.bIncludeADO;
+			if (Ini.GetBool("/Script/BuildSettings.BuildSettings", "bIncludeADO", out bValue))
+			{
+				UEBuildConfiguration.bIncludeADO = bValue;
+			}
+
+			bValue = UEBuildConfiguration.bCompileRecast;
+			if (Ini.GetBool("/Script/BuildSettings.BuildSettings", "bCompileRecast", out bValue))
+			{
+				UEBuildConfiguration.bCompileRecast = bValue;
+			}
+
+			bValue = UEBuildConfiguration.bCompileSpeedTree;
+			if (Ini.GetBool("/Script/BuildSettings.BuildSettings", "bCompileSpeedTree", out bValue))
+			{
+				UEBuildConfiguration.bCompileSpeedTree = bValue;
+			}
+
+			bValue = UEBuildConfiguration.bCompileWithPluginSupport;
+			if (Ini.GetBool("/Script/BuildSettings.BuildSettings", "bCompileWithPluginSupport", out bValue))
+			{
+				UEBuildConfiguration.bCompileWithPluginSupport = bValue;
+			}
+
+			bValue = UEBuildConfiguration.bCompilePhysXVehicle;
+			if (Ini.GetBool("/Script/BuildSettings.BuildSettings", "bCompilePhysXVehicle", out bValue))
+			{
+				UEBuildConfiguration.bCompilePhysXVehicle = bValue;
+			}
+
+			bValue = UEBuildConfiguration.bCompileFreeType;
+			if (Ini.GetBool("/Script/BuildSettings.BuildSettings", "bCompileFreeType", out bValue))
+			{
+				UEBuildConfiguration.bCompileFreeType = bValue;
+			}
+
+			bValue = UEBuildConfiguration.bCompileForSize;
+			if (Ini.GetBool("/Script/BuildSettings.BuildSettings", "bCompileForSize", out bValue))
+			{
+				UEBuildConfiguration.bCompileForSize = bValue;
+			}
+
+			bValue = UEBuildConfiguration.bCompileCEF3;
+			if (Ini.GetBool("/Script/BuildSettings.BuildSettings", "bCompileCEF3", out bValue))
+			{
+				UEBuildConfiguration.bCompileCEF3 = bValue;
+			}
+		}
+
 		/**
 		 *	Whether this platform should create debug information or not
 		 *	
@@ -851,6 +952,135 @@ namespace UnrealBuildTool
 		public virtual void SetupBinaries(UEBuildTarget InBuildTarget)
 		{
 		}
+
+		/**
+		 * Check for the default configuration
+		 *
+		 * return true if the project uses the default build config
+		 */
+		public virtual bool HasDefaultBuildConfig(UnrealTargetPlatform Platform, string ProjectPath)
+		{
+			ConfigCacheIni ProjIni = new ConfigCacheIni(Platform, "Engine", ProjectPath);
+			ConfigCacheIni DefaultIni = new ConfigCacheIni(Platform, "Engine", null);
+
+			bool bDefValue = UEBuildConfiguration.bCompileAPEX;
+			DefaultIni.GetBool("/Script/BuildSettings.BuildSettings", "bCompileApex", out bDefValue);
+			bool bProjValue = UEBuildConfiguration.bCompileAPEX;
+			ProjIni.GetBool("/Script/BuildSettings.BuildSettings", "bCompileApex", out bProjValue);
+			if (bDefValue != bProjValue)
+			{
+				return false;
+			}
+
+			bDefValue = UEBuildConfiguration.bCompileBox2D;
+			DefaultIni.GetBool("/Script/BuildSettings.BuildSettings", "bCompileBox2D", out bDefValue);
+			bProjValue = UEBuildConfiguration.bCompileBox2D;
+			ProjIni.GetBool("/Script/BuildSettings.BuildSettings", "bCompileBox2D", out bProjValue);
+			if (bDefValue != bProjValue)
+			{
+				return false;
+			}
+
+			bDefValue = UEBuildConfiguration.bCompileICU;
+			DefaultIni.GetBool("/Script/BuildSettings.BuildSettings", "bCompileICU", out bDefValue);
+			bProjValue = UEBuildConfiguration.bCompileICU;
+			ProjIni.GetBool("/Script/BuildSettings.BuildSettings", "bCompileICU", out bProjValue);
+			if (bDefValue != bProjValue)
+			{
+				return false;
+			}
+
+			bDefValue = UEBuildConfiguration.bCompileSimplygon;
+			DefaultIni.GetBool("/Script/BuildSettings.BuildSettings", "bCompileSimplygon", out bDefValue);
+			bProjValue = UEBuildConfiguration.bCompileSimplygon;
+			ProjIni.GetBool("/Script/BuildSettings.BuildSettings", "bCompileSimplygon", out bProjValue);
+			if (bDefValue != bProjValue)
+			{
+				return false;
+			}
+
+			bDefValue = UEBuildConfiguration.bCompileLeanAndMeanUE;
+			DefaultIni.GetBool("/Script/BuildSettings.BuildSettings", "bCompileLeanAndMeanUE", out bDefValue);
+			bProjValue = UEBuildConfiguration.bCompileLeanAndMeanUE;
+			ProjIni.GetBool("/Script/BuildSettings.BuildSettings", "bCompileLeanAndMeanUE", out bProjValue);
+			if (bDefValue != bProjValue)
+			{
+				return false;
+			}
+
+			bDefValue = UEBuildConfiguration.bIncludeADO;
+			DefaultIni.GetBool("/Script/BuildSettings.BuildSettings", "bIncludeADO", out bDefValue);
+			bProjValue = UEBuildConfiguration.bIncludeADO;
+			ProjIni.GetBool("/Script/BuildSettings.BuildSettings", "bIncludeADO", out bProjValue);
+			if (bDefValue != bProjValue)
+			{
+				return false;
+			}
+
+			bDefValue = UEBuildConfiguration.bCompileRecast;
+			DefaultIni.GetBool("/Script/BuildSettings.BuildSettings", "bCompileRecast", out bDefValue);
+			bProjValue = UEBuildConfiguration.bCompileRecast;
+			ProjIni.GetBool("/Script/BuildSettings.BuildSettings", "bCompileRecast", out bProjValue);
+			if (bDefValue != bProjValue)
+			{
+				return false;
+			}
+
+			bDefValue = UEBuildConfiguration.bCompileSpeedTree;
+			DefaultIni.GetBool("/Script/BuildSettings.BuildSettings", "bCompileSpeedTree", out bDefValue);
+			bProjValue = UEBuildConfiguration.bCompileSpeedTree;
+			ProjIni.GetBool("/Script/BuildSettings.BuildSettings", "bCompileSpeedTree", out bProjValue);
+			if (bDefValue != bProjValue)
+			{
+				return false;
+			}
+
+			bDefValue = UEBuildConfiguration.bCompileWithPluginSupport;
+			DefaultIni.GetBool("/Script/BuildSettings.BuildSettings", "bCompileWithPluginSupport", out bDefValue);
+			bProjValue = UEBuildConfiguration.bCompileWithPluginSupport;
+			ProjIni.GetBool("/Script/BuildSettings.BuildSettings", "bCompileWithPluginSupport", out bProjValue);
+			if (bDefValue != bProjValue)
+			{
+				return false;
+			}
+
+			bDefValue = UEBuildConfiguration.bCompilePhysXVehicle;
+			DefaultIni.GetBool("/Script/BuildSettings.BuildSettings", "bCompilePhysXVehicle", out bDefValue);
+			bProjValue = UEBuildConfiguration.bCompilePhysXVehicle;
+			ProjIni.GetBool("/Script/BuildSettings.BuildSettings", "bCompilePhysXVehicle", out bProjValue);
+			if (bDefValue != bProjValue)
+			{
+				return false;
+			}
+
+			bDefValue = UEBuildConfiguration.bCompileFreeType;
+			DefaultIni.GetBool("/Script/BuildSettings.BuildSettings", "bCompileFreeType", out bDefValue);
+			bProjValue = UEBuildConfiguration.bCompileFreeType;
+			ProjIni.GetBool("/Script/BuildSettings.BuildSettings", "bCompileFreeType", out bProjValue);
+			if (bDefValue != bProjValue)
+			{
+				return false;
+			}
+
+			bDefValue = UEBuildConfiguration.bCompileForSize;
+			DefaultIni.GetBool("/Script/BuildSettings.BuildSettings", "bCompileForSize", out bDefValue);
+			bProjValue = UEBuildConfiguration.bCompileForSize;
+			ProjIni.GetBool("/Script/BuildSettings.BuildSettings", "bCompileForSize", out bProjValue);
+			if (bDefValue != bProjValue)
+			{
+				return false;
+			}
+
+			bDefValue = UEBuildConfiguration.bCompileCEF3;
+			DefaultIni.GetBool("/Script/BuildSettings.BuildSettings", "bCompileCEF3", out bDefValue);
+			bProjValue = UEBuildConfiguration.bCompileCEF3;
+			ProjIni.GetBool("/Script/BuildSettings.BuildSettings", "bCompileCEF3", out bProjValue);
+			if (bDefValue != bProjValue)
+			{
+				return false;
+			}
+			return true;
+		}
 	}
 
 	// AutoSDKs handling portion
diff --git a/Engine/Source/Programs/UnrealBuildTool/Configuration/UEBuildTarget.cs b/Engine/Source/Programs/UnrealBuildTool/Configuration/UEBuildTarget.cs
--- a/Engine/Source/Programs/UnrealBuildTool/Configuration/UEBuildTarget.cs
+++ b/Engine/Source/Programs/UnrealBuildTool/Configuration/UEBuildTarget.cs
@@ -584,7 +584,7 @@ namespace UnrealBuildTool
 
 
 		/** 
-		 * Look for all folders ending in "Game" in the branch root that have a "Source" subfolder
+		 * Look for all folders with a uproject file, these are valid games
 		 * This is defined as a valid game
 		 */
 		public static List<string> DiscoverAllGameFolders()
@@ -2772,9 +2772,10 @@ namespace UnrealBuildTool
 			// tell the compiled code the name of the UBT platform (this affects folder on disk, etc that the game may need to know)
 			GlobalCompileEnvironment.Config.Definitions.Add("UBT_COMPILED_PLATFORM=" + Platform.ToString());
 
-			// Initialize the compile and link environments for the platform and configuration.
+			// Initialize the compile and link environments for the platform, configuration, and project.
 			SetUpPlatformEnvironment();
 			SetUpConfigurationEnvironment();
+			SetUpProjectEnvironment();
 
 			// Validates current settings and updates if required.
 			BuildConfiguration.ValidateConfiguration(
@@ -2815,6 +2816,11 @@ namespace UnrealBuildTool
 			GlobalLinkEnvironment.Config.bIsBuildingLibrary = bIsBuildingLibrary;
 		}
 
+		void SetUpProjectEnvironment()
+		{
+			UEBuildPlatform.GetBuildPlatform(Platform).SetUpProjectEnvironment(Platform);
+		}
+
 		/** Constructs a TargetInfo object for this target. */
 		public TargetInfo GetTargetInfo()
 		{
diff --git a/Engine/Source/Programs/UnrealBuildTool/IOS/IOSToolChain.cs b/Engine/Source/Programs/UnrealBuildTool/IOS/IOSToolChain.cs
--- a/Engine/Source/Programs/UnrealBuildTool/IOS/IOSToolChain.cs
+++ b/Engine/Source/Programs/UnrealBuildTool/IOS/IOSToolChain.cs
@@ -41,26 +41,14 @@ namespace UnrealBuildTool
 		public static string IOSSDKVersion = "latest";
 		public static float IOSSDKVersionFloat = 0.0f;
 
-		/** The architecture(s) to compile */
-		[XmlConfig]
-		public static string NonShippingArchitectures = "armv7";
-		[XmlConfig]
-		public static string ShippingArchitectures = "armv7,arm64";
-
 		/** Which version of the iOS to allow at build time */
 		[XmlConfig]
 		public static string BuildIOSVersion = "6.1";
 
-		/** Which version of the iOS to allow at run time */
-		public static string RunTimeIOSVersion = "6.1";
-
 		/** Which developer directory to root from */
 		[XmlConfig]
 		public static string XcodeDeveloperDir = "/Applications/Xcode.app/Contents/Developer/";
 
-		/** which devices the game is allowed to run on */
-		public static string RunTimeIOSDevices = "1,2";
-
 		/** Location of the SDKs */
 		private static string BaseSDKDir;
 		private static string BaseSDKDirSim;
@@ -89,111 +77,11 @@ namespace UnrealBuildTool
 			BaseSDKDirSim = XcodeDeveloperDir + "Platforms/iPhoneSimulator.platform/Developer/SDKs";
 		}
 
-		public override void ParseProjectSettings()
-		{
-			base.ParseProjectSettings();
-
-			// look in ini settings for what platforms to compile for
-			ConfigCacheIni Ini = new ConfigCacheIni(UnrealTargetPlatform.IOS, "Engine", UnrealBuildTool.GetUProjectPath());
-			string MinVersion = "IOS_6";
-			if (Ini.GetString("/Script/IOSRuntimeSettings.IOSRuntimeSettings", "MinimumiOSVersion", out MinVersion))
-			{
-				switch (MinVersion)
-				{
-					case "IOS_61":
-						RunTimeIOSVersion = "6.1";
-						break;
-					case "IOS_7":
-						RunTimeIOSVersion = "7.0";
-						break;
-					case "IOS_8":
-						RunTimeIOSVersion = "8.0";
-						break;
-				}
-			}
-
-			bool biPhoneAllowed = true;
-			bool biPadAllowed = true;
-			Ini.GetBool ("/Script/IOSRuntimeSettings.IOSRuntimeSettings", "bSupportsIPhone", out biPhoneAllowed);
-			Ini.GetBool ("/Script/IOSRuntimeSettings.IOSRuntimeSettings", "bSupportsIPad", out biPadAllowed);
-			if (biPhoneAllowed && biPadAllowed)
-			{
-				RunTimeIOSDevices = "1,2";
-			}
-			else if (biPadAllowed)
-			{
-				RunTimeIOSDevices = "2";
-			}
-			else if (biPhoneAllowed)
-			{
-				RunTimeIOSDevices = "1";
-			}
-		}
-
-		public static void ParseArchitectures()
-		{
-			// look in ini settings for what platforms to compile for
-			ConfigCacheIni Ini = new ConfigCacheIni(UnrealTargetPlatform.Android, "Engine", UnrealBuildTool.GetUProjectPath());
-			List<string> ProjectArches = new List<string>();
-			bool bBuild = true;
-			if (Ini.GetBool("/Script/IOSRuntimeSettings.IOSRuntimeSettings", "bDevForArmV7", out bBuild) && bBuild)
-			{
-				ProjectArches.Add("armv7");
-			}
-			if (Ini.GetBool("/Script/IOSRuntimeSettings.IOSRuntimeSettings", "bDevForArm64", out bBuild) && bBuild)
-			{
-				ProjectArches.Add("arm64");
-			}
-			if (Ini.GetBool("/Script/IOSRuntimeSettings.IOSRuntimeSettings", "bDevForArmV7S", out bBuild) && bBuild)
-			{
-				ProjectArches.Add("armv7s");
-			}
-
-			// force armv7 if something went wrong
-			if (ProjectArches.Count == 0)
-			{
-				ProjectArches.Add("armv7");
-			}
-			NonShippingArchitectures = ProjectArches[0];
-			for (int Index = 1; Index < ProjectArches.Count; ++Index)
-			{
-				NonShippingArchitectures += "," + ProjectArches[Index];
-			}
-
-			ProjectArches.Clear();
-			if (Ini.GetBool("/Script/IOSRuntimeSettings.IOSRuntimeSettings", "bShipForArmV7", out bBuild) && bBuild)
-			{
-				ProjectArches.Add("armv7");
-			}
-			if (Ini.GetBool("/Script/IOSRuntimeSettings.IOSRuntimeSettings", "bShipForArm64", out bBuild) && bBuild)
-			{
-				ProjectArches.Add("arm64");
-			}
-			if (Ini.GetBool("/Script/IOSRuntimeSettings.IOSRuntimeSettings", "bShipForArmV7S", out bBuild) && bBuild)
-			{
-				ProjectArches.Add("armv7s");
-			}
-
-			// force armv7 if something went wrong
-			if (ProjectArches.Count == 0)
-			{
-				ProjectArches.Add("armv7");
-				ProjectArches.Add("arm64");
-			}
-			ShippingArchitectures = ProjectArches[0];
-			for (int Index = 1; Index < ProjectArches.Count; ++Index)
-			{
-				ShippingArchitectures += "," + ProjectArches[Index];
-			}
-		}
-
 		/** Hunt down the latest IOS sdk if desired */
 		public override void SetUpGlobalEnvironment()
 		{
 			base.SetUpGlobalEnvironment();
 
-			ParseArchitectures();
-
 			if (IOSSDKVersion == "latest")
 			{
 				try
@@ -294,11 +182,13 @@ namespace UnrealBuildTool
 		static bool bHasPrinted = false;
 		static string GetArchitectureArgument(CPPTargetConfiguration Configuration, string UBTArchitecture)
 		{
+			IOSPlatform BuildPlat = UEBuildPlatform.GetBuildPlatform(UnrealTargetPlatform.IOS) as IOSPlatform;
+			BuildPlat.SetUpProjectEnvironment(UnrealTargetPlatform.IOS);
 
 			// get the list of architectures to compile
 			string Archs =
 				UBTArchitecture == "-simulator" ? "i386" :
-				(Configuration == CPPTargetConfiguration.Shipping) ? ShippingArchitectures : NonShippingArchitectures;
+				(Configuration == CPPTargetConfiguration.Shipping) ? IOSPlatform.ShippingArchitectures : IOSPlatform.NonShippingArchitectures;
 
 			if (!bHasPrinted)
 			{
@@ -408,7 +298,7 @@ namespace UnrealBuildTool
 			Result += " -fobjc-abi-version=2";
 			Result += " -fobjc-legacy-dispatch";
 			Result += " -std=c++11";
-			Result += " -stdlib=libstdc++";
+			Result += " -stdlib=libc++";
 			return Result;
 		}
 
@@ -420,7 +310,7 @@ namespace UnrealBuildTool
 			Result += " -fobjc-legacy-dispatch";
 			Result += " -fno-rtti";
 			Result += " -std=c++11";
-			Result += " -stdlib=libstdc++";
+			Result += " -stdlib=libc++";
 			return Result;
 		}
 
@@ -431,7 +321,7 @@ namespace UnrealBuildTool
 			Result += " -fobjc-abi-version=2";
 			Result += " -fobjc-legacy-dispatch";
 			Result += " -std=c++11";
-			Result += " -stdlib=libstdc++";
+			Result += " -stdlib=libc++";
 			return Result;
 		}
 
@@ -448,7 +338,7 @@ namespace UnrealBuildTool
 			Result += " -x objective-c++-header";
 			Result += " -fno-rtti";
 			Result += " -std=c++11";
-			Result += " -stdlib=libstdc++";
+			Result += " -stdlib=libc++";
 			return Result;
 		}
 
@@ -849,7 +739,8 @@ namespace UnrealBuildTool
 			}
 			else
 			{
-				string ResponsePath = Path.GetFullPath("..\\Intermediate\\IOS\\LinkFileList_" + Path.GetFileNameWithoutExtension(LinkEnvironment.Config.OutputFilePath) + ".tmp"); 
+				bool bIsUE4Game = LinkEnvironment.Config.OutputFilePath.Contains("UE4Game");
+				string ResponsePath = Path.GetFullPath(Path.Combine((!bIsUE4Game && !string.IsNullOrEmpty(UnrealBuildTool.GetUProjectPath())) ? UnrealBuildTool.GetUProjectPath() : BuildConfiguration.RelativeEnginePath, BuildConfiguration.PlatformIntermediateFolder, "LinkFileList_" + Path.GetFileNameWithoutExtension(LinkEnvironment.Config.OutputFilePath) + ".tmp")); 
 				if (!Utils.IsRunningOnMono && BuildHostPlatform.Current.Platform != UnrealTargetPlatform.Mac)
 				{
 					ResponseFile.Create (ResponsePath, InputFileNames);
@@ -891,7 +782,7 @@ namespace UnrealBuildTool
 		static public FileItem GenerateDebugInfo(FileItem Executable)
 		{
 			// Make a file item for the source and destination files
-			string FullDestPathRoot = Executable.AbsolutePath + ".app.dSYM";
+			string FullDestPathRoot = Executable.AbsolutePath + ".dSYM";
 			string FullDestPath = FullDestPathRoot;
 
 			FileItem DestFile;
@@ -918,14 +809,14 @@ namespace UnrealBuildTool
 			// note that the source and dest are switched from a copy command
 			if (!Utils.IsRunningOnMono && BuildHostPlatform.Current.Platform != UnrealTargetPlatform.Mac) 
 			{
-				GenDebugAction.CommandArguments = string.Format("-c '/usr/bin/dsymutil {0} -o {1}; cd {1}/..; zip -r -y -1 {2}.app.dSYM.zip {2}.app.dSYM'",
+				GenDebugAction.CommandArguments = string.Format("-c '/usr/bin/dsymutil \"{0}\" -f -o \"{1}\"; cd \"{1}/..\"; zip -r -y -1 {2}.dSYM.zip {2}.dSYM'",
 					Executable.AbsolutePath,
 					FullDestPathRoot,
 					Path.GetFileName(Executable.AbsolutePath));
 			} 
 			else
 			{
-				GenDebugAction.CommandArguments = string.Format("-c '/usr/bin/dsymutil {0} -o {1}'",
+				GenDebugAction.CommandArguments = string.Format("-c '/usr/bin/dsymutil \"{0}\" -f -o \"{1}\"'",
 					Executable.AbsolutePath,
 					FullDestPathRoot);
 			}
@@ -1307,7 +1198,7 @@ namespace UnrealBuildTool
 					if (!bUseDangerouslyFastMode)
 					{
 						// generate the dummy project so signing works
-						UnrealBuildTool.GenerateProjectFiles(new XcodeProjectFileGenerator(), new string[] { "-NoIntellisense", "-iosdeployonly" });
+						UnrealBuildTool.GenerateProjectFiles(new XcodeProjectFileGenerator(), new string[] { "-NoIntellisense", "-iosdeployonly", (!string.IsNullOrEmpty(UnrealBuildTool.GetUProjectPath()) ? "-game" : "") });
 					}
 
 					// now that 
@@ -1316,7 +1207,6 @@ namespace UnrealBuildTool
 					StubGenerateProcess.StartInfo.FileName = Path.Combine(StubGenerateProcess.StartInfo.WorkingDirectory, "iPhonePackager.exe");
 
 					string Arguments = "";
-						
 					string PathToApp = RulesCompiler.GetTargetFilename(AppName);
 
 					// right now, no programs have a Source subdirectory, so assume the PathToApp is directly in the root
@@ -1326,14 +1216,22 @@ namespace UnrealBuildTool
 					}
 					else
 					{
-						Int32 SourceIndex = PathToApp.LastIndexOf(@"\Source");
+						Int32 SourceIndex = PathToApp.LastIndexOf(@"\Intermediate\Source");
 						if (SourceIndex != -1)
 						{
 							PathToApp = PathToApp.Substring(0, SourceIndex);
 						}
 						else
 						{
-							throw new BuildException("The target was not in a /Source subdirectory");
+							SourceIndex = PathToApp.LastIndexOf(@"\Source");
+							if (SourceIndex != -1)
+							{
+								PathToApp = PathToApp.Substring(0, SourceIndex);
+							}
+							else
+							{
+								throw new BuildException("The target was not in a /Source subdirectory");
+							}
 						}
 					}
 
@@ -1450,16 +1348,6 @@ namespace UnrealBuildTool
 			}
 		}
 
-		public override string GetPlatformVersion()
-		{
-			return RunTimeIOSVersion;
-		}
-
-		public override string GetPlatformDevices()
-		{
-			return RunTimeIOSDevices;
-		}
-
 		public override UnrealTargetPlatform GetPlatform()
 		{
 			return UnrealTargetPlatform.IOS;
diff --git a/Engine/Source/Programs/UnrealBuildTool/IOS/UEBuildIOS.cs b/Engine/Source/Programs/UnrealBuildTool/IOS/UEBuildIOS.cs
--- a/Engine/Source/Programs/UnrealBuildTool/IOS/UEBuildIOS.cs
+++ b/Engine/Source/Programs/UnrealBuildTool/IOS/UEBuildIOS.cs
@@ -15,7 +15,29 @@ namespace UnrealBuildTool
         [XmlConfig]
         public static string IOSArchitecture = "";
 
-        // The current architecture - affects everything about how UBT operates on IOS
+		/** Which version of the iOS to allow at run time */
+		public static string RunTimeIOSVersion = "6.1";
+
+		/** which devices the game is allowed to run on */
+		public static string RunTimeIOSDevices = "1,2";
+
+		/** The architecture(s) to compile */
+		[XmlConfig]
+		public static string NonShippingArchitectures = "armv7";
+		[XmlConfig]
+		public static string ShippingArchitectures = "armv7,arm64";
+
+		public string GetRunTimeVersion()
+		{
+			return RunTimeIOSVersion;
+		}
+
+		public string GetRunTimeDevices()
+		{
+			return RunTimeIOSDevices;
+		}
+
+		// The current architecture - affects everything about how UBT operates on IOS
         public override string GetActiveArchitecture()
         {
             return IOSArchitecture;
@@ -104,7 +126,7 @@ namespace UnrealBuildTool
 
         public override string GetDebugInfoExtension(UEBuildBinaryType InBinaryType)
         {
-            return BuildConfiguration.bGeneratedSYMFile ? ".dsym" : "";
+            return BuildConfiguration.bGeneratedSYMFile ? ".dSYM" : "";
         }
 
         public override bool CanUseXGE()
@@ -117,6 +139,102 @@ namespace UnrealBuildTool
 			return true;
 		}
 
+		public override void SetUpProjectEnvironment(UnrealTargetPlatform InPlatform)
+		{
+			base.SetUpProjectEnvironment(InPlatform);
+
+			// update the configuration based on the project file
+			// look in ini settings for what platforms to compile for
+			ConfigCacheIni Ini = new ConfigCacheIni(InPlatform, "Engine", UnrealBuildTool.GetUProjectPath());
+			string MinVersion = "IOS_6";
+			if (Ini.GetString("/Script/IOSRuntimeSettings.IOSRuntimeSettings", "MinimumiOSVersion", out MinVersion))
+			{
+				switch (MinVersion)
+				{
+					case "IOS_61":
+						RunTimeIOSVersion = "6.1";
+						break;
+					case "IOS_7":
+						RunTimeIOSVersion = "7.0";
+						break;
+					case "IOS_8":
+						RunTimeIOSVersion = "8.0";
+						break;
+				}
+			}
+
+			bool biPhoneAllowed = true;
+			bool biPadAllowed = true;
+			Ini.GetBool("/Script/IOSRuntimeSettings.IOSRuntimeSettings", "bSupportsIPhone", out biPhoneAllowed);
+			Ini.GetBool("/Script/IOSRuntimeSettings.IOSRuntimeSettings", "bSupportsIPad", out biPadAllowed);
+			if (biPhoneAllowed && biPadAllowed)
+			{
+				RunTimeIOSDevices = "1,2";
+			}
+			else if (biPadAllowed)
+			{
+				RunTimeIOSDevices = "2";
+			}
+			else if (biPhoneAllowed)
+			{
+				RunTimeIOSDevices = "1";
+			}
+
+			List<string> ProjectArches = new List<string>();
+			bool bBuild = true;
+			if (Ini.GetBool("/Script/IOSRuntimeSettings.IOSRuntimeSettings", "bDevForArmV7", out bBuild) && bBuild)
+			{
+				ProjectArches.Add("armv7");
+			}
+			if (Ini.GetBool("/Script/IOSRuntimeSettings.IOSRuntimeSettings", "bDevForArm64", out bBuild) && bBuild)
+			{
+				ProjectArches.Add("arm64");
+			}
+			if (Ini.GetBool("/Script/IOSRuntimeSettings.IOSRuntimeSettings", "bDevForArmV7S", out bBuild) && bBuild)
+			{
+				ProjectArches.Add("armv7s");
+			}
+
+			// force armv7 if something went wrong
+			if (ProjectArches.Count == 0)
+			{
+				ProjectArches.Add("armv7");
+			}
+			NonShippingArchitectures = ProjectArches[0];
+			for (int Index = 1; Index < ProjectArches.Count; ++Index)
+			{
+				NonShippingArchitectures += "," + ProjectArches[Index];
+			}
+
+			ProjectArches.Clear();
+			if (Ini.GetBool("/Script/IOSRuntimeSettings.IOSRuntimeSettings", "bShipForArmV7", out bBuild) && bBuild)
+			{
+				ProjectArches.Add("armv7");
+			}
+			if (Ini.GetBool("/Script/IOSRuntimeSettings.IOSRuntimeSettings", "bShipForArm64", out bBuild) && bBuild)
+			{
+				ProjectArches.Add("arm64");
+			}
+			if (Ini.GetBool("/Script/IOSRuntimeSettings.IOSRuntimeSettings", "bShipForArmV7S", out bBuild) && bBuild)
+			{
+				ProjectArches.Add("armv7s");
+			}
+
+			// force armv7 if something went wrong
+			if (ProjectArches.Count == 0)
+			{
+				ProjectArches.Add("armv7");
+				ProjectArches.Add("arm64");
+			}
+			ShippingArchitectures = ProjectArches[0];
+			for (int Index = 1; Index < ProjectArches.Count; ++Index)
+			{
+				ShippingArchitectures += "," + ProjectArches[Index];
+			}
+
+			// determine if we need to generate the dsym
+			Ini.GetBool("/Script/IOSRuntimeSettings.IOSRuntimeSettings", "bGenerateSYMFile", out BuildConfiguration.bGeneratedSYMFile);
+		}
         /**
          *	Setup the target environment for building
          *	
@@ -209,13 +327,16 @@ namespace UnrealBuildTool
 
         public override void ValidateBuildConfiguration(CPPTargetConfiguration Configuration, CPPTargetPlatform Platform, bool bCreateDebugInfo)
         {
+			// check the base first
+			base.ValidateBuildConfiguration(Configuration, Platform, bCreateDebugInfo);
+
             BuildConfiguration.bUsePCHFiles = false;
             BuildConfiguration.bUseSharedPCHs = false;
             BuildConfiguration.bCheckExternalHeadersForModification = false;
             BuildConfiguration.bCheckSystemHeadersForModification = false;
             BuildConfiguration.ProcessorCountMultiplier = IOSToolChain.GetAdjustedProcessorCountMultiplier();
             BuildConfiguration.bDeployAfterCompile = true;
-        }
+		}
 
         /**
          *	Whether the platform requires the extra UnityCPPWriter
diff --git a/Engine/Source/Programs/UnrealBuildTool/System/ProjectFileGenerator.cs b/Engine/Source/Programs/UnrealBuildTool/System/ProjectFileGenerator.cs
--- a/Engine/Source/Programs/UnrealBuildTool/System/ProjectFileGenerator.cs
+++ b/Engine/Source/Programs/UnrealBuildTool/System/ProjectFileGenerator.cs
@@ -284,14 +284,17 @@ namespace UnrealBuildTool
 
 				if (!Directory.Exists(MasterProjectRelativePath + "/Source"))
 				{
-					if (BuildHostPlatform.Current.Platform == UnrealTargetPlatform.Mac)
+					if (!Directory.Exists(MasterProjectRelativePath + "/Intermediate/Source"))
 					{
-						MasterProjectRelativePath = Path.GetFullPath(Path.Combine (Utils.GetExecutingAssemblyDirectory (), "..", "..", "..", "Engine"));
-						GameProjectName = "UE4Game";
-					}
-					if (!Directory.Exists(MasterProjectRelativePath + "/Source"))
-					{
-						throw new BuildException ("Directory '{0}' is missing 'Source' folder.", MasterProjectRelativePath);
+						if (BuildHostPlatform.Current.Platform == UnrealTargetPlatform.Mac)
+						{
+							MasterProjectRelativePath = Path.GetFullPath(Path.Combine(Utils.GetExecutingAssemblyDirectory(), "..", "..", "..", "Engine"));
+							GameProjectName = "UE4Game";
+						}
+						if (!Directory.Exists(MasterProjectRelativePath + "/Source"))
+						{
+							throw new BuildException("Directory '{0}' is missing 'Source' folder.", MasterProjectRelativePath);
+						}
 					}
 				}
 				IntermediateProjectFilesPath = Path.Combine(MasterProjectRelativePath, "Intermediate", "ProjectFiles");
@@ -313,7 +316,10 @@ namespace UnrealBuildTool
 
 				if (!Directory.Exists(MasterProjectRelativePath + "/Source"))
 				{
-					throw new BuildException("Directory '{0}' is missing 'Source' folder.", MasterProjectRelativePath);
+					if (!Directory.Exists(MasterProjectRelativePath + "/Intermediate/Source"))
+					{
+						throw new BuildException("Directory '{0}' is missing 'Source' folder.", MasterProjectRelativePath);
+					}
 				}
 			}
 
diff --git a/Engine/Source/Programs/UnrealBuildTool/System/RulesCompiler.cs b/Engine/Source/Programs/UnrealBuildTool/System/RulesCompiler.cs
--- a/Engine/Source/Programs/UnrealBuildTool/System/RulesCompiler.cs
+++ b/Engine/Source/Programs/UnrealBuildTool/System/RulesCompiler.cs
@@ -1067,6 +1067,8 @@ namespace UnrealBuildTool
 				{
 					var GameSourceFolder = Path.GetFullPath(Path.Combine( GameFolder, "Source" ));
 					Folders.Add( GameSourceFolder );
+					var GameIntermediateSourceFolder = Path.GetFullPath(Path.Combine(GameFolder, "Intermediate", "Source"));
+					Folders.Add(GameIntermediateSourceFolder);
 				}
 			}
 
@@ -1220,6 +1222,12 @@ namespace UnrealBuildTool
 					{
 						AdditionalSearchPaths.Add( ProjectSourceDirectory );
 					}
+					// Add the games project's intermediate source folder
+					var ProjectIntermediateSourceDirectory = Path.Combine(UnrealBuildTool.GetUProjectPath(), "Intermediate", "Source");
+					if (Directory.Exists(ProjectIntermediateSourceDirectory))
+					{
+						AdditionalSearchPaths.Add(ProjectIntermediateSourceDirectory);
+					}
 				}
 				var ModuleFileNames = FindAllRulesSourceFiles(RulesFileType.Module, AdditionalSearchPaths);
 
diff --git a/Engine/Source/Programs/UnrealBuildTool/System/UnrealBuildTool.cs b/Engine/Source/Programs/UnrealBuildTool/System/UnrealBuildTool.cs
--- a/Engine/Source/Programs/UnrealBuildTool/System/UnrealBuildTool.cs
+++ b/Engine/Source/Programs/UnrealBuildTool/System/UnrealBuildTool.cs
@@ -689,7 +689,8 @@ namespace UnrealBuildTool
             if (HasUProjectFile())
             {
                 string SourceFolder = Path.Combine(Path.GetDirectoryName(GetUProjectFile()), "Source");
-                UProjectInfo.AddProject(GetUProjectFile(), Directory.Exists(SourceFolder));
+				string IntermediateSourceFolder = Path.Combine(Path.GetDirectoryName(GetUProjectFile()), "Intermediate", "Source");
+                UProjectInfo.AddProject(GetUProjectFile(), Directory.Exists(SourceFolder) || Directory.Exists(IntermediateSourceFolder));
             }
             else if(!RunningRocket())
             {
diff --git a/Engine/Source/Programs/UnrealBuildTool/System/XcodeProjectFileGenerator.cs b/Engine/Source/Programs/UnrealBuildTool/System/XcodeProjectFileGenerator.cs
--- a/Engine/Source/Programs/UnrealBuildTool/System/XcodeProjectFileGenerator.cs
+++ b/Engine/Source/Programs/UnrealBuildTool/System/XcodeProjectFileGenerator.cs
@@ -476,8 +476,8 @@ namespace UnrealBuildTool
 			}
 			else
 			{
-				IUEToolChain Toolchain = UEToolChain.GetPlatformToolChain(CPPTargetPlatform.IOS);
-				Toolchain.SetUpGlobalEnvironment();
+				IOSPlatform BuildPlat = UEBuildPlatform.GetBuildPlatform(UnrealTargetPlatform.IOS) as IOSPlatform;
+				BuildPlat.SetUpProjectEnvironment(UnrealTargetPlatform.IOS);
 
 				Contents.Append(
 					"\t\t" + ConfigGuid + " /* " + ConfigName + " */ = {" + ProjectFileGenerator.NewLine +
@@ -490,8 +490,8 @@ namespace UnrealBuildTool
 					"\t\t\t\t\"PRODUCT_NAME[sdk=iphoneos*]\" = \"$(TARGET_NAME)\";" + ProjectFileGenerator.NewLine +
 					"\t\t\t\t\"PRODUCT_NAME[sdk=iphonesimulator*]\" = \"$(TARGET_NAME)-simulator\";" + ProjectFileGenerator.NewLine +
 					"\t\t\t\tMACOSX_DEPLOYMENT_TARGET = 10.9;" + ProjectFileGenerator.NewLine +
-					"\t\t\t\tIPHONEOS_DEPLOYMENT_TARGET = " + Toolchain.GetPlatformVersion() + ";" + ProjectFileGenerator.NewLine +
-					"\t\t\t\tTARGETED_DEVICE_FAMILY = \"" + Toolchain.GetPlatformDevices() + "\";" + ProjectFileGenerator.NewLine +
+					"\t\t\t\tIPHONEOS_DEPLOYMENT_TARGET = " + BuildPlat.GetRunTimeVersion() + ";" + ProjectFileGenerator.NewLine +
+					"\t\t\t\tTARGETED_DEVICE_FAMILY = \"" + BuildPlat.GetRunTimeDevices() + "\";" + ProjectFileGenerator.NewLine +
 					"\t\t\t\tCONFIGURATION_BUILD_DIR = \"Engine/Binaries/IOS/Payload\";" + ProjectFileGenerator.NewLine +
 					"\t\t\t\tSDKROOT = macosx;" + ProjectFileGenerator.NewLine +
 					"\t\t\t\t\"SDKROOT[arch=x86_64]\" = macosx;" + ProjectFileGenerator.NewLine +
@@ -506,8 +506,8 @@ namespace UnrealBuildTool
 
 		private void AppendIOSBuildConfig(ref StringBuilder Contents, string ConfigName, string ConfigGuid)
 		{
-			IUEToolChain Toolchain = UEToolChain.GetPlatformToolChain(CPPTargetPlatform.IOS);
-			Toolchain.SetUpGlobalEnvironment();
+			IOSPlatform BuildPlat = UEBuildPlatform.GetBuildPlatform(UnrealTargetPlatform.IOS) as IOSPlatform;
+			BuildPlat.SetUpProjectEnvironment(UnrealTargetPlatform.IOS);
 
 			Contents.Append(
 				"\t\t" + ConfigGuid + " /* " + ConfigName + " */ = {" + ProjectFileGenerator.NewLine +
@@ -517,8 +517,8 @@ namespace UnrealBuildTool
 				"\t\t\t\t\"PRODUCT_NAME[sdk=iphonesimulator*]\" = \"$(TARGET_NAME)-simulator\";" + ProjectFileGenerator.NewLine +
 				"\t\t\t\tSUPPORTED_PLATFORMS = \"iphoneos iphonesimulator\";" + ProjectFileGenerator.NewLine +
 				"\t\t\t\tCONFIGURATION_BUILD_DIR = \"Engine/Binaries/IOS/Payload\";" + ProjectFileGenerator.NewLine +
-				"\t\t\t\tIPHONEOS_DEPLOYMENT_TARGET = " + Toolchain.GetPlatformVersion() + ";" + ProjectFileGenerator.NewLine +
-				"\t\t\t\tTARGETED_DEVICE_FAMILY = \"" + Toolchain.GetPlatformDevices() + "\";" + ProjectFileGenerator.NewLine +
+				"\t\t\t\tIPHONEOS_DEPLOYMENT_TARGET = " + BuildPlat.GetRunTimeVersion() + ";" + ProjectFileGenerator.NewLine +
+				"\t\t\t\tTARGETED_DEVICE_FAMILY = \"" + BuildPlat.GetRunTimeDevices() + "\";" + ProjectFileGenerator.NewLine +
 				"\t\t\t\tSDKROOT = iphoneos;" + ProjectFileGenerator.NewLine +
 				"\t\t\t};" + ProjectFileGenerator.NewLine +
 				"\t\t\tname = " + ConfigName + ";" + ProjectFileGenerator.NewLine +
@@ -527,8 +527,8 @@ namespace UnrealBuildTool
 
 		private void AppendIOSRunConfig(ref StringBuilder Contents, string ConfigName, string ConfigGuid, string TargetName, string EngineRelative, string GamePath, bool bIsUE4Game, bool IsAGame, bool bIsUE4Client)
 		{
-			IUEToolChain Toolchain = UEToolChain.GetPlatformToolChain(CPPTargetPlatform.IOS);
-			Toolchain.SetUpGlobalEnvironment();
+			IOSPlatform BuildPlat = UEBuildPlatform.GetBuildPlatform(UnrealTargetPlatform.IOS) as IOSPlatform;
+			BuildPlat.SetUpProjectEnvironment(UnrealTargetPlatform.IOS);
 
 			Contents.Append(
 				"\t\t" + ConfigGuid + " /* " + ConfigName + " */ = {" + ProjectFileGenerator.NewLine +
@@ -538,7 +538,7 @@ namespace UnrealBuildTool
 				"\t\t\t\t\"PRODUCT_NAME[sdk=iphonesimulator*]\" = \"" + TargetName + "-simulator\";" + ProjectFileGenerator.NewLine +
 				"\t\t\t\tCONFIGURATION_BUILD_DIR = \"Engine/Binaries/IOS/Payload\";" + ProjectFileGenerator.NewLine +
 				"\t\t\t\tSUPPORTED_PLATFORMS = \"iphoneos iphonesimulator\";" + ProjectFileGenerator.NewLine +
-				"\t\t\t\tIPHONEOS_DEPLOYMENT_TARGET = " + Toolchain.GetPlatformVersion() + ";" + ProjectFileGenerator.NewLine +
+				"\t\t\t\tIPHONEOS_DEPLOYMENT_TARGET = " + BuildPlat.GetRunTimeVersion() + ";" + ProjectFileGenerator.NewLine +
 				"\t\t\t\t\"CODE_SIGN_IDENTITY[sdk=iphoneos*]\" = \"iPhone Developer\";" + ProjectFileGenerator.NewLine);
 
 			string InfoPlistPath = "";
@@ -598,7 +598,7 @@ namespace UnrealBuildTool
 				"\t\t\t\t);" + ProjectFileGenerator.NewLine +
 				"\t\t\t\tINFOPLIST_OUTPUT_FORMAT = xml;" + ProjectFileGenerator.NewLine +
 				"\t\t\t\t\"PROVISIONING_PROFILE[sdk=iphoneos*]\" = \"\";" + ProjectFileGenerator.NewLine +
-				"\t\t\t\tTARGETED_DEVICE_FAMILY = \"" + Toolchain.GetPlatformDevices() + "\";" + ProjectFileGenerator.NewLine +
+				"\t\t\t\tTARGETED_DEVICE_FAMILY = \"" + BuildPlat.GetRunTimeDevices() + "\";" + ProjectFileGenerator.NewLine +
 				"\t\t\t\tSDKROOT = iphoneos;" + ProjectFileGenerator.NewLine +
 				"\t\t\t};" + ProjectFileGenerator.NewLine +
 				"\t\t\tname = " + ConfigName + ";" + ProjectFileGenerator.NewLine +
@@ -1420,7 +1420,7 @@ namespace UnrealBuildTool
 		private void WriteProject(StringBuilder XcodeProjectFileContent, List<XcodeProjectTarget> ProjectTargets)
 		{
 			// @todo: Windows and Mac should store the project in the same folder
-			string PathBranch = (BuildHostPlatform.Current.Platform != UnrealTargetPlatform.Mac) ? "/Engine/Intermediate/IOS" : "";
+			string PathBranch = (BuildHostPlatform.Current.Platform != UnrealTargetPlatform.Mac) ? (!string.IsNullOrEmpty(UnrealBuildTool.GetUProjectPath()) ? "/Intermediate/IOS" : "/Engine/Intermediate/IOS") : "";
 			var XcodeProjectPath = MasterProjectRelativePath + PathBranch + "/" + MasterProjectName + ".xcodeproj";
 			if (bGeneratingGameProjectFiles && GameProjectName == "UE4Game")
 			{
@@ -1456,7 +1456,7 @@ namespace UnrealBuildTool
 
 		private void WriteWorkspace()
 		{
-			string PathBranch = (BuildHostPlatform.Current.Platform != UnrealTargetPlatform.Mac) ? "/Engine/Intermediate/IOS" : "";
+			string PathBranch = (BuildHostPlatform.Current.Platform != UnrealTargetPlatform.Mac) ? (!string.IsNullOrEmpty(UnrealBuildTool.GetUProjectPath()) ? "/Intermediate/IOS" : "/Engine/Intermediate/IOS") : "";
 			var XcodeProjectPath = MasterProjectRelativePath + PathBranch + "/" + MasterProjectName + ".xcodeproj";
 			if (bGeneratingGameProjectFiles && GameProjectName == "UE4Game")
 			{
diff --git a/Engine/Source/Programs/UnrealBuildTool/ToolChain/UEToolChain.cs b/Engine/Source/Programs/UnrealBuildTool/ToolChain/UEToolChain.cs
--- a/Engine/Source/Programs/UnrealBuildTool/ToolChain/UEToolChain.cs
+++ b/Engine/Source/Programs/UnrealBuildTool/ToolChain/UEToolChain.cs
@@ -44,10 +44,6 @@ namespace UnrealBuildTool
 
 		void FixBundleBinariesPaths(UEBuildTarget Target, List<UEBuildBinary> Binaries);
 
-		string GetPlatformVersion();
-
-		string GetPlatformDevices();
-
 		UnrealTargetPlatform GetPlatform();
 	}
 
@@ -168,84 +164,6 @@ namespace UnrealBuildTool
 
 		public virtual void ParseProjectSettings()
 		{
-			ConfigCacheIni Ini = new ConfigCacheIni(GetPlatform(), "Engine", UnrealBuildTool.GetUProjectPath());
-			bool bValue = UEBuildConfiguration.bCompileAPEX;
-			if (Ini.GetBool("/Script/BuildSettings.BuildSettings", "bCompileApex", out bValue))
-			{
-				UEBuildConfiguration.bCompileAPEX = bValue;
-			}
-
-			bValue = UEBuildConfiguration.bCompileBox2D;
-			if (Ini.GetBool("/Script/BuildSettings.BuildSettings", "bCompileBox2D", out bValue))
-			{
-				UEBuildConfiguration.bCompileBox2D = bValue;
-			}
-
-			bValue = UEBuildConfiguration.bCompileICU;
-			if (Ini.GetBool("/Script/BuildSettings.BuildSettings", "bCompileICU", out bValue))
-			{
-				UEBuildConfiguration.bCompileICU = bValue;
-			}
-
-			bValue = UEBuildConfiguration.bCompileSimplygon;
-			if (Ini.GetBool("/Script/BuildSettings.BuildSettings", "bCompileSimplygon", out bValue))
-			{
-				UEBuildConfiguration.bCompileSimplygon = bValue;
-			}
-
-			bValue = UEBuildConfiguration.bCompileLeanAndMeanUE;
-			if (Ini.GetBool("/Script/BuildSettings.BuildSettings", "bCompileLeanAndMeanUE", out bValue))
-			{
-				UEBuildConfiguration.bCompileLeanAndMeanUE = bValue;
-			}
-
-			bValue = UEBuildConfiguration.bIncludeADO;
-			if (Ini.GetBool("/Script/BuildSettings.BuildSettings", "bIncludeADO", out bValue))
-			{
-				UEBuildConfiguration.bIncludeADO = bValue;
-			}
-
-			bValue = UEBuildConfiguration.bCompileRecast;
-			if (Ini.GetBool("/Script/BuildSettings.BuildSettings", "bCompileRecast", out bValue))
-			{
-				UEBuildConfiguration.bCompileRecast = bValue;
-			}
-
-			bValue = UEBuildConfiguration.bCompileSpeedTree;
-			if (Ini.GetBool("/Script/BuildSettings.BuildSettings", "bCompileSpeedTree", out bValue))
-			{
-				UEBuildConfiguration.bCompileSpeedTree = bValue;
-			}
-
-			bValue = UEBuildConfiguration.bCompileWithPluginSupport;
-			if (Ini.GetBool("/Script/BuildSettings.BuildSettings", "bCompileWithPluginSupport", out bValue))
-			{
-				UEBuildConfiguration.bCompileWithPluginSupport = bValue;
-			}
-
-			bValue = UEBuildConfiguration.bCompilePhysXVehicle;
-			if (Ini.GetBool("/Script/BuildSettings.BuildSettings", "bCompilePhysXVehicle", out bValue))
-			{
-				UEBuildConfiguration.bCompilePhysXVehicle = bValue;
-			}
-
-			bValue = UEBuildConfiguration.bCompileFreeType;
-			if (Ini.GetBool("/Script/BuildSettings.BuildSettings", "bCompileFreeType", out bValue))
-			{
-				UEBuildConfiguration.bCompileFreeType = bValue;
-			}
-
-			bValue = UEBuildConfiguration.bCompileForSize;
-			if (Ini.GetBool("/Script/BuildSettings.BuildSettings", "bCompileForSize", out bValue))
-			{
-				UEBuildConfiguration.bCompileForSize = bValue;
-			}
-
-			bValue = UEBuildConfiguration.bCompileCEF3;
-			if (Ini.GetBool("/Script/BuildSettings.BuildSettings", "bCompileCEF3", out bValue))
-			{
-				UEBuildConfiguration.bCompileCEF3 = bValue;
-			}
 		}
 
 		protected void RunUnrealHeaderToolIfNeeded()
@@ -303,16 +221,6 @@ namespace UnrealBuildTool
 
 		}
 
-		public virtual string GetPlatformVersion()
-		{
-			return "";
-		}
-
-		public virtual string GetPlatformDevices()
-		{
-			return "";
-		}
-
 		public virtual UnrealTargetPlatform GetPlatform()
 		{
 			return UnrealTargetPlatform.Unknown;
