check iOS specific build arguments to determine if they are different from default
#ios

diff --git a/Engine/Source/Programs/UnrealBuildTool/IOS/UEBuildIOS.cs b/Engine/Source/Programs/UnrealBuildTool/IOS/UEBuildIOS.cs
--- a/Engine/Source/Programs/UnrealBuildTool/IOS/UEBuildIOS.cs
+++ b/Engine/Source/Programs/UnrealBuildTool/IOS/UEBuildIOS.cs
@@ -235,7 +235,77 @@ namespace UnrealBuildTool
 			// determine if we need to generate the dsym
 			Ini.GetBool("/Script/IOSRuntimeSettings.IOSRuntimeSettings", "bGenerateSYMFile", out BuildConfiguration.bGeneratedSYMFile);
 		}
-        /**
+
+		/**
+		 * Check for the default configuration
+		 *
+		 * return true if the project uses the default build config
+		 */
+		public override bool HasDefaultBuildConfig(UnrealTargetPlatform Platform, string ProjectPath)
+		{
+			ConfigCacheIni ProjIni = new ConfigCacheIni(Platform, "Engine", ProjectPath);
+			ConfigCacheIni DefaultIni = new ConfigCacheIni(Platform, "Engine", null);
+
+			string DefaultMinVersion = "IOS_6";
+			string ProjectMinVersion = DefaultMinVersion;
+			ProjIni.GetString("/Script/IOSRuntimeSettings.IOSRuntimeSettings", "MinimumiOSVersion", out ProjectMinVersion);
+			DefaultIni.GetString("/Script/IOSRuntimeSettings.IOSRuntimeSettings", "MinimumiOSVersion", out DefaultMinVersion);
+			if (DefaultMinVersion != ProjectMinVersion)
+			{
+				return false;
+			}
+
+			bool bDefaultBuild = true;
+			bool bProjectBuild = true;
+			ProjIni.GetBool("/Script/IOSRuntimeSettings.IOSRuntimeSettings", "bDevForArmV7", out bProjectBuild);
+			DefaultIni.GetBool("/Script/IOSRuntimeSettings.IOSRuntimeSettings", "bDevForArmV7", out bDefaultBuild);
+			if (bDefaultBuild != bProjectBuild)
+			{
+				return false;
+			}
+			ProjIni.GetBool("/Script/IOSRuntimeSettings.IOSRuntimeSettings", "bDevForArm64", out bProjectBuild);
+			DefaultIni.GetBool("/Script/IOSRuntimeSettings.IOSRuntimeSettings", "bDevForArm64", out bDefaultBuild);
+			if (bDefaultBuild != bProjectBuild)
+			{
+				return false;
+			}
+			ProjIni.GetBool("/Script/IOSRuntimeSettings.IOSRuntimeSettings", "bDevForArmV7S", out bProjectBuild);
+			DefaultIni.GetBool("/Script/IOSRuntimeSettings.IOSRuntimeSettings", "bDevForArmV7S", out bDefaultBuild);
+			if (bDefaultBuild != bProjectBuild)
+			{
+				return false;
+			}
+			ProjIni.GetBool("/Script/IOSRuntimeSettings.IOSRuntimeSettings", "bShipForArmV7", out bProjectBuild);
+			DefaultIni.GetBool("/Script/IOSRuntimeSettings.IOSRuntimeSettings", "bShipForArmV7", out bDefaultBuild);
+			if (bDefaultBuild != bProjectBuild)
+			{
+				return false;
+			}
+			ProjIni.GetBool("/Script/IOSRuntimeSettings.IOSRuntimeSettings", "bShipForArm64", out bProjectBuild);
+			DefaultIni.GetBool("/Script/IOSRuntimeSettings.IOSRuntimeSettings", "bShipForArm64", out bDefaultBuild);
+			if (bDefaultBuild != bProjectBuild)
+			{
+				return false;
+			}
+			ProjIni.GetBool("/Script/IOSRuntimeSettings.IOSRuntimeSettings", "bShipForArmV7S", out bProjectBuild);
+			DefaultIni.GetBool("/Script/IOSRuntimeSettings.IOSRuntimeSettings", "bShipForArmV7S", out bDefaultBuild);
+			if (bDefaultBuild != bProjectBuild)
+			{
+				return false;
+			}
+
+			// determine if we need to generate the dsym
+			ProjIni.GetBool("/Script/IOSRuntimeSettings.IOSRuntimeSettings", "bGenerateSYMFile", out bProjectBuild);
+			DefaultIni.GetBool("/Script/IOSRuntimeSettings.IOSRuntimeSettings", "bGenerateSYMFile", out bDefaultBuild);
+			if (bDefaultBuild != bProjectBuild)
+			{
+				return false;
+			}
+
+			return base.HasDefaultBuildConfig(Platform, ProjectPath);
+		}
+
+		/**
          *	Setup the target environment for building
          *	
          *	@param	InBuildTarget		The target being built
