Prevent building GitDependencies by default. Causes problems for users that submit Git source to P4. Can still be build by right-clicking project file. UE-8549.

diff --git a/Engine/Source/Programs/UnrealBuildTool/System/Project.cs b/Engine/Source/Programs/UnrealBuildTool/System/Project.cs
--- a/Engine/Source/Programs/UnrealBuildTool/System/Project.cs
+++ b/Engine/Source/Programs/UnrealBuildTool/System/Project.cs
@@ -126,6 +126,7 @@ namespace UnrealBuildTool
 		protected ProjectFile( string InitRelativeFilePath )
 		{
 			RelativeProjectFilePath = InitRelativeFilePath;
+			ShouldBuildByDefaultForSolutionTargets = true;
 		}
 
 
@@ -178,6 +179,12 @@ namespace UnrealBuildTool
 			set;
 		}
 
+		/// Whether this project should be built by default. Can still be built from the IDE through the context menu.
+		public bool ShouldBuildByDefaultForSolutionTargets
+		{
+			get;
+			set;
+		}
 
 
 
diff --git a/Engine/Source/Programs/UnrealBuildTool/System/ProjectFileGenerator.cs b/Engine/Source/Programs/UnrealBuildTool/System/ProjectFileGenerator.cs
--- a/Engine/Source/Programs/UnrealBuildTool/System/ProjectFileGenerator.cs
+++ b/Engine/Source/Programs/UnrealBuildTool/System/ProjectFileGenerator.cs
@@ -511,7 +511,7 @@ namespace UnrealBuildTool
 					ProgramsFolder.ChildProjects.Add(AddSimpleCSharpProject("DotNETCommon/DotNETUtilities", bShouldBuildForAllSolutionTargets: true, bForceDevelopmentConfiguration: true));
 
 					// Add the Git dependencies project
-					ProgramsFolder.ChildProjects.Add(AddSimpleCSharpProject("GitDependencies", bShouldBuildForAllSolutionTargets: true, bForceDevelopmentConfiguration: true));
+					ProgramsFolder.ChildProjects.Add(AddSimpleCSharpProject("GitDependencies", bForceDevelopmentConfiguration: true, bShouldBuildByDefaultForSolutionTargets: false));
 
 					// Add all of the IOS C# projects
 					AddIOSProjects( ProgramsFolder );
@@ -1163,7 +1163,7 @@ namespace UnrealBuildTool
 		/// </summary>
 		/// <param name="ProjectName">Name of project file to add</param>
 		/// <returns>ProjectFile if the operation was successful, otherwise null.</returns>
-		private VCSharpProjectFile AddSimpleCSharpProject(string ProjectName, bool bShouldBuildForAllSolutionTargets = false, bool bForceDevelopmentConfiguration = false)
+		private VCSharpProjectFile AddSimpleCSharpProject(string ProjectName, bool bShouldBuildForAllSolutionTargets = false, bool bForceDevelopmentConfiguration = false, bool bShouldBuildByDefaultForSolutionTargets = true)
 		{
 			VCSharpProjectFile Project = null;
 
@@ -1174,6 +1174,7 @@ namespace UnrealBuildTool
 				var FileNameRelativeToMasterProject = Utils.MakePathRelativeTo( ProjectFileName, MasterProjectRelativePath );
 				Project = new VCSharpProjectFile(FileNameRelativeToMasterProject);
 				Project.ShouldBuildForAllSolutionTargets = bShouldBuildForAllSolutionTargets;
+				Project.ShouldBuildByDefaultForSolutionTargets = bShouldBuildByDefaultForSolutionTargets;
 				AddExistingProjectFile(Project, bForceDevelopmentConfiguration: bForceDevelopmentConfiguration);
 			}
 			else
diff --git a/Engine/Source/Programs/UnrealBuildTool/System/VCProjectFileGenerator.cs b/Engine/Source/Programs/UnrealBuildTool/System/VCProjectFileGenerator.cs
--- a/Engine/Source/Programs/UnrealBuildTool/System/VCProjectFileGenerator.cs
+++ b/Engine/Source/Programs/UnrealBuildTool/System/VCProjectFileGenerator.cs
@@ -657,7 +657,7 @@ namespace UnrealBuildTool
 
 
 										// Set whether this project configuration should be built when the user initiates "build solution"
-										if( MatchingProjectTarget != null )
+										if( MatchingProjectTarget != null && CurProject.ShouldBuildByDefaultForSolutionTargets )
 										{
 											VCSolutionFileContent.Append(
 													"		" + CurProjectGUID + "." + SolutionConfigCombination.VCSolutionConfigAndPlatformName + ".Build.0 = " + ProjectConfigAndPlatformPair + ProjectFileGenerator.NewLine );
