PR #869: Harmonize ECompilationResult return codes between Editor platforms (Contributed by [REDACTED_PERSON])

https://github.com/EpicGames/UnrealEngine/pull/869

diff --git a/Engine/Source/Programs/UnrealBuildTool/System/ExternalExecution.cs b/Engine/Source/Programs/UnrealBuildTool/System/ExternalExecution.cs
--- a/Engine/Source/Programs/UnrealBuildTool/System/ExternalExecution.cs
+++ b/Engine/Source/Programs/UnrealBuildTool/System/ExternalExecution.cs
@@ -13,23 +13,23 @@ using System.Runtime.InteropServices;
 namespace UnrealBuildTool
 {
 	// This enum has to be compatible with the one defined in the
-	// UE4\Engine\Source\Runtime\Core\Public\Modules\ModuleManager.h
+	// UE4\Engine\Source\Runtime\Core\Public\Misc\ComplilationResult.h 
 	// to keep communication between UHT, UBT and Editor compiling
 	// processes valid.
 	public enum ECompilationResult
 	{
-		/** All targets were up to date, used only with -canskiplink */
-		UpToDate = -2,
-		/** Build was canceled, this is used on the engine side only */
-		Canceled = -1,
 		/** Compilation succeeded */
 		Succeeded = 0,
-		/** Compilation failed because generated code changed which was not supported */
-		FailedDueToHeaderChange = 1,
-		/** Compilation failed due to compilation errors */
-		OtherCompilationError = 2,
+		/** Build was canceled, this is used on the engine side only */
+		Canceled = 1,
+		/** All targets were up to date, used only with -canskiplink */
+		UpToDate = 2,
 		/** The process has most likely crashed. This is what UE returns in case of an assert */
 		CrashOrAssert = 3,
+		/** Compilation failed because generated code changed which was not supported */
+		FailedDueToHeaderChange = 4,
+		/** Compilation failed due to compilation errors */
+		OtherCompilationError = 5,
 		/** Compilation is not supported in the current build */
 		Unsupported,
 		/** Unknown error */
diff --git a/Engine/Source/Runtime/Core/Public/Misc/CompilationResult.h b/Engine/Source/Runtime/Core/Public/Misc/CompilationResult.h
--- a/Engine/Source/Runtime/Core/Public/Misc/CompilationResult.h
+++ b/Engine/Source/Runtime/Core/Public/Misc/CompilationResult.h
@@ -14,18 +14,18 @@ namespace ECompilationResult
 {
 	enum Type
 	{
-		/** All targets were up to date, used only with -canskiplink */
-		UpToDate = -2,
-		/** Build was canceled, this is used on the engine side only */
-		Canceled = -1,
 		/** Compilation succeeded */
 		Succeeded = 0,
-		/** Compilation failed because generated code changed which was not supported */
-		FailedDueToHeaderChange = 1,
-		/** Compilation failed due to compilation errors */
-		OtherCompilationError = 2,
+		/** Build was canceled, this is used on the engine side only */
+		Canceled = 1,
+		/** All targets were up to date, used only with -canskiplink */
+		UpToDate = 2,
 		/** The process has most likely crashed. This is what UE returns in case of an assert */
 		CrashOrAssert = 3,
+		/** Compilation failed because generated code changed which was not supported */
+		FailedDueToHeaderChange = 4,
+		/** Compilation failed due to compilation errors */
+		OtherCompilationError = 5,
 		/** Compilation is not supported in the current build */
 		Unsupported,
 		/** Unknown error */
