Rollback //depot/UE4/Engine/Source/Programs/UnrealBuildTool/UnrealBuildTool.csproj to revision 87

We cannot use conditional includes for files we need on Mac and Linux, as Xamarin Studio and Mono Develop do not support that

diff --git a/Engine/Source/Programs/UnrealBuildTool/UnrealBuildTool.csproj b/Engine/Source/Programs/UnrealBuildTool/UnrealBuildTool.csproj
--- a/Engine/Source/Programs/UnrealBuildTool/UnrealBuildTool.csproj
+++ b/Engine/Source/Programs/UnrealBuildTool/UnrealBuildTool.csproj
@@ -122,26 +122,18 @@
       </ItemGroup>
     </When>
   </Choose>
-  <Choose>
-	<When Condition="Exists('Mac\UEBuildMac.cs')">
-      <ItemGroup>
-	    <Compile Include="Mac\MacToolChain.cs" />
-	    <Compile Include="Mac\UEBuildMac.cs" />
-	    <Compile Include="Mac\MacProjectGenerator.cs" />
-	    <Compile Include="Mac\UEDeployMac.cs" />
-      </ItemGroup>
-	</When>
-  </Choose>
-  <Choose>
-	<When Condition="Exists('Android\UEBuildAndroid.cs')">
-	  <ItemGroup>
-	    <Compile Include="Android\UEBuildAndroid.cs" />
-	    <Compile Include="Android\UEDeployAndroid.cs" />
-		<Compile Include="Android\AndroidToolChain.cs" />
-		<Compile Include="Android\AndroidProjectGenerator.cs" />
-      </ItemGroup>
-	</When>
-  </Choose>
+  <ItemGroup>
+	<Compile Include="Mac\MacToolChain.cs" />
+	<Compile Include="Mac\UEBuildMac.cs" />
+	<Compile Include="Mac\MacProjectGenerator.cs" />
+	<Compile Include="Mac\UEDeployMac.cs" />
+  </ItemGroup>
+  <ItemGroup>
+	<Compile Include="Android\UEBuildAndroid.cs" />
+	<Compile Include="Android\UEDeployAndroid.cs" />
+	<Compile Include="Android\AndroidToolChain.cs" />
+	<Compile Include="Android\AndroidProjectGenerator.cs" />
+  </ItemGroup>
   <Choose>
     <When Condition="Exists('WinRT\UEBuildWinRT.cs')">
       <ItemGroup>
@@ -152,35 +144,23 @@
       </ItemGroup>
     </When>
   </Choose>
-  <Choose>
-	<When Condition="Exists('IOS\UEBuildIOS.cs')">
-      <ItemGroup>
-	    <Compile Include="IOS\IOSProjectGenerator.cs" />
-	    <Compile Include="IOS\IOSToolChain.cs" />
-	    <Compile Include="IOS\UEBuildIOS.cs" />
-	    <Compile Include="IOS\UEDeployIOS.cs" />
-      </ItemGroup>
-	</When>
-  </Choose>
-  <Choose>
-	<When Condition="Exists('HTML5\UEBuildHTML5.cs')">
-      <ItemGroup>
-	    <Compile Include="HTML5\HTML5ProjectGenerator.cs" />
-	    <Compile Include="HTML5\HTML5ToolChain.cs" />
-	    <Compile Include="HTML5\UEBuildHTML5.cs" />
-	    <Compile Include="HTML5\HTML5SDKInfo.cs" />
-      </ItemGroup>
-    </When>
-  </Choose>
-  <Choose>
-	<When Condition="Exists('Linux\UEBuildLinux.cs')">
-      <ItemGroup>
-	    <Compile Include="Linux\LinuxProjectGenerator.cs" />
-	    <Compile Include="Linux\LinuxToolChain.cs" />
-	    <Compile Include="Linux\UEBuildLinux.cs" />
-      </ItemGroup>
-	</When>
-  </Choose>
+  <ItemGroup>
+	<Compile Include="IOS\IOSProjectGenerator.cs" />
+	<Compile Include="IOS\IOSToolChain.cs" />
+	<Compile Include="IOS\UEBuildIOS.cs" />
+	<Compile Include="IOS\UEDeployIOS.cs" />
+  </ItemGroup>
+  <ItemGroup>
+	<Compile Include="HTML5\HTML5ProjectGenerator.cs" />
+	<Compile Include="HTML5\HTML5ToolChain.cs" />
+	<Compile Include="HTML5\UEBuildHTML5.cs" />
+	<Compile Include="HTML5\HTML5SDKInfo.cs" />
+  </ItemGroup>
+  <ItemGroup>
+	<Compile Include="Linux\LinuxProjectGenerator.cs" />
+	<Compile Include="Linux\LinuxToolChain.cs" />
+	<Compile Include="Linux\UEBuildLinux.cs" />
+  </ItemGroup>
   <Choose>
     <When Condition="Exists('System\NotForLicensees\TelemetryProviderET.cs')">
       <ItemGroup>
