Removed unnecessary code forcing Rocket build to use Development now that DebugGame is an option. Opted to retain the same intermediate folder naming though as these are technically always development.

Added IsEngineInstalled function to UBT to better explain why some of the Rocket code is different.

diff --git a/Engine/Source/Programs/AutomationTool/Scripts/BuildCookRun.Automation.cs b/Engine/Source/Programs/AutomationTool/Scripts/BuildCookRun.Automation.cs
--- a/Engine/Source/Programs/AutomationTool/Scripts/BuildCookRun.Automation.cs
+++ b/Engine/Source/Programs/AutomationTool/Scripts/BuildCookRun.Automation.cs
@@ -58,26 +58,6 @@ public class BuildCookRun : BuildCommand
 			RawProjectPath: ProjectPath
 		);
 
-		// @rocket hack: non-code projects cannot run in Debug
-		if (Params.Rocket && !ProjectUtils.IsCodeBasedUProjectFile(ProjectPath))
-		{
-			if (Params.ClientConfigsToBuild.Contains(UnrealTargetConfiguration.Debug))
-			{
-				Log("Non-code projects cannot run in Debug game clients. Defaulting to Development.");
-
-				Params.ClientConfigsToBuild.Remove(UnrealTargetConfiguration.Debug);
-				Params.ClientConfigsToBuild.Add(UnrealTargetConfiguration.Development);
-			}
-
-			if (Params.ClientConfigsToBuild.Contains(UnrealTargetConfiguration.Debug))
-			{
-				Log("Non-code projects cannot run in Debug game servers. Defaulting to Development.");
-
-				Params.ServerConfigsToBuild.Remove(UnrealTargetConfiguration.Debug);
-				Params.ServerConfigsToBuild.Add(UnrealTargetConfiguration.Development);
-			}
-		}
-
 		var DirectoriesToCook = ParseParamValue("cookdir");
 		if (!String.IsNullOrEmpty(DirectoriesToCook))
 		{
diff --git a/Engine/Source/Programs/AutomationTool/Win/WinPlatform.Automation.cs b/Engine/Source/Programs/AutomationTool/Win/WinPlatform.Automation.cs
--- a/Engine/Source/Programs/AutomationTool/Win/WinPlatform.Automation.cs
+++ b/Engine/Source/Programs/AutomationTool/Win/WinPlatform.Automation.cs
@@ -43,13 +43,13 @@ public abstract class BaseWinPlatform : Platform
 		string PhysXVer = "VS" + WindowsPlatform.GetVisualStudioCompilerVersionName();
 		string ApexVer = "VS" + WindowsPlatform.GetVisualStudioCompilerVersionName();
 		string PhysXMaskForDebugConfiguration = Params.bDebugBuildsActuallyUseDebugCRT ? "*DEBUG*.*" : "*PROFILE*.*";
-		if (SC.StageTargetConfigurations.Contains(UnrealTargetConfiguration.Debug) && !Params.Rocket)
+		if (SC.StageTargetConfigurations.Contains(UnrealTargetConfiguration.Debug))
 		{
 			StageExecutable("dll", SC, CommandUtils.CombinePaths(SC.LocalRoot, "Engine/Binaries/ThirdParty/PhysX/APEX-1.3", SC.PlatformDir, ApexVer), PhysXMaskForDebugConfiguration, true);
 			StageExecutable("dll", SC, CommandUtils.CombinePaths(SC.LocalRoot, "Engine/Binaries/ThirdParty/PhysX/PhysX-3.3", SC.PlatformDir, PhysXVer), PhysXMaskForDebugConfiguration, true);
 			StageExecutable("dll", SC, CommandUtils.CombinePaths(SC.LocalRoot, "Engine/Binaries/ThirdParty/PhysX/PhysX-3.3", SC.PlatformDir, PhysXVer), "nvToolsExt*.", true);
 		}
-		if (SC.StageTargetConfigurations.Any(x => x != UnrealTargetConfiguration.Debug) || Params.Rocket)
+		if (SC.StageTargetConfigurations.Any(x => x != UnrealTargetConfiguration.Debug))
 		{
 			StageExecutable("dll", SC, CommandUtils.CombinePaths(SC.LocalRoot, "Engine/Binaries/ThirdParty/PhysX/APEX-1.3", SC.PlatformDir, ApexVer), "*.", true, new string[] { "*DEBUG*.*", "*CHECKED*.*" });
 			StageExecutable("dll", SC, CommandUtils.CombinePaths(SC.LocalRoot, "Engine/Binaries/ThirdParty/PhysX/PhysX-3.3", SC.PlatformDir, PhysXVer), "*.", true, new string[] { "*DEBUG*.*", "*CHECKED*.*" });
diff --git a/Engine/Source/Programs/UnrealBuildTool/Configuration/UEBuildPlatform.cs b/Engine/Source/Programs/UnrealBuildTool/Configuration/UEBuildPlatform.cs
--- a/Engine/Source/Programs/UnrealBuildTool/Configuration/UEBuildPlatform.cs
+++ b/Engine/Source/Programs/UnrealBuildTool/Configuration/UEBuildPlatform.cs
@@ -697,25 +697,7 @@ namespace UnrealBuildTool
 			// Determine the C++ compile/link configuration based on the Unreal configuration.
 			CPPTargetConfiguration CompileConfiguration;
 			UnrealTargetConfiguration CheckConfig = InBuildTarget.Configuration;
-			//@todo SAS: Add a true Debug mode!
-			if (UnrealBuildTool.RunningRocket())
-			{
-				if (Utils.IsFileUnderDirectory(InBuildTarget.OutputPaths[0], UnrealBuildTool.GetUProjectPath()))
-				{
-					if (CheckConfig == UnrealTargetConfiguration.Debug)
-					{
-						CheckConfig = UnrealTargetConfiguration.DebugGame;
-					}
-				}
-				else
-				{
-					// Only Development and Shipping are supported for engine modules
-					if (CheckConfig != UnrealTargetConfiguration.Development && CheckConfig != UnrealTargetConfiguration.Shipping)
-					{
-						CheckConfig = UnrealTargetConfiguration.Development;
-					}
-				}
-			}
+
 			switch (CheckConfig)
 			{
 				default:
diff --git a/Engine/Source/Programs/UnrealBuildTool/Configuration/UEBuildTarget.cs b/Engine/Source/Programs/UnrealBuildTool/Configuration/UEBuildTarget.cs
--- a/Engine/Source/Programs/UnrealBuildTool/Configuration/UEBuildTarget.cs
+++ b/Engine/Source/Programs/UnrealBuildTool/Configuration/UEBuildTarget.cs
@@ -2632,20 +2632,16 @@ namespace UnrealBuildTool
                 GlobalCompileEnvironment.Config.Definitions.Add("UE_ROCKET=0");
             }
 
-			// Rocket intermediates go to the project's intermediate folder.  Rocket never writes to the engine intermediate folder. (Those files are immutable)
+			// Installed Engine intermediates go to the project's intermediate folder. Installed Engine never writes to the engine intermediate folder. (Those files are immutable)
 			// Also, when compiling in monolithic, all intermediates go to the project's folder.  This is because a project can change definitions that affects all engine translation
 			// units too, so they can't be shared between different targets.  They are effectively project-specific engine intermediates.
-			if( UnrealBuildTool.RunningRocket() || ( UnrealBuildTool.HasUProjectFile() && ShouldCompileMonolithic() ) )
+			if( UnrealBuildTool.IsEngineInstalled() || ( UnrealBuildTool.HasUProjectFile() && ShouldCompileMonolithic() ) )
 			{
-				//@todo SAS: Add a true Debug mode!
 				var IntermediateConfiguration = Configuration;
 				if( UnrealBuildTool.RunningRocket() )
 				{
-				// Only Development and Shipping are supported for engine modules
-					if( IntermediateConfiguration != UnrealTargetConfiguration.Development && IntermediateConfiguration != UnrealTargetConfiguration.Shipping )
-				{
-					IntermediateConfiguration = UnrealTargetConfiguration.Development;
-				}
+					// Only Development and Shipping are supported for engine modules
+					IntermediateConfiguration = Configuration == UnrealTargetConfiguration.DebugGame ? UnrealTargetConfiguration.Development : Configuration;
 				}
 
 				GlobalCompileEnvironment.Config.OutputDirectory = Path.Combine(BuildConfiguration.PlatformIntermediatePath, OutputAppName, IntermediateConfiguration.ToString());
diff --git a/Engine/Source/Programs/UnrealBuildTool/System/RulesCompiler.cs b/Engine/Source/Programs/UnrealBuildTool/System/RulesCompiler.cs
--- a/Engine/Source/Programs/UnrealBuildTool/System/RulesCompiler.cs
+++ b/Engine/Source/Programs/UnrealBuildTool/System/RulesCompiler.cs
@@ -1354,20 +1354,6 @@ namespace UnrealBuildTool
 
 			UnrealTargetPlatform LocalPlatform = Target.Platform;
 			UnrealTargetConfiguration LocalConfiguration = Target.Configuration;
-			//@todo SAS: Add a true Debug mode!
-			if (UnrealBuildTool.RunningRocket())
-			{
-				// Have to do absolute here as this could be a project that is under the root
-				string FullUProjectPath = Path.GetFullPath(UnrealBuildTool.GetUProjectPath());
-				if ( !Utils.IsFileUnderDirectory( ModuleFileName, FullUProjectPath ) )
-				{
-					// Only Development and Shipping are supported for engine modules
-					if( LocalConfiguration != UnrealTargetConfiguration.Development && LocalConfiguration != UnrealTargetConfiguration.Shipping )
-					{
-						LocalConfiguration = UnrealTargetConfiguration.Development;
-					}
-				}
-			}
 			TargetInfo LocalTarget = new TargetInfo(LocalPlatform, LocalConfiguration, Target.Type);
 
 			// The build module must define a type named 'Rules' that derives from our 'ModuleRules' type.  
diff --git a/Engine/Source/Programs/UnrealBuildTool/System/UnrealBuildTool.cs b/Engine/Source/Programs/UnrealBuildTool/System/UnrealBuildTool.cs
--- a/Engine/Source/Programs/UnrealBuildTool/System/UnrealBuildTool.cs
+++ b/Engine/Source/Programs/UnrealBuildTool/System/UnrealBuildTool.cs
@@ -150,6 +150,16 @@ namespace UnrealBuildTool
             return bBuildingRocket;
         }
 
+		/// <summary>
+		/// Returns true if UnrealBuildTool is running using installed Engine components
+		/// </summary>
+		/// <returns>True if running using installed Engine components</returns>
+		static public bool IsEngineInstalled()
+		{
+			// For now this is only true when running Rocket
+			return RunningRocket();
+		}
+
         /// <summary>
         /// Returns true if a UProject file is available.  This should always be the case when working with game projects.  For engine or program targets, a UProject will not be available.
         /// </summary>
@@ -796,15 +806,6 @@ namespace UnrealBuildTool
                     UnrealTargetPlatform CheckPlatform;
                     UnrealTargetConfiguration CheckConfiguration;
                     UEBuildTarget.ParsePlatformAndConfiguration(Arguments, out CheckPlatform, out CheckConfiguration, false);
-                    //@todo SAS: Add a true Debug mode!
-                    if (RunningRocket())
-                    {
-                        // Only Development and Shipping are supported for engine modules
-                        if( CheckConfiguration != UnrealTargetConfiguration.Development && CheckConfiguration != UnrealTargetConfiguration.Shipping )
-                        {
-                            CheckConfiguration = UnrealTargetConfiguration.Development;
-                        }
-                    }
 
 					// @todo ubtmake: remove this when building with RPCUtility works
 					if (CheckPlatform == UnrealTargetPlatform.Mac || CheckPlatform == UnrealTargetPlatform.IOS)
@@ -1370,14 +1371,6 @@ namespace UnrealBuildTool
             UnrealTargetPlatform ResetPlatform = UnrealTargetPlatform.Unknown;
             UnrealTargetConfiguration ResetConfiguration;
             UEBuildTarget.ParsePlatformAndConfiguration(Arguments, out ResetPlatform, out ResetConfiguration);
-            if (RunningRocket())
-            {
-                // Only Development and Shipping are supported for engine modules
-                if( ResetConfiguration != UnrealTargetConfiguration.Development && ResetConfiguration != UnrealTargetConfiguration.Shipping )
-                {
-                    ResetConfiguration = UnrealTargetConfiguration.Development;
-                }
-            }
 			var BuildPlatform = UEBuildPlatform.GetBuildPlatform(ResetPlatform);
 			BuildPlatform.ResetBuildConfiguration(ResetPlatform, ResetConfiguration);
 
