Fixed use of custom icons and Info.plist files for Mac apps built remotely from PC. Fixes UE4EditorServices on GitHub.

diff --git a/Engine/Source/Programs/UnrealBuildTool/Mac/MacToolChain.cs b/Engine/Source/Programs/UnrealBuildTool/Mac/MacToolChain.cs
--- a/Engine/Source/Programs/UnrealBuildTool/Mac/MacToolChain.cs
+++ b/Engine/Source/Programs/UnrealBuildTool/Mac/MacToolChain.cs
@@ -1012,12 +1012,6 @@ namespace UnrealBuildTool
 						CustomBuildPath = Path.GetDirectoryName(UProjectFilePath) + "/Build/Mac";
 					}
 
-					if (!string.IsNullOrEmpty(CustomResourcesPath) && !string.IsNullOrEmpty(CustomBuildPath))
-					{
-						CustomResourcesPath = ConvertPath(Path.GetFullPath(CustomResourcesPath));
-						CustomBuildPath = ConvertPath(Path.GetFullPath(CustomBuildPath));
-					}
-
 					bool bBuildingEditor = GameName.EndsWith("Editor");
 
 					// Copy resources
@@ -1038,6 +1032,12 @@ namespace UnrealBuildTool
 								CustomIcon = DefaultIcon;
 							}
 						}
+
+						if (CustomIcon != DefaultIcon)
+						{
+							QueueFileForBatchUpload(FileItem.GetItemByFullPath(Path.GetFullPath(CustomIcon)));
+							CustomIcon = ConvertPath(CustomIcon);
+						}
 					}
 					AppendMacLine(FinalizeAppBundleScript, "cp -f \"{0}\" \"{2}.app/Contents/Resources/{1}.icns\"", CustomIcon, IconName, ExeName);
 
@@ -1051,6 +1051,11 @@ namespace UnrealBuildTool
 					{
 						InfoPlistFile = EngineSourcePath + "/Runtime/Launch/Resources/Mac/" + (bBuildingEditor ? "Info-Editor.plist" : "Info.plist");
 					}
+					else
+					{
+						QueueFileForBatchUpload(FileItem.GetItemByFullPath(Path.GetFullPath(InfoPlistFile)));
+						InfoPlistFile = ConvertPath(InfoPlistFile);
+					}
 					AppendMacLine(FinalizeAppBundleScript, "cp -f \"{0}\" \"{1}.app/Contents/Info.plist\"", InfoPlistFile, ExeName);
 
 					// Fix contents of Info.plist
