Linux: do not use initial-exec TLS model in dynamic libraries.

- Fixes editor problems on Ubuntu 14.04.1 and lower.
- Monolithic games/servers are unaffected.

diff --git a/Engine/Source/Programs/UnrealBuildTool/Linux/LinuxToolChain.cs b/Engine/Source/Programs/UnrealBuildTool/Linux/LinuxToolChain.cs
--- a/Engine/Source/Programs/UnrealBuildTool/Linux/LinuxToolChain.cs
+++ b/Engine/Source/Programs/UnrealBuildTool/Linux/LinuxToolChain.cs
@@ -372,6 +372,9 @@ namespace UnrealBuildTool
             if (CompileEnvironment.Config.bIsBuildingDLL)
             {
                 Result += " -fPIC";
+                // Use local-dynamic TLS model. This generates less efficient runtime code for __thread variables, but avoids problems of running into
+                // glibc/ld.so limit (DTV_SURPLUS) for number of dlopen()'ed DSOs with static TLS (see e.g. https://www.cygwin.com/ml/libc-help/2013-11/msg00033.html)
+                Result += " -ftls-model=local-dynamic";
             }
 
             //Result += " -v";                            // for better error diagnosis
