Support out-of-root game project targets in makefile generator.

- PR #861 contributed by [REDACTED_PERSON].

https://github.com/EpicGames/UnrealEngine/pull/861

diff --git a/Engine/Source/Programs/UnrealBuildTool/System/MakefileGenerator.cs b/Engine/Source/Programs/UnrealBuildTool/System/MakefileGenerator.cs
--- a/Engine/Source/Programs/UnrealBuildTool/System/MakefileGenerator.cs
+++ b/Engine/Source/Programs/UnrealBuildTool/System/MakefileGenerator.cs
@@ -60,13 +60,34 @@ namespace UnrealBuildTool
 
         private bool WriteMakefile()
         {
+            string GameProjectFile = "";
+            string BuildCommand = "";
+            string ProjectBuildCommand = "";
+
+            string MakeGameProjectFile = "";
+
+            var UnrealRootPath = Path.GetFullPath(ProjectFileGenerator.RootRelativePath);
+
+            if (!String.IsNullOrEmpty(GameProjectName))
+            {
+                GameProjectFile = UnrealBuildTool.GetUProjectFile();
+                MakeGameProjectFile = "GAMEPROJECTFILE =" + GameProjectFile + "\n";
+                ProjectBuildCommand = "PROJECTBUILD = mono $(UNREALROOTPATH)/Engine/Binaries/DotNET/UnrealBuildTool.exe\n";
+            }
+
+            BuildCommand = "BUILD = bash $(UNREALROOTPATH)/Engine/Build/BatchFiles/Linux/Build.sh\n";
+
             var FileName = "Makefile"; // MasterProjectName + ".mk";
             var MakefileContent = new StringBuilder();
             MakefileContent.Append(
                 "# Makefile generated by MakefileGenerator.cs\n" +
                 "# *DO NOT EDIT*\n\n" +
+                "UNREALROOTPATH = \"" + UnrealRootPath + "\"\n" +
+                MakeGameProjectFile + "\n" +
                 "TARGETS ="
-                );
+            );
+            String MakeProjectCmdArg = "";
+            String MakeBuildCommand = "";
             foreach (string Target in DiscoverTargets())
             {
                 var Basename = Path.GetFileNameWithoutExtension(Target);
@@ -78,33 +99,60 @@ namespace UnrealBuildTool
                         if (UnrealBuildTool.IsValidConfiguration(CurConfiguration))
                         {
                             var Confname = Enum.GetName(typeof(UnrealTargetConfiguration), CurConfiguration);
-                            MakefileContent.Append(String.Format(" \\\n    {0}-Linux-{1}", Basename, Confname));
+                            MakefileContent.Append(String.Format(" \\\n\t{0}-Linux-{1} ", Basename, Confname));
                         }
                     }
                 }
-                MakefileContent.Append(" \\\n    " + Basename);
+                MakefileContent.Append(" \\\n\t" + Basename);
             }
-            MakefileContent.Append("\n\nBUILD = Engine/Build/BatchFiles/Linux/Build.sh\n\n" +
+            MakefileContent.Append("\\\n\tconfigure");
+
+            MakefileContent.Append("\n\n" + BuildCommand + ProjectBuildCommand + "\n" +
                 "all: $(TARGETS)\n"
-                );
+            );
+
             foreach (string Target in DiscoverTargets())
             {
                 var Basename = Path.GetFileNameWithoutExtension(Target);
                 Basename = Path.GetFileNameWithoutExtension(Basename);
                 foreach (UnrealTargetConfiguration CurConfiguration in Enum.GetValues(typeof(UnrealTargetConfiguration)))
                 {
+                    if (Basename == GameProjectName || Basename == (GameProjectName + "Editor")) 
+                    {
+                        MakeProjectCmdArg = " -project=\"\\\"$(GAMEPROJECTFILE)\\\"\"";
+                        MakeBuildCommand = "$(PROJECTBUILD)";
+                    }
+                    else
+                    {
+                        MakeBuildCommand = "$(BUILD)";
+                    }
+
                     if (CurConfiguration != UnrealTargetConfiguration.Unknown && CurConfiguration != UnrealTargetConfiguration.Development)
                     {
                         if (UnrealBuildTool.IsValidConfiguration(CurConfiguration))
                         {
                             var Confname = Enum.GetName(typeof(UnrealTargetConfiguration), CurConfiguration);
-                            MakefileContent.Append(String.Format("\n{0}-Linux-{1}:\n\t$(BUILD) {0} Linux {1} $(ARGS)\n", Basename, Confname));
+                            MakefileContent.Append(String.Format("\n{1}-Linux-{2}:\n\t {0} {1} Linux {2} {3} $(ARGS)\n", MakeBuildCommand, Basename, Confname, MakeProjectCmdArg));
                         }
                     }
                 }
 
-                MakefileContent.Append(String.Format("\n{0}:\n\t$(BUILD) {0} Linux Development $(ARGS)\n", Basename));
+                MakefileContent.Append(String.Format("\n{1}:\n\t {0} {1} Linux Development {2} $(ARGS)\n", MakeBuildCommand, Basename, MakeProjectCmdArg));
+            }
+
+            MakefileContent.Append("\nconfigure:\n");
+            if (!String.IsNullOrEmpty (GameProjectName))
+            {
+                // Make sure UBT is updated.
+                MakefileContent.Append("\txbuild /property:Configuration=Development /property:TargetFrameworkVersion=v4.0 /verbosity:quiet /nologo ");
+                MakefileContent.Append("$(UNREALROOTPATH)/Engine/Source/Programs/UnrealBuildTool/UnrealBuildTool_Mono.csproj\n");
+                MakefileContent.Append("\t$(PROJECTBUILD) -makefile -qmakefile -cmakefile -project=\"\\\"$(GAMEPROJECTFILE)\\\"\" -game -engine \n");
             }
+            else
+            {
+                MakefileContent.Append("\tbash $(UNREALROOTPATH)/GenerateProjectFiles.sh \n");
+            }
+
             MakefileContent.Append("\n.PHONY: $(TARGETS)\n");
             var FullFileName = Path.Combine(MasterProjectRelativePath, FileName);
             return WriteFileIfChanged(FullFileName, MakefileContent.ToString());
@@ -112,7 +160,7 @@ namespace UnrealBuildTool
 
         /// ProjectFileGenerator interface
         //protected override bool WriteMasterProjectFile( ProjectFile UBTProject )
-        protected override bool WriteProjectFiles ()
+        protected override bool WriteProjectFiles()
         {
 			return WriteMakefile();
         }
