Linux: do not inhibit optimizations in Development and higher configurations.

- bCreateDebugInfo is pretty much true all the time.

diff --git a/Engine/Source/Programs/UnrealBuildTool/Linux/LinuxToolChain.cs b/Engine/Source/Programs/UnrealBuildTool/Linux/LinuxToolChain.cs
--- a/Engine/Source/Programs/UnrealBuildTool/Linux/LinuxToolChain.cs
+++ b/Engine/Source/Programs/UnrealBuildTool/Linux/LinuxToolChain.cs
@@ -344,8 +344,8 @@ namespace UnrealBuildTool
                 Result += " -fno-inline";                   // disable inlining for better debuggability (e.g. callstacks, "skip file" in gdb)
             }
 
-            // debug info
-            if (CompileEnvironment.Config.bCreateDebugInfo)
+            // debug info (bCreateDebugInfo is normally set for all configurations, and we don't want it to affect Development/Shipping performance)
+            if (CompileEnvironment.Config.bCreateDebugInfo && CompileEnvironment.Config.Target.Configuration == CPPTargetConfiguration.Debug)
             {
                 Result += " -g3";
                 Result += " -fno-omit-frame-pointer";
@@ -553,11 +553,6 @@ namespace UnrealBuildTool
             string Arguments = GetCLArguments_Global(CompileEnvironment);
             string PCHArguments = "";
 
-            if (CompileEnvironment.Config.bIsBuildingDLL)
-            {
-                Arguments += " -fPIC";
-            }
-
             if (CompileEnvironment.Config.PrecompiledHeaderAction == PrecompiledHeaderAction.Include)
             {
                 // Add the precompiled header file's path to the include path so Clang can find it.
