Swap RunningRocket() check over to IsEngineInstalled() - they may not be the same.

diff --git a/Engine/Source/Programs/UnrealBuildTool/System/JunkDeleter.cs b/Engine/Source/Programs/UnrealBuildTool/System/JunkDeleter.cs
--- a/Engine/Source/Programs/UnrealBuildTool/System/JunkDeleter.cs
+++ b/Engine/Source/Programs/UnrealBuildTool/System/JunkDeleter.cs
@@ -17,7 +17,7 @@ namespace UnrealBuildTool
 		{
 			var JunkStartTime = DateTime.UtcNow;
 
-			if (UnrealBuildTool.RunningRocket() == false)
+			if (UnrealBuildTool.IsEngineInstalled() == false)
 			{
 				List<string> JunkManifest = LoadJunkManifest();
 				DeleteAllJunk(JunkManifest);
