fix for platform string sitting somewhere above the engine or project directory causing issues with project generation and staging of files
UE-11502
UE-11479

diff --git a/Engine/Source/Programs/AutomationTool/DeploymentContext.cs b/Engine/Source/Programs/AutomationTool/DeploymentContext.cs
--- a/Engine/Source/Programs/AutomationTool/DeploymentContext.cs
+++ b/Engine/Source/Programs/AutomationTool/DeploymentContext.cs
@@ -371,6 +371,17 @@ public class DeploymentContext //: ProjectParams
                                     Search = Search.Substring(LocalRoot.Length);
                                 }
                             }
+							if (Search.StartsWith(ProjectRoot, StringComparison.InvariantCultureIgnoreCase))
+							{
+								if (ProjectRoot.EndsWith("\\") || ProjectRoot.EndsWith("/"))
+								{
+									Search = Search.Substring(ProjectRoot.Length - 1);
+								}
+								else
+								{
+									Search = Search.Substring(ProjectRoot.Length);
+								}
+							}
                             if (Search.IndexOf(CommandUtils.CombinePaths("/" + Plat.ToString() + "/"), 0, StringComparison.InvariantCultureIgnoreCase) >= 0)
                             {
                                 OtherPlatform = true;
diff --git a/Engine/Source/Programs/UnrealBuildTool/System/ProjectFileGenerator.cs b/Engine/Source/Programs/UnrealBuildTool/System/ProjectFileGenerator.cs
--- a/Engine/Source/Programs/UnrealBuildTool/System/ProjectFileGenerator.cs
+++ b/Engine/Source/Programs/UnrealBuildTool/System/ProjectFileGenerator.cs
@@ -964,11 +964,26 @@ namespace UnrealBuildTool
 			{
 				var CleanTargetFileName = Utils.CleanDirectorySeparators( CurTargetFile );
 
+				// remove the local root
+				var LocalRoot = Path.GetFullPath(RootRelativePath);
+				var Search = CleanTargetFileName;
+				if (Search.StartsWith(LocalRoot, StringComparison.InvariantCultureIgnoreCase))
+				{
+					if (LocalRoot.EndsWith("\\") || LocalRoot.EndsWith("/"))
+					{
+						Search = Search.Substring(LocalRoot.Length - 1);
+					}
+					else
+					{
+						Search = Search.Substring(LocalRoot.Length);
+					}
+				}
+
 				// Skip targets in unsupported platform directories
 				bool IncludeThisTarget = true;
 				foreach( var CurPlatformName in UnsupportedPlatformNameStrings )
 				{
-					if( CleanTargetFileName.IndexOf( Path.DirectorySeparatorChar + CurPlatformName + Path.DirectorySeparatorChar, StringComparison.InvariantCultureIgnoreCase ) != -1 )
+					if (Search.IndexOf(Path.DirectorySeparatorChar + CurPlatformName + Path.DirectorySeparatorChar, StringComparison.InvariantCultureIgnoreCase) != -1)
 					{
 						IncludeThisTarget = false;
 						break;
