Fixed typos in reverseLandscape and sensorLandscape AndroidManifest.xml generation
#ue4
#android

diff --git a/Engine/Source/Programs/UnrealBuildTool/Android/UEDeployAndroid.cs b/Engine/Source/Programs/UnrealBuildTool/Android/UEDeployAndroid.cs
--- a/Engine/Source/Programs/UnrealBuildTool/Android/UEDeployAndroid.cs
+++ b/Engine/Source/Programs/UnrealBuildTool/Android/UEDeployAndroid.cs
@@ -598,9 +598,9 @@ namespace UnrealBuildTool.Android
 					return "sensorPortrait";
 				case "landscape":
 					return "landscape";
-				case "reversepandscape":
+				case "reverselandscape":
 					return "reverseLandscape";
-				case "sensorpandscape":
+				case "sensorlandscape":
 					return "sensorLandscape";
 				case "sensor":
 					return "sensor";
