- Added support for checking Android build settings to be non-default
- Refactored code to make it easier to check .ini settings for non-default build settings

diff --git a/Engine/Source/Programs/UnrealBuildTool/Android/UEBuildAndroid.cs b/Engine/Source/Programs/UnrealBuildTool/Android/UEBuildAndroid.cs
--- a/Engine/Source/Programs/UnrealBuildTool/Android/UEBuildAndroid.cs
+++ b/Engine/Source/Programs/UnrealBuildTool/Android/UEBuildAndroid.cs
@@ -254,6 +254,24 @@ namespace UnrealBuildTool
 
 			BuildConfiguration.bUseSharedPCHs = false;
 		}
+			
+		public override bool HasDefaultBuildConfig(UnrealTargetPlatform Platform, string ProjectPath)
+		{
+			string[] BoolKeys = new string[] {
+				"bBuildForArmV7", "bBuildForArm64", "bBuildForX86", "bBuildForX8664", 
+				"bBuildForES2", "bBuildForES31",
+			};
+
+			// look up iOS specific settings
+			if (!DoProjectSettingsMatchDefault(Platform, ProjectPath, "/Script/AndroidRuntimeSettings.AndroidRuntimeSettings",
+				BoolKeys, null, null))
+			{
+				return false;
+			}
+
+			// check the base settings
+			return base.HasDefaultBuildConfig(Platform, ProjectPath);
+		}
 
 		public override bool ShouldCompileMonolithicBinary(UnrealTargetPlatform InPlatform)
 		{
diff --git a/Engine/Source/Programs/UnrealBuildTool/Configuration/UEBuildPlatform.cs b/Engine/Source/Programs/UnrealBuildTool/Configuration/UEBuildPlatform.cs
--- a/Engine/Source/Programs/UnrealBuildTool/Configuration/UEBuildPlatform.cs
+++ b/Engine/Source/Programs/UnrealBuildTool/Configuration/UEBuildPlatform.cs
@@ -935,133 +935,69 @@ namespace UnrealBuildTool
 		{
 		}
 
-		/**
-		 * Check for the default configuration
-		 *
-		 * return true if the project uses the default build config
-		 */
-		public virtual bool HasDefaultBuildConfig(UnrealTargetPlatform Platform, string ProjectPath)
+		protected static bool DoProjectSettingsMatchDefault(UnrealTargetPlatform Platform, string ProjectPath, string Section, string[] BoolKeys, string[] IntKeys, string[] StringKeys)
 		{
 			ConfigCacheIni ProjIni = new ConfigCacheIni(Platform, "Engine", ProjectPath);
 			ConfigCacheIni DefaultIni = new ConfigCacheIni(Platform, "Engine", null);
 
-			bool bDefValue = UEBuildConfiguration.bCompileAPEX;
-			DefaultIni.GetBool("/Script/BuildSettings.BuildSettings", "bCompileApex", out bDefValue);
-			bool bProjValue = UEBuildConfiguration.bCompileAPEX;
-			ProjIni.GetBool("/Script/BuildSettings.BuildSettings", "bCompileApex", out bProjValue);
-			if (bDefValue != bProjValue)
-			{
-				return false;
-			}
-
-			bDefValue = UEBuildConfiguration.bCompileBox2D;
-			DefaultIni.GetBool("/Script/BuildSettings.BuildSettings", "bCompileBox2D", out bDefValue);
-			bProjValue = UEBuildConfiguration.bCompileBox2D;
-			ProjIni.GetBool("/Script/BuildSettings.BuildSettings", "bCompileBox2D", out bProjValue);
-			if (bDefValue != bProjValue)
-			{
-				return false;
-			}
-
-			bDefValue = UEBuildConfiguration.bCompileICU;
-			DefaultIni.GetBool("/Script/BuildSettings.BuildSettings", "bCompileICU", out bDefValue);
-			bProjValue = UEBuildConfiguration.bCompileICU;
-			ProjIni.GetBool("/Script/BuildSettings.BuildSettings", "bCompileICU", out bProjValue);
-			if (bDefValue != bProjValue)
-			{
-				return false;
-			}
-
-			bDefValue = UEBuildConfiguration.bCompileSimplygon;
-			DefaultIni.GetBool("/Script/BuildSettings.BuildSettings", "bCompileSimplygon", out bDefValue);
-			bProjValue = UEBuildConfiguration.bCompileSimplygon;
-			ProjIni.GetBool("/Script/BuildSettings.BuildSettings", "bCompileSimplygon", out bProjValue);
-			if (bDefValue != bProjValue)
-			{
-				return false;
-			}
-
-			bDefValue = UEBuildConfiguration.bCompileLeanAndMeanUE;
-			DefaultIni.GetBool("/Script/BuildSettings.BuildSettings", "bCompileLeanAndMeanUE", out bDefValue);
-			bProjValue = UEBuildConfiguration.bCompileLeanAndMeanUE;
-			ProjIni.GetBool("/Script/BuildSettings.BuildSettings", "bCompileLeanAndMeanUE", out bProjValue);
-			if (bDefValue != bProjValue)
-			{
-				return false;
-			}
-
-			bDefValue = UEBuildConfiguration.bIncludeADO;
-			DefaultIni.GetBool("/Script/BuildSettings.BuildSettings", "bIncludeADO", out bDefValue);
-			bProjValue = UEBuildConfiguration.bIncludeADO;
-			ProjIni.GetBool("/Script/BuildSettings.BuildSettings", "bIncludeADO", out bProjValue);
-			if (bDefValue != bProjValue)
-			{
-				return false;
-			}
-
-			bDefValue = UEBuildConfiguration.bCompileRecast;
-			DefaultIni.GetBool("/Script/BuildSettings.BuildSettings", "bCompileRecast", out bDefValue);
-			bProjValue = UEBuildConfiguration.bCompileRecast;
-			ProjIni.GetBool("/Script/BuildSettings.BuildSettings", "bCompileRecast", out bProjValue);
-			if (bDefValue != bProjValue)
+			// look at all bool values
+			if (BoolKeys != null) foreach (string Key in BoolKeys)
 			{
-				return false;
-			}
-
-			bDefValue = UEBuildConfiguration.bCompileSpeedTree;
-			DefaultIni.GetBool("/Script/BuildSettings.BuildSettings", "bCompileSpeedTree", out bDefValue);
-			bProjValue = UEBuildConfiguration.bCompileSpeedTree;
-			ProjIni.GetBool("/Script/BuildSettings.BuildSettings", "bCompileSpeedTree", out bProjValue);
-			if (bDefValue != bProjValue)
-			{
-				return false;
+				bool Default = false, Project = false;
+				DefaultIni.GetBool(Section, Key, out Default);
+				ProjIni.GetBool(Section, Key, out Project);
+				Console.WriteLine("Comparing {0} settings: {1}, {2}", Section + ":" + Key, Default, Project); 
+				if (Default != Project)
+				{
+					return false;
+				}
 			}
 
-			bDefValue = UEBuildConfiguration.bCompileWithPluginSupport;
-			DefaultIni.GetBool("/Script/BuildSettings.BuildSettings", "bCompileWithPluginSupport", out bDefValue);
-			bProjValue = UEBuildConfiguration.bCompileWithPluginSupport;
-			ProjIni.GetBool("/Script/BuildSettings.BuildSettings", "bCompileWithPluginSupport", out bProjValue);
-			if (bDefValue != bProjValue)
+			// look at all int values
+			if (IntKeys != null) foreach (string Key in IntKeys)
 			{
-				return false;
+				int Default = 0, Project = 0;
+				DefaultIni.GetInt32(Section, Key, out Default);
+				ProjIni.GetInt32(Section, Key, out Project);
+				if (Default != Project)
+				{
+					return false;
+				}
 			}
 
-			bDefValue = UEBuildConfiguration.bCompilePhysXVehicle;
-			DefaultIni.GetBool("/Script/BuildSettings.BuildSettings", "bCompilePhysXVehicle", out bDefValue);
-			bProjValue = UEBuildConfiguration.bCompilePhysXVehicle;
-			ProjIni.GetBool("/Script/BuildSettings.BuildSettings", "bCompilePhysXVehicle", out bProjValue);
-			if (bDefValue != bProjValue)
+			// look for all string values
+			if (StringKeys != null) foreach (string Key in StringKeys)
 			{
-				return false;
+				string Default = "", Project = "";
+				DefaultIni.GetString(Section, Key, out Default);
+				ProjIni.GetString(Section, Key, out Project);
+				if (Default != Project)
+				{
+					return false;
+				}
 			}
 
-			bDefValue = UEBuildConfiguration.bCompileFreeType;
-			DefaultIni.GetBool("/Script/BuildSettings.BuildSettings", "bCompileFreeType", out bDefValue);
-			bProjValue = UEBuildConfiguration.bCompileFreeType;
-			ProjIni.GetBool("/Script/BuildSettings.BuildSettings", "bCompileFreeType", out bProjValue);
-			if (bDefValue != bProjValue)
-			{
-				return false;
-			}
+			// if we get here, we match all important settings
+			return true;
+		}
 
-			bDefValue = UEBuildConfiguration.bCompileForSize;
-			DefaultIni.GetBool("/Script/BuildSettings.BuildSettings", "bCompileForSize", out bDefValue);
-			bProjValue = UEBuildConfiguration.bCompileForSize;
-			ProjIni.GetBool("/Script/BuildSettings.BuildSettings", "bCompileForSize", out bProjValue);
-			if (bDefValue != bProjValue)
-			{
-				return false;
-			}
+		/**
+		 * Check for the default configuration
+		 *
+		 * return true if the project uses the default build config
+		 */
+		public virtual bool HasDefaultBuildConfig(UnrealTargetPlatform Platform, string ProjectPath)
+		{
+			Console.WriteLine("HasDefaultBuildConfig!");
+			string[] BoolKeys = new string[] {
+				"bCompileApex", "bCompileBox2D", "bCompileICU", "bCompileSimplygon", 
+				"bCompileLeanAndMeanUE", "bIncludeADO", "bCompileRecast", "bCompileSpeedTree", 
+				"bCompileWithPluginSupport", "bCompilePhysXVehicle", "bCompileFreeType", 
+				"bCompileForSize", "bCompileCEF3"
+			};
 
-			bDefValue = UEBuildConfiguration.bCompileCEF3;
-			DefaultIni.GetBool("/Script/BuildSettings.BuildSettings", "bCompileCEF3", out bDefValue);
-			bProjValue = UEBuildConfiguration.bCompileCEF3;
-			ProjIni.GetBool("/Script/BuildSettings.BuildSettings", "bCompileCEF3", out bProjValue);
-			if (bDefValue != bProjValue)
-			{
-				return false;
-			}
-			return true;
+			return DoProjectSettingsMatchDefault(Platform, ProjectPath, "/Script/BuildSettings.BuildSettings",
+				BoolKeys, null, null);
 		}
 	}
 
diff --git a/Engine/Source/Programs/UnrealBuildTool/IOS/UEBuildIOS.cs b/Engine/Source/Programs/UnrealBuildTool/IOS/UEBuildIOS.cs
--- a/Engine/Source/Programs/UnrealBuildTool/IOS/UEBuildIOS.cs
+++ b/Engine/Source/Programs/UnrealBuildTool/IOS/UEBuildIOS.cs
@@ -243,65 +243,22 @@ namespace UnrealBuildTool
 		 */
 		public override bool HasDefaultBuildConfig(UnrealTargetPlatform Platform, string ProjectPath)
 		{
-			ConfigCacheIni ProjIni = new ConfigCacheIni(Platform, "Engine", ProjectPath);
-			ConfigCacheIni DefaultIni = new ConfigCacheIni(Platform, "Engine", null);
-
-			string DefaultMinVersion = "IOS_6";
-			string ProjectMinVersion = DefaultMinVersion;
-			ProjIni.GetString("/Script/IOSRuntimeSettings.IOSRuntimeSettings", "MinimumiOSVersion", out ProjectMinVersion);
-			DefaultIni.GetString("/Script/IOSRuntimeSettings.IOSRuntimeSettings", "MinimumiOSVersion", out DefaultMinVersion);
-			if (DefaultMinVersion != ProjectMinVersion)
-			{
-				return false;
-			}
-
-			bool bDefaultBuild = true;
-			bool bProjectBuild = true;
-			ProjIni.GetBool("/Script/IOSRuntimeSettings.IOSRuntimeSettings", "bDevForArmV7", out bProjectBuild);
-			DefaultIni.GetBool("/Script/IOSRuntimeSettings.IOSRuntimeSettings", "bDevForArmV7", out bDefaultBuild);
-			if (bDefaultBuild != bProjectBuild)
-			{
-				return false;
-			}
-			ProjIni.GetBool("/Script/IOSRuntimeSettings.IOSRuntimeSettings", "bDevForArm64", out bProjectBuild);
-			DefaultIni.GetBool("/Script/IOSRuntimeSettings.IOSRuntimeSettings", "bDevForArm64", out bDefaultBuild);
-			if (bDefaultBuild != bProjectBuild)
-			{
-				return false;
-			}
-			ProjIni.GetBool("/Script/IOSRuntimeSettings.IOSRuntimeSettings", "bDevForArmV7S", out bProjectBuild);
-			DefaultIni.GetBool("/Script/IOSRuntimeSettings.IOSRuntimeSettings", "bDevForArmV7S", out bDefaultBuild);
-			if (bDefaultBuild != bProjectBuild)
-			{
-				return false;
-			}
-			ProjIni.GetBool("/Script/IOSRuntimeSettings.IOSRuntimeSettings", "bShipForArmV7", out bProjectBuild);
-			DefaultIni.GetBool("/Script/IOSRuntimeSettings.IOSRuntimeSettings", "bShipForArmV7", out bDefaultBuild);
-			if (bDefaultBuild != bProjectBuild)
-			{
-				return false;
-			}
-			ProjIni.GetBool("/Script/IOSRuntimeSettings.IOSRuntimeSettings", "bShipForArm64", out bProjectBuild);
-			DefaultIni.GetBool("/Script/IOSRuntimeSettings.IOSRuntimeSettings", "bShipForArm64", out bDefaultBuild);
-			if (bDefaultBuild != bProjectBuild)
-			{
-				return false;
-			}
-			ProjIni.GetBool("/Script/IOSRuntimeSettings.IOSRuntimeSettings", "bShipForArmV7S", out bProjectBuild);
-			DefaultIni.GetBool("/Script/IOSRuntimeSettings.IOSRuntimeSettings", "bShipForArmV7S", out bDefaultBuild);
-			if (bDefaultBuild != bProjectBuild)
-			{
-				return false;
-			}
-
-			// determine if we need to generate the dsym
-			ProjIni.GetBool("/Script/IOSRuntimeSettings.IOSRuntimeSettings", "bGenerateSYMFile", out bProjectBuild);
-			DefaultIni.GetBool("/Script/IOSRuntimeSettings.IOSRuntimeSettings", "bGenerateSYMFile", out bDefaultBuild);
-			if (bDefaultBuild != bProjectBuild)
+			string[] BoolKeys = new string[] {
+				"bDevForArmV7", "bDevForArm64", "bDevForArmV7S", "bShipForArmV7", 
+				"bShipForArm64", "bShipForArmV7S", "bGenerateSYMFile",
+			};
+			string[] StringKeys = new string[] {
+				"MinimumiOSVersion", 
+			};
+
+			// look up iOS specific settings
+			if (!DoProjectSettingsMatchDefault(Platform, ProjectPath, "/Script/IOSRuntimeSettings.IOSRuntimeSettings",
+				    BoolKeys, null, StringKeys))
 			{
 				return false;
 			}
 
+			// check the base settings
 			return base.HasDefaultBuildConfig(Platform, ProjectPath);
 		}
 
