- Removed some log spam

diff --git a/Engine/Source/Programs/UnrealBuildTool/Configuration/UEBuildPlatform.cs b/Engine/Source/Programs/UnrealBuildTool/Configuration/UEBuildPlatform.cs
--- a/Engine/Source/Programs/UnrealBuildTool/Configuration/UEBuildPlatform.cs
+++ b/Engine/Source/Programs/UnrealBuildTool/Configuration/UEBuildPlatform.cs
@@ -946,7 +946,6 @@ namespace UnrealBuildTool
 				bool Default = false, Project = false;
 				DefaultIni.GetBool(Section, Key, out Default);
 				ProjIni.GetBool(Section, Key, out Project);
-				Console.WriteLine("Comparing {0} settings: {1}, {2}", Section + ":" + Key, Default, Project); 
 				if (Default != Project)
 				{
 					return false;
