Fix linker errors on non-monolithic shipping builds.

diff --git a/Engine/Source/Programs/UnrealBuildTool/WinRT/WinRTToolChain.cs b/Engine/Source/Programs/UnrealBuildTool/WinRT/WinRTToolChain.cs
--- a/Engine/Source/Programs/UnrealBuildTool/WinRT/WinRTToolChain.cs
+++ b/Engine/Source/Programs/UnrealBuildTool/WinRT/WinRTToolChain.cs
@@ -798,7 +798,7 @@ namespace UnrealBuildTool
 			{
 				// Xbox 360 LTCG does not seem to produce those.
 				if (LinkEnvironment.Config.bHasExports &&
-					(LinkEnvironment.Config.Target.Configuration != CPPTargetConfiguration.Shipping))
+					(LinkEnvironment.Config.Target.Configuration != CPPTargetConfiguration.Shipping) && (LinkEnvironment.Config.Target.Platform != CPPTargetPlatform.XboxOne))
 				{
 					// Write the import library to the output directory for nFringe support.
 					FileItem ImportLibraryFile = FileItem.GetItemByPath(ImportLibraryFilePath);
diff --git a/Engine/Source/Programs/UnrealBuildTool/Windows/VCToolChain.cs b/Engine/Source/Programs/UnrealBuildTool/Windows/VCToolChain.cs
--- a/Engine/Source/Programs/UnrealBuildTool/Windows/VCToolChain.cs
+++ b/Engine/Source/Programs/UnrealBuildTool/Windows/VCToolChain.cs
@@ -1513,7 +1513,7 @@ namespace UnrealBuildTool
 			{
 				// Xbox 360 LTCG does not seem to produce those.
 				if (LinkEnvironment.Config.bHasExports &&
-					(LinkEnvironment.Config.Target.Configuration != CPPTargetConfiguration.Shipping))
+					(LinkEnvironment.Config.Target.Configuration != CPPTargetConfiguration.Shipping) && (LinkEnvironment.Config.Target.Platform != CPPTargetPlatform.XboxOne))
 				{
 					// Write the import library to the output directory for nFringe support.
 					FileItem ImportLibraryFile = FileItem.GetItemByPath(ImportLibraryFilePath);
