Create targets using the TargetDescriptor, rather than separating the fields out and passing them all to constructors individually. Makes it easier to extend, and is an awful lot easier to read. Also change the target's GameName field to TargetName, which is what everything seems to actually require.

diff --git a/Engine/Source/Programs/UnrealBuildTool/Configuration/UEBuildClient.cs b/Engine/Source/Programs/UnrealBuildTool/Configuration/UEBuildClient.cs
--- a/Engine/Source/Programs/UnrealBuildTool/Configuration/UEBuildClient.cs
+++ b/Engine/Source/Programs/UnrealBuildTool/Configuration/UEBuildClient.cs
@@ -10,29 +10,8 @@ namespace UnrealBuildTool
 	[Serializable]
     public class UEBuildClient : UEBuildTarget
 	{
-		public UEBuildClient(
-			string InGameName, 
-			UnrealTargetPlatform InPlatform, 
-			UnrealTargetConfiguration InConfiguration,
-			TargetRules InRulesObject,
-			List<string> InAdditionalDefinitions, 
-			string InRemoteRoot, 
-			List<OnlyModule> InOnlyModules,
-			bool bInEditorRecompile)
-			// NOTE: If we're building a monolithic binary, then the game and engine code are linked together into one
-			//       program executable, so we want the application name to be the game name.  In the case of a modular
-			//       binary, we use 'UnrealEngine' for our application name
-            : base(
-                InAppName: UEBuildTarget.GetBinaryBaseName(InGameName, InRulesObject, InPlatform, InConfiguration, "Client"),
-                InGameName: InGameName,
-                InPlatform: InPlatform,
-                InConfiguration: InConfiguration,
-                InRulesObject: InRulesObject,
-                InAdditionalDefinitions: InAdditionalDefinitions,
-                InRemoteRoot: InRemoteRoot,
-                InOnlyModules: InOnlyModules,
-				bInEditorRecompile: bInEditorRecompile
-			)
+		public UEBuildClient(TargetDescriptor InDesc, TargetRules InRulesObject)
+            : base(GetDefaultAppName(InDesc, InRulesObject, "Client"), InDesc, InRulesObject)
         {
             if (ShouldCompileMonolithic())
             {
@@ -46,7 +25,7 @@ namespace UnrealBuildTool
                     {
 						for (int Index = 0; Index < OutputPaths.Length; Index++ )
 						{
-							OutputPaths[Index] = OutputPaths[Index].Replace("Engine\\Binaries", InGameName + "\\Binaries");
+							OutputPaths[Index] = OutputPaths[Index].Replace("Engine\\Binaries", InDesc.TargetName + "\\Binaries");
 						}
                     }
                 }
diff --git a/Engine/Source/Programs/UnrealBuildTool/Configuration/UEBuildEditor.cs b/Engine/Source/Programs/UnrealBuildTool/Configuration/UEBuildEditor.cs
--- a/Engine/Source/Programs/UnrealBuildTool/Configuration/UEBuildEditor.cs
+++ b/Engine/Source/Programs/UnrealBuildTool/Configuration/UEBuildEditor.cs
@@ -10,35 +10,8 @@ namespace UnrealBuildTool
 	[Serializable]
 	public class UEBuildEditor : UEBuildTarget
 	{
-		public UEBuildEditor(
-			string InGameName, 
-			UnrealTargetPlatform InPlatform, 
-			UnrealTargetConfiguration InConfiguration,
-			TargetRules InRulesObject,
-			List<string> InAdditionalDefinitions, 
-			string InRemoteRoot, 
-			List<OnlyModule> InOnlyModules,
-			bool bInEditorRecompile)
-			// NOTE: If we're building a monolithic binary, then the game and engine code are linked together into one
-			//       program executable, so we want the application name to be the game name.  In the case of a modular
-			//       binary, we use 'UnrealEngine' for our application name
-			: base(
-				InAppName:UEBuildTarget.GetBinaryBaseName(
-					InGameName, 
-					InRulesObject, 
-					InPlatform, 
-					InConfiguration, 
-					(InRulesObject.Type == TargetRules.TargetType.Editor) ? "Editor" : ""
-					),
-				InGameName:InGameName, 
-				InPlatform:InPlatform, 
-				InConfiguration:InConfiguration,
-				InRulesObject: InRulesObject, 
-				InAdditionalDefinitions:InAdditionalDefinitions, 
-				InRemoteRoot:InRemoteRoot, 
-				InOnlyModules:InOnlyModules,
-				bInEditorRecompile: bInEditorRecompile
-			)
+		public UEBuildEditor(TargetDescriptor InDesc, TargetRules InRulesObject)
+			: base(GetDefaultAppName(InDesc, InRulesObject, "Editor"), InDesc, InRulesObject)
 		{
 		}
 
diff --git a/Engine/Source/Programs/UnrealBuildTool/Configuration/UEBuildGame.cs b/Engine/Source/Programs/UnrealBuildTool/Configuration/UEBuildGame.cs
--- a/Engine/Source/Programs/UnrealBuildTool/Configuration/UEBuildGame.cs
+++ b/Engine/Source/Programs/UnrealBuildTool/Configuration/UEBuildGame.cs
@@ -10,29 +10,8 @@ namespace UnrealBuildTool
 	[Serializable]
 	public class UEBuildGame : UEBuildTarget
 	{
-		public UEBuildGame(
-			string InGameName, 
-			UnrealTargetPlatform InPlatform, 
-			UnrealTargetConfiguration InConfiguration,
-			TargetRules InRulesObject,
-			List<string> InAdditionalDefinitions, 
-			string InRemoteRoot, 
-			List<OnlyModule> InOnlyModules,
-			bool bInEditorRecompile)
-			// NOTE: If we're building a monolithic binary, then the game and engine code are linked together into one
-			//       program executable, so we want the application name to be the game name.  In the case of a modular
-			//       binary, we use 'UnrealEngine' for our application name
-			: base(
-				InAppName: UEBuildTarget.GetBinaryBaseName(InGameName, InRulesObject, InPlatform, InConfiguration, ""),
-				InGameName:InGameName,
-				InPlatform:InPlatform,
-				InConfiguration:InConfiguration,
-				InRulesObject: InRulesObject, 
-				InAdditionalDefinitions:InAdditionalDefinitions,
-				InRemoteRoot:InRemoteRoot,
-				InOnlyModules:InOnlyModules,
-				bInEditorRecompile: bInEditorRecompile
-			)
+		public UEBuildGame(TargetDescriptor InDesc, TargetRules InRulesObject)
+			: base(GetDefaultAppName(InDesc, InRulesObject, ""), InDesc, InRulesObject)
 		{
 			if (ShouldCompileMonolithic())
 			{
@@ -46,14 +25,13 @@ namespace UnrealBuildTool
 					{
 						for (int Index = 0; Index < OutputPaths.Length; Index++)
 						{
-							OutputPaths[Index] = OutputPaths[Index].Replace("Engine\\Binaries", InGameName + "\\Binaries");
+							OutputPaths[Index] = OutputPaths[Index].Replace("Engine\\Binaries", InDesc.TargetName + "\\Binaries");
 						}
 					}
 				}
 			}
 		}
 
-
 		//
 		// UEBuildTarget interface.
 		//
diff --git a/Engine/Source/Programs/UnrealBuildTool/Configuration/UEBuildServer.cs b/Engine/Source/Programs/UnrealBuildTool/Configuration/UEBuildServer.cs
--- a/Engine/Source/Programs/UnrealBuildTool/Configuration/UEBuildServer.cs
+++ b/Engine/Source/Programs/UnrealBuildTool/Configuration/UEBuildServer.cs
@@ -10,29 +10,8 @@ namespace UnrealBuildTool
 	[Serializable]
 	public class UEBuildServer : UEBuildTarget
 	{
-		public UEBuildServer(
-			string InGameName, 
-			UnrealTargetPlatform InPlatform, 
-			UnrealTargetConfiguration InConfiguration,
-			TargetRules InRulesObject,
-			List<string> InAdditionalDefinitions, 
-			string InRemoteRoot, 
-			List<OnlyModule> InOnlyModules,
-			bool bInEditorRecompile)
-			// NOTE: If we're building a monolithic binary, then the game and engine code are linked together into one
-			//       program executable, so we want the application name to be the game name.  In the case of a modular
-			//       binary, we use 'UnrealEngine' for our application name
-			: base(
-				InAppName:UEBuildTarget.GetBinaryBaseName(InGameName, InRulesObject, InPlatform, InConfiguration, "Server"),
-				InGameName:InGameName,
-				InPlatform:InPlatform,
-				InConfiguration:InConfiguration,
-				InRulesObject: InRulesObject, 
-				InAdditionalDefinitions:InAdditionalDefinitions,
-				InRemoteRoot:InRemoteRoot,
-				InOnlyModules:InOnlyModules,
-				bInEditorRecompile: bInEditorRecompile
-			)
+		public UEBuildServer(TargetDescriptor InDesc, TargetRules InRulesObject)
+			: base(GetDefaultAppName(InDesc, InRulesObject, "Server"), InDesc, InRulesObject)
 		{
 		}
 
diff --git a/Engine/Source/Programs/UnrealBuildTool/Configuration/UEBuildTarget.cs b/Engine/Source/Programs/UnrealBuildTool/Configuration/UEBuildTarget.cs
--- a/Engine/Source/Programs/UnrealBuildTool/Configuration/UEBuildTarget.cs
+++ b/Engine/Source/Programs/UnrealBuildTool/Configuration/UEBuildTarget.cs
@@ -182,7 +182,7 @@ namespace UnrealBuildTool
 
 		public string GetTargetName()
 		{
-			return !String.IsNullOrEmpty(GameName) ? GameName : AppName;
+			return TargetName;
 		}
 
 		public static UnrealTargetPlatform CPPTargetPlatformToUnrealTargetPlatform(CPPTargetPlatform InCPPPlatform)
@@ -452,20 +452,12 @@ namespace UnrealBuildTool
 
 		public static UEBuildTarget CreateTarget( TargetDescriptor Desc )
 		{
-			string TargetName = Desc.TargetName;
-			List<string> AdditionalDefinitions = Desc.AdditionalDefinitions;
-			UnrealTargetPlatform Platform = Desc.Platform;
-			UnrealTargetConfiguration Configuration = Desc.Configuration;
-			string RemoteRoot = Desc.RemoteRoot;
-			List<OnlyModule> OnlyModules = Desc.OnlyModules;
-			bool bIsEditorRecompile = Desc.bIsEditorRecompile;
-
 			UEBuildTarget BuildTarget = null;
 			if( !ProjectFileGenerator.bGenerateProjectFiles )
 			{
 				// Configure the rules compiler
-				string PossibleAssemblyName = TargetName;
-				if (bIsEditorRecompile == true)
+				string PossibleAssemblyName = Desc.TargetName;
+				if (Desc.bIsEditorRecompile == true)
 				{
 					PossibleAssemblyName += "_EditorRecompile";
 				}
@@ -476,20 +468,14 @@ namespace UnrealBuildTool
 			}
 
 			// Try getting it from the RulesCompiler
-			UEBuildTarget Target = RulesCompiler.CreateTarget(
-				TargetName:TargetName, 
-				Target:new TargetInfo(Platform, Configuration),
-				InAdditionalDefinitions:AdditionalDefinitions, 
-				InRemoteRoot:RemoteRoot, 
-				InOnlyModules:OnlyModules, 
-				bInEditorRecompile:bIsEditorRecompile);
+			UEBuildTarget Target = RulesCompiler.CreateTarget(Desc);
 			if (Target == null)
 			{
 				if (UEBuildConfiguration.bCleanProject)
 				{
 					return null;
 				}
-				throw new BuildException( "Couldn't find target name {0}.", TargetName );
+				throw new BuildException( "Couldn't find target name {0}.", Desc.TargetName );
 			}
 			else
 			{
@@ -602,26 +588,28 @@ namespace UnrealBuildTool
 		}
 
 		/// <summary>
-		/// 
+		/// Finds the base name for all binaries produced for the given target
 		/// </summary>
-		/// <param name="InGameName"></param>
-		/// <param name="InRulesObject"></param>
-		/// <param name="InPlatform"></param>
-		/// <param name="InConfiguration"></param>
-		/// <returns></returns>
-		public static string GetBinaryBaseName(string InGameName, TargetRules InRulesObject, UnrealTargetPlatform InPlatform, UnrealTargetConfiguration InConfiguration, string InNameAppend)
+		/// <param name="InDesc">Target descriptor</param>
+		/// <param name="InRulesObject">Target rules object created by RulesCompiler</param>
+		/// <param name="InNameAppend">Suffix to append to the name</param>
+		/// <returns>Base name for the application</returns>
+		public static string GetDefaultAppName(TargetDescriptor InDesc, TargetRules InRulesObject, string InNameAppend)
 		{
 			//@todo. Allow failure to get build platform here?
 			bool bPlatformRequiresMonolithic = false;
-			var BuildPlatform = UEBuildPlatform.GetBuildPlatform(InPlatform, true);
+			var BuildPlatform = UEBuildPlatform.GetBuildPlatform(InDesc.Platform, true);
 			if (BuildPlatform != null)
 			{
-				bPlatformRequiresMonolithic = BuildPlatform.ShouldCompileMonolithicBinary(InPlatform);
+				bPlatformRequiresMonolithic = BuildPlatform.ShouldCompileMonolithicBinary(InDesc.Platform);
 			}
 
-			if (InRulesObject.ShouldCompileMonolithic(InPlatform, InConfiguration) || bPlatformRequiresMonolithic)
+			// If we're building a monolithic binary, then the game and engine code are linked together into one
+			// program executable, so we want the application name to be the game name.  In the case of a modular
+			// binary, we use 'UnrealEngine' for our application name
+			if (InRulesObject.ShouldCompileMonolithic(InDesc.Platform, InDesc.Configuration) || bPlatformRequiresMonolithic)
 			{
-				return InGameName;
+				return InDesc.TargetName;
 			}
 			else
 			{
@@ -643,11 +631,11 @@ namespace UnrealBuildTool
 		}
 		private readonly TargetRules.TargetType? TargetTypeOrNull;
 
-		/** The name of the application the target is part of. */
+		/** The name of the application the target is part of. For non-monolithic targets, this is typically the name of the base application, eg. UE4Editor for any game editor. */
 		public string AppName;
 
-		/** The name of the game the target is part of - can be empty */
-		public string GameName;
+		/** The name of the target */
+		public string TargetName;
 
 		/** Platform as defined by the VCProject and passed via the command line. Not the same as internal config names. */
 		public UnrealTargetPlatform Platform;
@@ -738,30 +726,21 @@ namespace UnrealBuildTool
 		 * @param InAdditionalDefinitions - Additional definitions provided on the UBT command-line for the target.
 		 * @param InRemoteRoot - The remote path that the build output is synced to.
 		 */
-		public UEBuildTarget(
-			string InAppName,
-			string InGameName,
-			UnrealTargetPlatform InPlatform,
-			UnrealTargetConfiguration InConfiguration,
-			TargetRules InRulesObject,
-			List<string> InAdditionalDefinitions,
-			string InRemoteRoot,
-			List<OnlyModule> InOnlyModules,
-			bool bInEditorRecompile)
+		public UEBuildTarget(string InAppName, TargetDescriptor InDesc, TargetRules InRulesObject)
 		{
 			AppName = InAppName;
-			GameName = InGameName;
-			Platform = InPlatform;
-			Configuration = InConfiguration;
+			TargetName = InDesc.TargetName;
+			Platform = InDesc.Platform;
+			Configuration = InDesc.Configuration;
 			Rules = InRulesObject;
-			bEditorRecompile = bInEditorRecompile;
+			bEditorRecompile = InDesc.bIsEditorRecompile;
 
 
 			{
-				bCompileMonolithic = (Rules != null) ? Rules.ShouldCompileMonolithic(InPlatform, InConfiguration) : false;
+				bCompileMonolithic = (Rules != null) ? Rules.ShouldCompileMonolithic(InDesc.Platform, InDesc.Configuration) : false;
 
 				// Platforms may *require* monolithic compilation...
-				bCompileMonolithic |= UEBuildPlatform.PlatformRequiresMonolithicBuilds(InPlatform, InConfiguration);
+				bCompileMonolithic |= UEBuildPlatform.PlatformRequiresMonolithicBuilds(InDesc.Platform, InDesc.Configuration);
 
 				// Force monolithic or modular mode if we were asked to
 				if( UnrealBuildTool.CommandLineContains("-Monolithic") ||
@@ -803,9 +782,9 @@ namespace UnrealBuildTool
 				EngineIntermediateDirectory = Path.GetFullPath(Path.Combine(BuildConfiguration.RelativeEnginePath, BuildConfiguration.PlatformIntermediateFolder, AppName, Configuration.ToString()));
 			}
 
-			RemoteRoot = InRemoteRoot;
+			RemoteRoot = InDesc.RemoteRoot;
 
-			OnlyModules = InOnlyModules;
+			OnlyModules = InDesc.OnlyModules;
 
 			TargetTypeOrNull = (Rules != null) ? Rules.Type : (TargetRules.TargetType?)null;
 
@@ -827,7 +806,7 @@ namespace UnrealBuildTool
 					{
 						for (int Index = 0; Index < OutputPaths.Length; Index++)
 						{
-							OutputPaths[Index] = OutputPaths[Index].Replace(Path.Combine("Engine", "Binaries"), Path.Combine(InGameName, "Binaries"));
+							OutputPaths[Index] = OutputPaths[Index].Replace(Path.Combine("Engine", "Binaries"), Path.Combine(InDesc.TargetName, "Binaries"));
 						}
 					}
 					else
@@ -836,7 +815,7 @@ namespace UnrealBuildTool
 						string UProjectPath = UnrealBuildTool.GetUProjectPath();
 						if (Path.IsPathRooted(UProjectPath) == false)
 						{
-							string FilePath = UProjectInfo.GetProjectForTarget(InGameName);
+							string FilePath = UProjectInfo.GetProjectForTarget(InDesc.TargetName);
 							string FullPath = Path.GetFullPath(FilePath);
 							UProjectPath = Path.GetDirectoryName(FullPath);
 						}
@@ -851,7 +830,7 @@ namespace UnrealBuildTool
 
 			// handle some special case defines (so build system can pass -DEFINE as normal instead of needing
 			// to know about special parameters)
-			foreach (string Define in InAdditionalDefinitions)
+			foreach (string Define in InDesc.AdditionalDefinitions)
 			{
 				switch (Define)
 				{
@@ -875,7 +854,7 @@ namespace UnrealBuildTool
 			}
 
 			// Add the definitions specified on the command-line.
-			GlobalCompileEnvironment.Config.Definitions.AddRange(InAdditionalDefinitions);
+			GlobalCompileEnvironment.Config.Definitions.AddRange(InDesc.AdditionalDefinitions);
 		}
 
 		/// <summary>
@@ -973,10 +952,9 @@ namespace UnrealBuildTool
 		protected void CleanTarget(List<UEBuildBinary> Binaries, CPPTargetPlatform Platform, BuildManifest Manifest)
 		{
 			{
-				var LocalTargetName = (TargetType == TargetRules.TargetType.Program) ? AppName : GameName;
-				Log.TraceVerbose("Cleaning target {0} - AppName {1}", LocalTargetName, AppName);
+				Log.TraceVerbose("Cleaning target {0} - AppName {1}", TargetName, AppName);
 
-				var TargetFilename = RulesCompiler.GetTargetFilename(GameName);
+				var TargetFilename = RulesCompiler.GetTargetFilename(TargetName);
 				Log.TraceVerbose("\tTargetFilename {0}", TargetFilename);
 
 				// Collect all files to delete.
@@ -1057,7 +1035,7 @@ namespace UnrealBuildTool
 
 				//
 				{
-					var AppEnginePath = Path.Combine(PlatformEngineBuildDataFolder, LocalTargetName, Configuration.ToString());
+					var AppEnginePath = Path.Combine(PlatformEngineBuildDataFolder, TargetName, Configuration.ToString());
 					if (Directory.Exists(AppEnginePath))
 					{
 						CleanDirectory(AppEnginePath);
@@ -1444,7 +1422,7 @@ namespace UnrealBuildTool
 			// On Mac and Linux we have actions that should be executed after all the binaries are created
 			if (GlobalLinkEnvironment.Config.Target.Platform == CPPTargetPlatform.Mac || GlobalLinkEnvironment.Config.Target.Platform == CPPTargetPlatform.Linux)
 			{
-				TargetToolChain.SetupBundleDependencies(AppBinaries, GameName);
+				TargetToolChain.SetupBundleDependencies(AppBinaries, TargetName);
 			}
 
 			// Generate the external file list 
@@ -1639,7 +1617,7 @@ namespace UnrealBuildTool
 			UEBuildPlatform.GetBuildPlatform(Platform).SetupBinaries(this);
 
 			// Describe what's being built.
-			Log.TraceVerbose("Building {0} - {1} - {2} - {3}", AppName, GameName, Platform, Configuration);
+			Log.TraceVerbose("Building {0} - {1} - {2} - {3}", AppName, TargetName, Platform, Configuration);
 
 			// Put the non-executable output files (PDB, import library, etc) in the intermediate directory.
 			GlobalLinkEnvironment.Config.IntermediateDirectory = Path.GetFullPath(GlobalCompileEnvironment.Config.OutputDirectory);
diff --git a/Engine/Source/Programs/UnrealBuildTool/IOS/IOSToolChain.cs b/Engine/Source/Programs/UnrealBuildTool/IOS/IOSToolChain.cs
--- a/Engine/Source/Programs/UnrealBuildTool/IOS/IOSToolChain.cs
+++ b/Engine/Source/Programs/UnrealBuildTool/IOS/IOSToolChain.cs
@@ -989,7 +989,7 @@ namespace UnrealBuildTool
 		{
 			base.PostBuildSync(Target);
 
-			string AppName = Target.TargetType == TargetRules.TargetType.Game ? Target.GameName : Target.AppName;
+			string AppName = Target.TargetType == TargetRules.TargetType.Game ? Target.TargetName : Target.AppName;
 
 			if (BuildHostPlatform.Current.Platform == UnrealTargetPlatform.Mac)
 			{
@@ -1170,7 +1170,7 @@ namespace UnrealBuildTool
 
 						// copy the executable
 						string RemoteShadowDirectoryMac = ConvertPath(Path.GetDirectoryName(Target.OutputPath));
-						string FinalRemoteExecutablePath = String.Format("{0}/Payload/{1}.app/{1}", RemoteShadowDirectoryMac, Target.GameName);
+						string FinalRemoteExecutablePath = String.Format("{0}/Payload/{1}.app/{1}", RemoteShadowDirectoryMac, Target.TargetName);
 						RPCUtilHelper.Command("/", String.Format("cp -f {0} {1}", RemoteExecutablePath, FinalRemoteExecutablePath), "", null);
 					}
 					else if (!Utils.IsRunningOnMono && BuildHostPlatform.Current.Platform != UnrealTargetPlatform.Mac)
diff --git a/Engine/Source/Programs/UnrealBuildTool/System/RulesCompiler.cs b/Engine/Source/Programs/UnrealBuildTool/System/RulesCompiler.cs
--- a/Engine/Source/Programs/UnrealBuildTool/System/RulesCompiler.cs
+++ b/Engine/Source/Programs/UnrealBuildTool/System/RulesCompiler.cs
@@ -1748,26 +1748,25 @@ namespace UnrealBuildTool
 		/// <param name="TargetName">Name of the target</param>
 		/// <param name="Target">Information about the target associated with this target</param>
 		/// <returns>The build target object for the specified build rules source file</returns>
-		public static UEBuildTarget CreateTarget(string TargetName, TargetInfo Target,
-			List<string> InAdditionalDefinitions, string InRemoteRoot, List<OnlyModule> InOnlyModules, bool bInEditorRecompile)
+		public static UEBuildTarget CreateTarget(TargetDescriptor Desc)
 		{
 			var CreateTargetStartTime = DateTime.UtcNow;
 
 			string TargetFileName;
-			TargetRules RulesObject = CreateTargetRules(TargetName, Target, bInEditorRecompile, out TargetFileName);
-			if (bInEditorRecompile)
+			TargetRules RulesObject = CreateTargetRules(Desc.TargetName, new TargetInfo(Desc.Platform, Desc.Configuration), Desc.bIsEditorRecompile, out TargetFileName);
+			if (Desc.bIsEditorRecompile)
 			{
 				// Now that we found the actual Editor target, make sure we're no longer using the old TargetName (which is the Game target)
 				var TargetSuffixIndex = RulesObject.TargetName.LastIndexOf("Target");
-				TargetName = (TargetSuffixIndex > 0) ? RulesObject.TargetName.Substring(0, TargetSuffixIndex) : RulesObject.TargetName;
+				Desc.TargetName = (TargetSuffixIndex > 0) ? RulesObject.TargetName.Substring(0, TargetSuffixIndex) : RulesObject.TargetName;
 			}
-			if ((ProjectFileGenerator.bGenerateProjectFiles == false) && (RulesObject.SupportsPlatform(Target.Platform) == false))
+			if ((ProjectFileGenerator.bGenerateProjectFiles == false) && (RulesObject.SupportsPlatform(Desc.Platform) == false))
 			{
 				if (UEBuildConfiguration.bCleanProject)
 				{
 					return null;
 				}
-				throw new BuildException("{0} does not support the {1} platform.", TargetName, Target.Platform.ToString());
+				throw new BuildException("{0} does not support the {1} platform.", Desc.TargetName, Desc.Platform.ToString());
 			}
 
 			// Generate a build target from this rules module
@@ -1775,82 +1774,31 @@ namespace UnrealBuildTool
 			switch (RulesObject.Type)
 			{
 				case TargetRules.TargetType.Game:
-					{
-						BuildTarget = new UEBuildGame(
-							InGameName:TargetName, 
-							InPlatform:Target.Platform, 
-							InConfiguration:Target.Configuration,
-							InRulesObject: RulesObject,
-							InAdditionalDefinitions: InAdditionalDefinitions, 
-							InRemoteRoot:InRemoteRoot, 
-							InOnlyModules:InOnlyModules,
-							bInEditorRecompile: bInEditorRecompile);
-					}
+					BuildTarget = new UEBuildGame(Desc, RulesObject);
 					break;
 				case TargetRules.TargetType.Editor:
-					{
-						BuildTarget = new UEBuildEditor(
-							InGameName: TargetName,
-							InPlatform: Target.Platform,
-							InConfiguration: Target.Configuration,
-							InRulesObject: RulesObject,
-							InAdditionalDefinitions: InAdditionalDefinitions,
-							InRemoteRoot: InRemoteRoot,
-							InOnlyModules: InOnlyModules,
-							bInEditorRecompile: bInEditorRecompile);
-					}
+					BuildTarget = new UEBuildEditor(Desc, RulesObject);
 					break;
                 case TargetRules.TargetType.Client:
-                    {
-                        BuildTarget = new UEBuildClient(
-                            InGameName: TargetName,
-                            InPlatform: Target.Platform,
-                            InConfiguration: Target.Configuration,
-                            InRulesObject: RulesObject,
-                            InAdditionalDefinitions: InAdditionalDefinitions,
-                            InRemoteRoot: InRemoteRoot,
-                            InOnlyModules: InOnlyModules,
-							bInEditorRecompile: bInEditorRecompile);
-                    }
+                    BuildTarget = new UEBuildClient(Desc, RulesObject);
                     break;
 				case TargetRules.TargetType.Server:
-					{
-						BuildTarget = new UEBuildServer(
-							InGameName: TargetName,
-							InPlatform: Target.Platform,
-							InConfiguration: Target.Configuration,
-							InRulesObject: RulesObject,
-							InAdditionalDefinitions: InAdditionalDefinitions,
-							InRemoteRoot: InRemoteRoot,
-							InOnlyModules: InOnlyModules,
-							bInEditorRecompile: bInEditorRecompile);
-					}
+					BuildTarget = new UEBuildServer(Desc, RulesObject);
 					break;
 				case TargetRules.TargetType.Program:
-					{
-						BuildTarget = new UEBuildTarget(
-							InAppName:TargetName,
-							InGameName:"", 
-							InPlatform:Target.Platform, 
-							InConfiguration:Target.Configuration,
-							InRulesObject:RulesObject, 
-							InAdditionalDefinitions:InAdditionalDefinitions, 
-							InRemoteRoot:InRemoteRoot, 
-							InOnlyModules:InOnlyModules,
-							bInEditorRecompile:bInEditorRecompile);
-					}
+					BuildTarget = new UEBuildTarget(Desc.TargetName, Desc, RulesObject);
 					break;
 			}
 
 			if( BuildConfiguration.bPrintPerformanceInfo )
 			{ 
 				var CreateTargetTime = (DateTime.UtcNow - CreateTargetStartTime).TotalSeconds;
-				Log.TraceInformation( "CreateTarget for " + TargetName + " took " + CreateTargetTime + "s" );
+				Log.TraceInformation( "CreateTarget for " + Desc.TargetName + " took " + CreateTargetTime + "s" );
 			}
 
 			if (BuildTarget == null)
 			{
-				throw new BuildException("Failed to create build target for '{0}'.", TargetName);
+				throw new BuildException("Failed to create build target for '{0}'.", Desc.TargetName);
 			}
 
 			return BuildTarget;
diff --git a/Engine/Source/Programs/UnrealBuildTool/System/UnrealBuildTool.cs b/Engine/Source/Programs/UnrealBuildTool/System/UnrealBuildTool.cs
--- a/Engine/Source/Programs/UnrealBuildTool/System/UnrealBuildTool.cs
+++ b/Engine/Source/Programs/UnrealBuildTool/System/UnrealBuildTool.cs
@@ -1191,7 +1191,7 @@ namespace UnrealBuildTool
                                         (CheckTarget.TargetType == TargetRules.TargetType.Server) ||
                                         (CheckTarget.TargetType == TargetRules.TargetType.Client))
                                     {
-                                        CheckTarget.AppName = CheckTarget.GameName;
+                                        CheckTarget.AppName = CheckTarget.TargetName;
                                     }
                                     DeployHandler.PrepTargetForDeployment(CheckTarget);
                                 }
@@ -1647,7 +1647,7 @@ namespace UnrealBuildTool
                         // Send out telemetry for this target
                         Telemetry.SendEvent("TargetBuildStats.2",
                             "AppName", Target.AppName,
-                            "GameName", Target.GameName,
+                            "GameName", Target.TargetName,
                             "Platform", Target.Platform.ToString(),
                             "Configuration", Target.Configuration.ToString(),
                             "CleanTarget", UEBuildConfiguration.bCleanProject.ToString(),
