Fix linker errors on non-monolithic shipping builds.

diff --git a/Engine/Source/Programs/UnrealBuildTool/Configuration/UEBuildTarget.cs b/Engine/Source/Programs/UnrealBuildTool/Configuration/UEBuildTarget.cs
--- a/Engine/Source/Programs/UnrealBuildTool/Configuration/UEBuildTarget.cs
+++ b/Engine/Source/Programs/UnrealBuildTool/Configuration/UEBuildTarget.cs
@@ -1529,6 +1529,8 @@ namespace UnrealBuildTool
 				CreateLinkerFixupsCPPFile();
 			}
 
+			GlobalLinkEnvironment.bShouldCompileMonolithic = ShouldCompileMonolithic();
+
 			// Build the target's binaries.
 			foreach (var Binary in AppBinaries)
 			{
diff --git a/Engine/Source/Programs/UnrealBuildTool/System/LinkEnvironment.cs b/Engine/Source/Programs/UnrealBuildTool/System/LinkEnvironment.cs
--- a/Engine/Source/Programs/UnrealBuildTool/System/LinkEnvironment.cs
+++ b/Engine/Source/Programs/UnrealBuildTool/System/LinkEnvironment.cs
@@ -157,6 +157,9 @@ namespace UnrealBuildTool
 	/** Encapsulates the environment that is used to link object files. */
 	public class LinkEnvironment
 	{
+		/** Whether we're linking in monolithic mode. Determines if linking should produce import library file. Relevant only for VC++, clang stores imports in shared library. */
+		public bool bShouldCompileMonolithic = false;
+
 		/** A list of the object files to be linked. */
 		public List<FileItem> InputFiles = new List<FileItem>();
 
diff --git a/Engine/Source/Programs/UnrealBuildTool/WinRT/WinRTToolChain.cs b/Engine/Source/Programs/UnrealBuildTool/WinRT/WinRTToolChain.cs
--- a/Engine/Source/Programs/UnrealBuildTool/WinRT/WinRTToolChain.cs
+++ b/Engine/Source/Programs/UnrealBuildTool/WinRT/WinRTToolChain.cs
@@ -796,9 +796,10 @@ namespace UnrealBuildTool
 
 			if (!bIsBuildingLibrary)
 			{
-				// Xbox 360 LTCG does not seem to produce those.
-				if (LinkEnvironment.Config.bHasExports &&
-					(LinkEnvironment.Config.Target.Configuration != CPPTargetConfiguration.Shipping) && (LinkEnvironment.Config.Target.Platform != CPPTargetPlatform.XboxOne))
+				// There is anything to export
+				if (LinkEnvironment.Config.bHasExports
+					// Shipping monolithic builds don't need exports
+					&& (!((LinkEnvironment.Config.Target.Configuration == CPPTargetConfiguration.Shipping) && (LinkEnvironment.bShouldCompileMonolithic != false))))
 				{
 					// Write the import library to the output directory for nFringe support.
 					FileItem ImportLibraryFile = FileItem.GetItemByPath(ImportLibraryFilePath);
diff --git a/Engine/Source/Programs/UnrealBuildTool/Windows/VCToolChain.cs b/Engine/Source/Programs/UnrealBuildTool/Windows/VCToolChain.cs
--- a/Engine/Source/Programs/UnrealBuildTool/Windows/VCToolChain.cs
+++ b/Engine/Source/Programs/UnrealBuildTool/Windows/VCToolChain.cs
@@ -1511,9 +1511,10 @@ namespace UnrealBuildTool
 
 			if (!bIsBuildingLibrary)
 			{
-				// Xbox 360 LTCG does not seem to produce those.
-				if (LinkEnvironment.Config.bHasExports &&
-					(LinkEnvironment.Config.Target.Configuration != CPPTargetConfiguration.Shipping) && (LinkEnvironment.Config.Target.Platform != CPPTargetPlatform.XboxOne))
+				// There is anything to export
+				if (LinkEnvironment.Config.bHasExports
+					// Shipping monolithic builds don't need exports
+					&& (!((LinkEnvironment.Config.Target.Configuration == CPPTargetConfiguration.Shipping) && (LinkEnvironment.bShouldCompileMonolithic != false))))
 				{
 					// Write the import library to the output directory for nFringe support.
 					FileItem ImportLibraryFile = FileItem.GetItemByPath(ImportLibraryFilePath);
