Remove unnecessary ref keywords.

diff --git a/Engine/Source/Programs/UnrealBuildTool/Windows/VCToolChain.cs b/Engine/Source/Programs/UnrealBuildTool/Windows/VCToolChain.cs
--- a/Engine/Source/Programs/UnrealBuildTool/Windows/VCToolChain.cs
+++ b/Engine/Source/Programs/UnrealBuildTool/Windows/VCToolChain.cs
@@ -22,7 +22,7 @@ namespace UnrealBuildTool
 		}
 
 		
-		static void AddDefinition( ref StringBuilder String, string Variable, string Value = null )
+		static void AddDefinition( StringBuilder String, string Variable, string Value = null )
 		{
 			// If the value has a space in it and isn't wrapped in quotes, do that now
 			if( Value != null && !Value.StartsWith( "\"" ) && ( Value.Contains( " " ) || Value.Contains( "$" ) ) )
@@ -55,7 +55,7 @@ namespace UnrealBuildTool
 		}
 
 
-		static void AddIncludePath( ref StringBuilder String, string IncludePath )
+		static void AddIncludePath( StringBuilder String, string IncludePath )
 		{
 			// If the value has a space in it and isn't wrapped in quotes, do that now
 			if( !IncludePath.StartsWith( "\"" ) && ( IncludePath.Contains( " " ) || IncludePath.Contains( "$" ) ) )
@@ -117,7 +117,7 @@ namespace UnrealBuildTool
 				}
 
 				Arguments.Append(" -fms-compatibility-version=" + VersionString);
-				AddDefinition( ref Arguments, "_MSC_FULL_VER", FullVersionString + "00000" );
+				AddDefinition( Arguments, "_MSC_FULL_VER", FullVersionString + "00000" );
 			}
 
 			// @todo clang: Clang on Windows doesn't respect "#pragma warning (error: ####)", and we're not passing "/WX", so warnings are not
@@ -202,14 +202,14 @@ namespace UnrealBuildTool
 			// If compiling as a DLL, set the relevant defines
 			if (CompileEnvironment.Config.bIsBuildingDLL)
 			{
-				AddDefinition( ref Arguments, "_WINDLL" );
+				AddDefinition( Arguments, "_WINDLL" );
 			}
 
 			// When targeting Windows XP with Visual Studio 2012+, we need to tell the compiler to use the older Windows SDK that works
 			// with Windows XP (http://blogs.msdn.com/b/vcblog/archive/2012/10/08/10357555.aspx)
 			if (WindowsPlatform.IsWindowsXPSupported())
 			{
-				AddDefinition( ref Arguments, "_USING_V110_SDK71_");
+				AddDefinition( Arguments, "_USING_V110_SDK71_");
 			}
 
 
@@ -403,8 +403,8 @@ namespace UnrealBuildTool
 					}
 					else
 					{
-						AddDefinition( ref Arguments, "_MT" );
-						AddDefinition( ref Arguments, "_DEBUG" );
+						AddDefinition( Arguments, "_MT" );
+						AddDefinition( Arguments, "_DEBUG" );
 					}
 				}
 				else
@@ -415,7 +415,7 @@ namespace UnrealBuildTool
 					}
 					else
 					{
-						AddDefinition( ref Arguments, "_MT" );
+						AddDefinition( Arguments, "_MT" );
 					}
 				}
 			}
@@ -429,9 +429,9 @@ namespace UnrealBuildTool
 					}
 					else
 					{
-						AddDefinition( ref Arguments, "_MT" );
-						AddDefinition( ref Arguments, "_DEBUG" );
-						AddDefinition( ref Arguments, "_DLL" );
+						AddDefinition( Arguments, "_MT" );
+						AddDefinition( Arguments, "_DEBUG" );
+						AddDefinition( Arguments, "_DLL" );
 					}
 				}
 				else
@@ -442,8 +442,8 @@ namespace UnrealBuildTool
 					}
 					else
 					{
-						AddDefinition( ref Arguments, "_MT" );
-						AddDefinition( ref Arguments, "_DLL" );
+						AddDefinition( Arguments, "_MT" );
+						AddDefinition( Arguments, "_DLL" );
 					}
 				}
 			}
@@ -512,7 +512,7 @@ namespace UnrealBuildTool
 			//
 			if (BuildConfiguration.bRunUnrealCodeAnalyzer)
 			{
-				AddDefinition( ref Arguments, "UNREAL_CODE_ANALYZER");
+				AddDefinition( Arguments, "UNREAL_CODE_ANALYZER");
 			}
 		}
 
@@ -915,14 +915,14 @@ namespace UnrealBuildTool
 			{
 				foreach (string IncludePath in CompileEnvironment.Config.CPPIncludeInfo.IncludePaths)
 				{
-					AddIncludePath( ref SharedArguments, IncludePath );
+					AddIncludePath( SharedArguments, IncludePath );
 				}
 			}
 			else
 			{
 				foreach (string IncludePath in CompileEnvironment.Config.CPPIncludeInfo.IncludePaths)
 				{
-					AddIncludePath( ref SharedArguments, System.IO.Path.GetFullPath(IncludePath) );
+					AddIncludePath( SharedArguments, System.IO.Path.GetFullPath(IncludePath) );
 				}
 			}
 			foreach (string IncludePath in CompileEnvironment.Config.CPPIncludeInfo.SystemIncludePaths)
@@ -946,7 +946,7 @@ namespace UnrealBuildTool
 				}
 				else
 				{
-					AddIncludePath( ref SharedArguments, IncludePath );
+					AddIncludePath( SharedArguments, IncludePath );
 				}
 			}
 
@@ -980,7 +980,7 @@ namespace UnrealBuildTool
 			{
 				// Escape all quotation marks so that they get properly passed with the command line.
 				var DefinitionArgument = Definition.Contains("\"") ? Definition.Replace("\"", "\\\"") : Definition;
-				AddDefinition( ref SharedArguments, DefinitionArgument);
+				AddDefinition( SharedArguments, DefinitionArgument);
 			}
 
 			var BuildPlatform = UEBuildPlatform.GetBuildPlatformForCPPTargetPlatform(CompileEnvironment.Config.Target.Platform);
@@ -1027,7 +1027,7 @@ namespace UnrealBuildTool
 					// Make sure the original source directory the PCH header file existed in is added as an include
 					// path -- it might be a private PCH header and we need to make sure that its found!
 					string OriginalPCHHeaderDirectory = Path.GetDirectoryName( SourceFile.AbsolutePath );
-					AddIncludePath( ref FileArguments, OriginalPCHHeaderDirectory );
+					AddIncludePath( FileArguments, OriginalPCHHeaderDirectory );
 
 					var PrecompiledFileExtension = UEBuildPlatform.BuildPlatformDictionary[UnrealTargetPlatform.Win64].GetBinaryExtension(UEBuildBinaryType.PrecompiledHeader);
 					// Add the precompiled header file to the produced items list.
@@ -1309,14 +1309,14 @@ namespace UnrealBuildTool
 				// compiler so that we can switch on that in the .rc file using #ifdef.
 				if (Environment.Config.Target.Platform == CPPTargetPlatform.Win64)
 				{
-					AddDefinition( ref Arguments, "_WIN64" );
+					AddDefinition( Arguments, "_WIN64" );
 				}
 
 				// When targeting Windows XP with Visual Studio 2012+, we need to tell the compiler to use the older Windows SDK that works
 				// with Windows XP (http://blogs.msdn.com/b/vcblog/archive/2012/10/08/10357555.aspx)
 				if (WindowsPlatform.IsWindowsXPSupported())
 				{
-					AddDefinition( ref Arguments, "_USING_V110_SDK71_" );
+					AddDefinition( Arguments, "_USING_V110_SDK71_" );
 				}
 
 				// Language
@@ -1325,19 +1325,19 @@ namespace UnrealBuildTool
 				// Include paths.
 				foreach (string IncludePath in Environment.Config.CPPIncludeInfo.IncludePaths)
 				{
-					AddIncludePath( ref Arguments, IncludePath );
+					AddIncludePath( Arguments, IncludePath );
 				}
 
 				// System include paths.
 				foreach( var SystemIncludePath in Environment.Config.CPPIncludeInfo.SystemIncludePaths )
 				{
-					AddIncludePath( ref Arguments, SystemIncludePath );
+					AddIncludePath( Arguments, SystemIncludePath );
 				}
 
 				// Preprocessor definitions.
 				foreach (string Definition in Environment.Config.Definitions)
 				{
-					AddDefinition( ref Arguments, Definition);
+					AddDefinition( Arguments, Definition);
 				}
 
 				// Add the RES file to the produced item list.
