Remove some more unnecessary ref keywords.

diff --git a/Engine/Source/Programs/UnrealBuildTool/System/ProjectFileGenerator.cs b/Engine/Source/Programs/UnrealBuildTool/System/ProjectFileGenerator.cs
--- a/Engine/Source/Programs/UnrealBuildTool/System/ProjectFileGenerator.cs
+++ b/Engine/Source/Programs/UnrealBuildTool/System/ProjectFileGenerator.cs
@@ -2203,7 +2203,7 @@ namespace UnrealBuildTool
 			else
 			{
 				// For existing project files, just support the default desktop platforms and configurations
-				UnrealBuildTool.GetAllDesktopPlatforms(ref ProjectTarget.ExtraSupportedPlatforms, false);
+				UnrealBuildTool.GetAllDesktopPlatforms(ProjectTarget.ExtraSupportedPlatforms, false);
 				// Debug and Development only
 				ProjectTarget.ExtraSupportedConfigurations.Add(UnrealTargetConfiguration.Debug);
 				ProjectTarget.ExtraSupportedConfigurations.Add(UnrealTargetConfiguration.Development);
diff --git a/Engine/Source/Programs/UnrealBuildTool/System/RulesCompiler.cs b/Engine/Source/Programs/UnrealBuildTool/System/RulesCompiler.cs
--- a/Engine/Source/Programs/UnrealBuildTool/System/RulesCompiler.cs
+++ b/Engine/Source/Programs/UnrealBuildTool/System/RulesCompiler.cs
@@ -504,21 +504,21 @@ namespace UnrealBuildTool
 		/// </summary>
 		/// <param name="OutPlatforms">The list of platforms supported</param>
 		/// <returns>true if successful, false if not</returns>
-		public virtual bool GetSupportedPlatforms(ref List<UnrealTargetPlatform> OutPlatforms)
+		public virtual bool GetSupportedPlatforms(List<UnrealTargetPlatform> OutPlatforms)
 		{
 			if(Type == TargetType.Program)
 			{
 				// By default, all programs are desktop only.
-				return UnrealBuildTool.GetAllDesktopPlatforms(ref OutPlatforms, false);
+				return UnrealBuildTool.GetAllDesktopPlatforms(OutPlatforms, false);
 			}
 			else if(IsEditorType(Type))
 			{
-				return UnrealBuildTool.GetAllEditorPlatforms(ref OutPlatforms, false);
+				return UnrealBuildTool.GetAllEditorPlatforms(OutPlatforms, false);
 			}
 			else if (TargetRules.IsGameType(Type))
 			{
 				// By default all games support all platforms
-				return UnrealBuildTool.GetAllPlatforms(ref OutPlatforms, false);
+				return UnrealBuildTool.GetAllPlatforms(OutPlatforms, false);
 			}
 
 			return false;
@@ -527,7 +527,7 @@ namespace UnrealBuildTool
 		public bool SupportsPlatform(UnrealTargetPlatform InPlatform)
 		{
 			List<UnrealTargetPlatform> SupportedPlatforms = new List<UnrealTargetPlatform>();
-			if (GetSupportedPlatforms(ref SupportedPlatforms) == true)
+			if (GetSupportedPlatforms(SupportedPlatforms) == true)
 			{
 				return SupportedPlatforms.Contains(InPlatform);
 			}
diff --git a/Engine/Source/Programs/UnrealBuildTool/System/UnrealBuildTool.cs b/Engine/Source/Programs/UnrealBuildTool/System/UnrealBuildTool.cs
--- a/Engine/Source/Programs/UnrealBuildTool/System/UnrealBuildTool.cs
+++ b/Engine/Source/Programs/UnrealBuildTool/System/UnrealBuildTool.cs
@@ -281,7 +281,7 @@ namespace UnrealBuildTool
         /// <param name="Predicate">The predicate to filter the platforms.</param>
         /// <param name="bCheckValidity">If true, ensure it is a valid platform.</param>
         /// <returns>True if successful and platforms are found, false if not.</returns>
-        static private bool GetPlatforms(ref List<UnrealTargetPlatform> OutPlatforms, Func<UnrealTargetPlatform, bool> Predicate, bool bCheckValidity = true)
+        static private bool GetPlatforms(List<UnrealTargetPlatform> OutPlatforms, Func<UnrealTargetPlatform, bool> Predicate, bool bCheckValidity = true)
         {
             OutPlatforms.Clear();
             foreach (UnrealTargetPlatform Platform in Enum.GetValues(typeof(UnrealTargetPlatform)))
@@ -305,9 +305,9 @@ namespace UnrealBuildTool
         /// <param name="OutPlatforms">The list of platform to fill in.</param>
         /// <param name="bCheckValidity">If true, ensure it is a valid platform.</param>
         /// <returns>True if successful and platforms are found, false if not.</returns>
-        static public bool GetAllPlatforms(ref List<UnrealTargetPlatform> OutPlatforms, bool bCheckValidity = true)
+        static public bool GetAllPlatforms(List<UnrealTargetPlatform> OutPlatforms, bool bCheckValidity = true)
         {
-            return GetPlatforms(ref OutPlatforms, (Platform) => { return true; }, bCheckValidity);
+            return GetPlatforms(OutPlatforms, (Platform) => { return true; }, bCheckValidity);
         }
 
         /// <summary>
@@ -316,9 +316,9 @@ namespace UnrealBuildTool
         /// <param name="OutPlatforms">The list of platform to fill in.</param>
         /// <param name="bCheckValidity">If true, ensure it is a valid platform.</param>
         /// <returns>True if successful and platforms are found, false if not.</returns>
-        static public bool GetAllDesktopPlatforms(ref List<UnrealTargetPlatform> OutPlatforms, bool bCheckValidity = true)
+        static public bool GetAllDesktopPlatforms(List<UnrealTargetPlatform> OutPlatforms, bool bCheckValidity = true)
         {
-            return GetPlatforms(ref OutPlatforms, UnrealBuildTool.IsDesktopPlatform, bCheckValidity);
+            return GetPlatforms(OutPlatforms, UnrealBuildTool.IsDesktopPlatform, bCheckValidity);
         }
 
         /// <summary>
@@ -327,9 +327,9 @@ namespace UnrealBuildTool
         /// <param name="OutPlatforms">The list of platform to fill in.</param>
         /// <param name="bCheckValidity">If true, ensure it is a valid platform.</param>
         /// <returns>True if successful and platforms are found, false if not.</returns>
-        static public bool GetAllEditorPlatforms(ref List<UnrealTargetPlatform> OutPlatforms, bool bCheckValidity = true)
+        static public bool GetAllEditorPlatforms(List<UnrealTargetPlatform> OutPlatforms, bool bCheckValidity = true)
         {
-            return GetPlatforms(ref OutPlatforms, UnrealBuildTool.IsEditorPlatform, bCheckValidity);
+            return GetPlatforms(OutPlatforms, UnrealBuildTool.IsEditorPlatform, bCheckValidity);
         }
 
         /// <summary>
@@ -338,9 +338,9 @@ namespace UnrealBuildTool
         /// <param name="OutPlatforms">The list of platform to fill in.</param>
         /// <param name="bCheckValidity">If true, ensure it is a valid platform.</param>
         /// <returns>True if successful and platforms are found, false if not.</returns>
-        static public bool GetAllServerPlatforms(ref List<UnrealTargetPlatform> OutPlatforms, bool bCheckValidity = true)
+        static public bool GetAllServerPlatforms(List<UnrealTargetPlatform> OutPlatforms, bool bCheckValidity = true)
         {
-            return GetPlatforms(ref OutPlatforms, UnrealBuildTool.IsServerPlatform, bCheckValidity);
+            return GetPlatforms(OutPlatforms, UnrealBuildTool.IsServerPlatform, bCheckValidity);
         }
 
         /// <summary>
diff --git a/Engine/Source/Programs/UnrealBuildTool/System/VCProject.cs b/Engine/Source/Programs/UnrealBuildTool/System/VCProject.cs
--- a/Engine/Source/Programs/UnrealBuildTool/System/VCProject.cs
+++ b/Engine/Source/Programs/UnrealBuildTool/System/VCProject.cs
@@ -195,12 +195,12 @@ namespace UnrealBuildTool
 			{
 				if( ProjectTarget.TargetRules != null )
 				{
-					ProjectTarget.TargetRules.GetSupportedPlatforms(ref SupportedPlatforms);
+					ProjectTarget.TargetRules.GetSupportedPlatforms(SupportedPlatforms);
 				}
 			}
 			else
 			{
-				UnrealBuildTool.GetAllPlatforms(ref SupportedPlatforms);
+				UnrealBuildTool.GetAllPlatforms(SupportedPlatforms);
 			}
 			bool bIncludeTestAndShippingConfigs = ProjectFileGenerator.bIncludeTestAndShippingConfigs || ProjectFileGenerator.bGeneratingRocketProjectFiles;
 			if( ProjectTarget.TargetRules != null )
diff --git a/Engine/Source/Programs/UnrealBuildTool/System/VCProjectFileGenerator.cs b/Engine/Source/Programs/UnrealBuildTool/System/VCProjectFileGenerator.cs
--- a/Engine/Source/Programs/UnrealBuildTool/System/VCProjectFileGenerator.cs
+++ b/Engine/Source/Programs/UnrealBuildTool/System/VCProjectFileGenerator.cs
@@ -616,7 +616,7 @@ namespace UnrealBuildTool
 											if( CurProject.ProjectTargets[0].TargetRules != null )
 											{
 												var ProjectSupportedPlatforms = new List<UnrealTargetPlatform>();
-												CurProject.ProjectTargets[0].TargetRules.GetSupportedPlatforms( ref ProjectSupportedPlatforms );
+												CurProject.ProjectTargets[0].TargetRules.GetSupportedPlatforms( ProjectSupportedPlatforms );
 												if( !ProjectSupportedPlatforms.Contains( SolutionPlatform ) )
 												{
 													SolutionPlatform = ProjectSupportedPlatforms[ 0 ];
diff --git a/Engine/Source/Programs/UnrealBuildTool/System/XcodeProjectFileGenerator.cs b/Engine/Source/Programs/UnrealBuildTool/System/XcodeProjectFileGenerator.cs
--- a/Engine/Source/Programs/UnrealBuildTool/System/XcodeProjectFileGenerator.cs
+++ b/Engine/Source/Programs/UnrealBuildTool/System/XcodeProjectFileGenerator.cs
@@ -1108,7 +1108,7 @@ namespace UnrealBuildTool
 					var Target = new TargetInfo(UnrealTargetPlatform.Mac, UnrealTargetConfiguration.Development);
 					var TargetRulesObject = RulesCompiler.CreateTargetRules(TargetName, Target, false, out TargetFilePath);
 					List<UnrealTargetPlatform> SupportedPlatforms = new List<UnrealTargetPlatform>();
-					TargetRulesObject.GetSupportedPlatforms(ref SupportedPlatforms);
+					TargetRulesObject.GetSupportedPlatforms(SupportedPlatforms);
 					LinkEnvironmentConfiguration LinkConfiguration = new LinkEnvironmentConfiguration();
 					CPPEnvironmentConfiguration CPPConfiguration = new CPPEnvironmentConfiguration();
 					TargetRulesObject.SetupGlobalEnvironment(Target, ref LinkConfiguration, ref CPPConfiguration);
