Added support for ARMv8a and fallback for ARMv7a
#ue4
#android

diff --git a/Engine/Source/Programs/AutomationTool/Android/AndroidPlatform.Automation.cs b/Engine/Source/Programs/AutomationTool/Android/AndroidPlatform.Automation.cs
--- a/Engine/Source/Programs/AutomationTool/Android/AndroidPlatform.Automation.cs
+++ b/Engine/Source/Programs/AutomationTool/Android/AndroidPlatform.Automation.cs
@@ -769,6 +769,11 @@ public class AndroidPlatform : Platform
 			{
 				DeviceArch = "-armv7";
 			}
+            else
+            {
+                // future-proof by dropping back to armv7 for unknown
+                DeviceArch = "-armv7";
+            }
 		}
 
 		// if after the fallbacks, we still don't have it, we can't continue
diff --git a/Engine/Source/Programs/UnrealBuildTool/Android/UEDeployAndroid.cs b/Engine/Source/Programs/UnrealBuildTool/Android/UEDeployAndroid.cs
--- a/Engine/Source/Programs/UnrealBuildTool/Android/UEDeployAndroid.cs
+++ b/Engine/Source/Programs/UnrealBuildTool/Android/UEDeployAndroid.cs
@@ -278,7 +278,7 @@ namespace UnrealBuildTool.Android
 			switch (UE4Arch)
 			{
 				case "-armv7":	return "armeabi-v7a";
-				case "-arm64":	return "arm64";
+                case "-arm64":  return "arm64-v8a";
 				case "-x64":	return "x86_64";
 				case "-x86":	return "x86";
 
@@ -291,12 +291,15 @@ namespace UnrealBuildTool.Android
 			switch (NDKArch)
 			{
 				case "armeabi-v7a": return "-armv7";
-				case "x86":			return "-x86";
-				case "arm64":		return "-arm64";
+                case "arm64-v8a":   return "-arm64";
+                case "x86":         return "-x86";
+                case "arm64":       return "-arm64";
 				case "x86_64":
 				case "x64":			return "-x64";
 					
-				default: throw new BuildException("Unknown NDK architecture '{0}'", NDKArch);
+//				default: throw new BuildException("Unknown NDK architecture '{0}'", NDKArch);
+                // future-proof by returning armv7 for unknown
+                default:            return "-armv7";
 			}
 		}
 
