Conditionally adding the IAP billing permission to the android manifest.

diff --git a/Engine/Source/Programs/UnrealBuildTool/Android/UEDeployAndroid.cs b/Engine/Source/Programs/UnrealBuildTool/Android/UEDeployAndroid.cs
--- a/Engine/Source/Programs/UnrealBuildTool/Android/UEDeployAndroid.cs
+++ b/Engine/Source/Programs/UnrealBuildTool/Android/UEDeployAndroid.cs
@@ -645,6 +645,8 @@ namespace UnrealBuildTool.Android
 			Ini.GetArray("/Script/AndroidRuntimeSettings.AndroidRuntimeSettings", "ExtraPermissions", out ExtraPermissions);
 			bool bPackageForGearVR;
 			Ini.GetBool("/Script/AndroidRuntimeSettings.AndroidRuntimeSettings", "bPackageForGearVR", out bPackageForGearVR);
+			bool bEnableIAP = false;
+			Ini.GetBool("OnlineSubsystemGooglePlay.Store", "bSupportsInAppPurchasing", out bEnableIAP);
 
 			// replace some variables
 			PackageName = PackageName.Replace("[PROJECT]", ProjectName);
@@ -741,8 +743,12 @@ namespace UnrealBuildTool.Android
 			Text.AppendLine("\t<uses-permission android:name=\"android.permission.MODIFY_AUDIO_SETTINGS\"/>");
 			Text.AppendLine("\t<uses-permission android:name=\"android.permission.GET_ACCOUNTS\"/>");
             Text.AppendLine("\t<uses-permission android:name=\"android.permission.VIBRATE\"/>");
-			Text.AppendLine("\t<uses-permission android:name=\"com.android.vending.BILLING\"/>");
 //			Text.AppendLine("\t<uses-permission android:name=\"android.permission.DISABLE_KEYGUARD\"/>");
+
+			if (bEnableIAP)
+			{
+				Text.AppendLine("\t<uses-permission android:name=\"com.android.vending.BILLING\"/>");
+			}
 			if (bPackageForGearVR)
 			{
 				Text.AppendLine("\t<uses-permission android:name=\"android.permission.CAMERA\"/>");
