Fix handling of compiler definitions which contain spaces; AddDefinition() only looks at value, which is always null when adding define from the compile environment.

diff --git a/Engine/Source/Programs/UnrealBuildTool/Windows/VCToolChain.cs b/Engine/Source/Programs/UnrealBuildTool/Windows/VCToolChain.cs
--- a/Engine/Source/Programs/UnrealBuildTool/Windows/VCToolChain.cs
+++ b/Engine/Source/Programs/UnrealBuildTool/Windows/VCToolChain.cs
@@ -21,8 +21,23 @@ namespace UnrealBuildTool
 			UEToolChain.RegisterPlatformToolChain(CPPTargetPlatform.Win32, this);
 		}
 
-		
-		static void AddDefinition( StringBuilder String, string Variable, string Value = null )
+
+		static void AddDefinition( StringBuilder String, string Definition )
+		{
+			// Split the definition into name and value
+			int ValueIdx = Definition.IndexOf('=');
+			if(ValueIdx == -1)
+			{
+				AddDefinition(String, Definition, null);
+			}
+			else
+			{
+				AddDefinition(String, Definition.Substring(0, ValueIdx), Definition.Substring(ValueIdx + 1));
+			}
+		}
+
+
+		static void AddDefinition( StringBuilder String, string Variable, string Value )
 		{
 			// If the value has a space in it and isn't wrapped in quotes, do that now
 			if( Value != null && !Value.StartsWith( "\"" ) && ( Value.Contains( " " ) || Value.Contains( "$" ) ) )
@@ -980,7 +995,7 @@ namespace UnrealBuildTool
 			{
 				// Escape all quotation marks so that they get properly passed with the command line.
 				var DefinitionArgument = Definition.Contains("\"") ? Definition.Replace("\"", "\\\"") : Definition;
-				AddDefinition( SharedArguments, DefinitionArgument);
+				AddDefinition(SharedArguments, DefinitionArgument);
 			}
 
 			var BuildPlatform = UEBuildPlatform.GetBuildPlatformForCPPTargetPlatform(CompileEnvironment.Config.Target.Platform);
