UBT: Safety path for cases when ActionHistory.bin is corrupted or can't be read for any reason. If those situations happen UBT will create a fresh history and run normally instead of crashing.

diff --git a/Engine/Source/Programs/UnrealBuildTool/System/ActionHistory.cs b/Engine/Source/Programs/UnrealBuildTool/System/ActionHistory.cs
--- a/Engine/Source/Programs/UnrealBuildTool/System/ActionHistory.cs
+++ b/Engine/Source/Programs/UnrealBuildTool/System/ActionHistory.cs
@@ -25,16 +25,29 @@ namespace UnrealBuildTool
 		public ActionHistory(string InFilePath)
 		{
 			FilePath = Path.GetFullPath(InFilePath);
+
+			var bFoundCache = false;
 			if (File.Exists(FilePath) == true)
 			{
-				// Deserialize the history from disk if the file exists.
-				using (FileStream Stream = new FileStream(FilePath, FileMode.Open, FileAccess.Read))
+				try
+				{
+					// Deserialize the history from disk if the file exists.
+					using (FileStream Stream = new FileStream(FilePath, FileMode.Open, FileAccess.Read))
+					{
+						BinaryFormatter Formatter = new BinaryFormatter();
+						ProducedItemToPreviousActionCommandLine = Formatter.Deserialize(Stream) as Dictionary<string, string>;
+					}
+
+					bFoundCache = true;
+				}
+				catch (Exception)
 				{
-					BinaryFormatter Formatter = new BinaryFormatter();
-					ProducedItemToPreviousActionCommandLine = Formatter.Deserialize(Stream) as Dictionary<string, string>;
+					// If this fails for any reason just reset. History will be created later.
+					ProducedItemToPreviousActionCommandLine = null;
 				}
 			}
-			else
+
+			if (!bFoundCache)
 			{
 				// Otherwise create a fresh history.
 				ProducedItemToPreviousActionCommandLine = new Dictionary<string, string>();
