Fix for UBT build failure reported on AnswerHUB (File.ReadLines splits a line when it shouldn't be)

diff --git a/Engine/Source/Programs/UnrealBuildTool/System/Utils.cs b/Engine/Source/Programs/UnrealBuildTool/System/Utils.cs
--- a/Engine/Source/Programs/UnrealBuildTool/System/Utils.cs
+++ b/Engine/Source/Programs/UnrealBuildTool/System/Utils.cs
@@ -227,13 +227,17 @@ namespace UnrealBuildTool
 				{
 					// Parse the environment variable name and value from the string ("name=value")
 					int EqualSignPos = EnvString.IndexOf( '=' );
-					var EnvironmentVariableName = EnvString.Substring( 0, EqualSignPos );
-					var EnvironmentVariableValue = EnvString.Substring( EqualSignPos + 1 );
+					// Guard against strange cases where File.ReadAllLines splits a line in half
+					if (EqualSignPos >= 0)
+					{
+						var EnvironmentVariableName = EnvString.Substring(0, EqualSignPos);
+						var EnvironmentVariableValue = EnvString.Substring(EqualSignPos + 1);
 
-					Log.TraceVerbose("Setting environment variable: {0}={1}", EnvironmentVariableName, EnvironmentVariableValue);
+						Log.TraceVerbose("Setting environment variable: {0}={1}", EnvironmentVariableName, EnvironmentVariableValue);
 
-					// Set the environment variable
-					Environment.SetEnvironmentVariable( EnvironmentVariableName, EnvironmentVariableValue );
+						// Set the environment variable
+						Environment.SetEnvironmentVariable(EnvironmentVariableName, EnvironmentVariableValue);
+					}
 				}
 
 				// Clean up the temporary files we created earlier on, so the temp directory doesn't fill up
