Remove BuildingRocket() special case for output paths. I don't think we build anything that would trigger this code-path in Rocket anyway, but even then - it seems like the function just ends up being a no-op anyway.

diff --git a/Engine/Source/Programs/UnrealBuildTool/Configuration/UEBuildModule.cs b/Engine/Source/Programs/UnrealBuildTool/Configuration/UEBuildModule.cs
--- a/Engine/Source/Programs/UnrealBuildTool/Configuration/UEBuildModule.cs
+++ b/Engine/Source/Programs/UnrealBuildTool/Configuration/UEBuildModule.cs
@@ -507,7 +507,7 @@ namespace UnrealBuildTool
 		public virtual string FixupOutputPath(string InOutputPath)
 		{
 			string ModuleOutputPath = InOutputPath;
-			if ((OutputDirectory.Length != 0) && (UnrealBuildTool.BuildingRocket() == false))
+			if (OutputDirectory.Length != 0)
 			{
 				// Use from 'Binaries/' on as it is setup w/ the correct platform already
 				int BinariesIndex = ModuleOutputPath.LastIndexOf("Binaries");
