Remove BuildingRocket() special case. UE4Editor does not need to override global environment for Rocket any more.

diff --git a/Engine/Source/Programs/UnrealBuildTool/Configuration/UEBuildTarget.cs b/Engine/Source/Programs/UnrealBuildTool/Configuration/UEBuildTarget.cs
--- a/Engine/Source/Programs/UnrealBuildTool/Configuration/UEBuildTarget.cs
+++ b/Engine/Source/Programs/UnrealBuildTool/Configuration/UEBuildTarget.cs
@@ -2554,7 +2554,6 @@ namespace UnrealBuildTool
 				// 
 				bool bAllowGameOverride = !TargetRules.IsGameType(TargetType);
 				if (bAllowGameOverride ||
-					(Rules.ToString() == "UE4EditorTarget" && UnrealBuildTool.BuildingRocket()) ||	// @todo Rocket: Hard coded target name hack
 					Rules.ShouldCompileMonolithic(Platform, Configuration))
 				{
 					Rules.SetupGlobalEnvironment(
