UBT: Deleting someone's debugging message.

diff --git a/Engine/Source/Programs/UnrealBuildTool/Configuration/UEBuildPlatform.cs b/Engine/Source/Programs/UnrealBuildTool/Configuration/UEBuildPlatform.cs
--- a/Engine/Source/Programs/UnrealBuildTool/Configuration/UEBuildPlatform.cs
+++ b/Engine/Source/Programs/UnrealBuildTool/Configuration/UEBuildPlatform.cs
@@ -990,7 +990,6 @@ namespace UnrealBuildTool
 		 */
 		public virtual bool HasDefaultBuildConfig(UnrealTargetPlatform Platform, string ProjectPath)
 		{
-			Console.WriteLine("HasDefaultBuildConfig!");
 			string[] BoolKeys = new string[] {
 				"bCompileApex", "bCompileBox2D", "bCompileICU", "bCompileSimplygon", 
 				"bCompileLeanAndMeanUE", "bIncludeADO", "bCompileRecast", "bCompileSpeedTree", 
