HTML5 SDK setup catches bogus SDK paths.

diff --git a/Engine/Source/Programs/UnrealBuildTool/HTML5/HTML5SDKInfo.cs b/Engine/Source/Programs/UnrealBuildTool/HTML5/HTML5SDKInfo.cs
--- a/Engine/Source/Programs/UnrealBuildTool/HTML5/HTML5SDKInfo.cs
+++ b/Engine/Source/Programs/UnrealBuildTool/HTML5/HTML5SDKInfo.cs
@@ -236,6 +236,12 @@ namespace UnrealBuildTool
 				{
 					var SDKVersions = GetInstalledVersions(SDKPathString);
 
+					// Invalid SDK path
+					if (SDKVersions.Count == 0)
+					{
+						return "";
+					}
+
 					if (VersionString == "-1.-1.-1" && SDKVersions.Count > 0)
 					{
 						var Ver = SDKVersions[SDKVersions.Count - 1];
