Handle falling back to EMSCRIPTEN env var correctly.

diff --git a/Engine/Source/Programs/UnrealBuildTool/HTML5/HTML5SDKInfo.cs b/Engine/Source/Programs/UnrealBuildTool/HTML5/HTML5SDKInfo.cs
--- a/Engine/Source/Programs/UnrealBuildTool/HTML5/HTML5SDKInfo.cs
+++ b/Engine/Source/Programs/UnrealBuildTool/HTML5/HTML5SDKInfo.cs
@@ -201,7 +201,7 @@ namespace UnrealBuildTool
 			// Search the ini file for Emscripten root first
 			string EmscriptenRoot;
 			string EMCCPath;
-			if (ConfigCache.GetString("/Script/HTML5PlatformEditor.HTML5SDKSettings", "EmscriptenRoot", out EmscriptenRoot))
+			if (ConfigCache.GetString("/Script/HTML5PlatformEditor.HTML5SDKSettings", "EmscriptenRoot", out EmscriptenRoot) || Environment.GetEnvironmentVariable("EMSCRIPTEN") != null)
 			{
 				string SDKPathString = null;
 				string VersionString = null;
@@ -236,12 +236,6 @@ namespace UnrealBuildTool
 				{
 					var SDKVersions = GetInstalledVersions(SDKPathString);
 
-					// Invalid SDK path
-					if (SDKVersions.Count == 0)
-					{
-						return "";
-					}
-
 					if (VersionString == "-1.-1.-1" && SDKVersions.Count > 0)
 					{
 						var Ver = SDKVersions[SDKVersions.Count - 1];
@@ -250,13 +244,15 @@ namespace UnrealBuildTool
 							return Ver.Directory;
 						}
 					}
-
-					var RequiredVersion = System.Version.Parse(VersionString);
-					foreach (var i in SDKVersions)
+					else if (SDKVersions.Count > 0)
 					{
- 						if (i.Version == RequiredVersion && System.IO.Directory.Exists(i.Directory))
+						var RequiredVersion = System.Version.Parse(VersionString);
+						foreach (var i in SDKVersions)
 						{
-							return i.Directory;
+							if (i.Version == RequiredVersion && System.IO.Directory.Exists(i.Directory))
+							{
+								return i.Directory;
+							}
 						}
 					}
 				}
