Add support for precompiling modules separately to the -BuildRocket argument. Targets can specify modules which can be precompiled by overriding the GetModulesToPrecompile() function in their target rules. Running UBT with the -PrecompileModules argument will generate them in the form of static libraries for monolithic builds, or extra DLLs for non-monolithic builds, and running UBT with the -UsePrecompiledModules argument will use them.

diff --git a/Engine/Source/Programs/UnrealBuildTool/Configuration/UEBuildBinary.cs b/Engine/Source/Programs/UnrealBuildTool/Configuration/UEBuildBinary.cs
--- a/Engine/Source/Programs/UnrealBuildTool/Configuration/UEBuildBinary.cs
+++ b/Engine/Source/Programs/UnrealBuildTool/Configuration/UEBuildBinary.cs
@@ -358,12 +358,15 @@ namespace UnrealBuildTool
 				if (Config.bHasModuleRules)
 				{
 					Module = Target.FindOrCreateModuleByName(ModuleName);
-					if (Module.Binary != null)
+					if(Module.Binary == null)
+					{
+						Module.Binary = this;
+						Module.bIncludedInTarget = true;
+					}
+					else if(Module.Binary.Config.Type != UEBuildBinaryType.StaticLibrary)
 					{
 						throw new BuildException("Module \"{0}\" linked into both {1} and {2}, which creates ambiguous linkage for dependents.", ModuleName, Module.Binary.Config.OutputFilePath, Config.OutputFilePath);
 					}
-					Module.Binary = this;
-					Module.bIncludedInTarget = true;
 				}
 
 				// We set whether the binary is being compiled monolithic here to know later - specifically
@@ -610,21 +613,28 @@ namespace UnrealBuildTool
 			{
 				var Module = Target.GetModuleByName(ModuleName);
 
-				// Compile each module.
-				Log.TraceVerbose("Compile module: " + ModuleName);
-
-				var LinkInputFiles = Module.Compile(CompileEnvironment, BinaryCompileEnvironment, Config.bCompileMonolithic);
-
-				// NOTE: Because of 'Shared PCHs', in monolithic builds the same PCH file may appear as a link input
-				// multiple times for a single binary.  We'll check for that here, and only add it once.  This avoids
-				// a linker warning about redundant .obj files. 
-				foreach (var LinkInputFile in LinkInputFiles)
+				List<FileItem> LinkInputFiles; 
+				if(Module.Binary == null || Module.Binary == this)
 				{
-					if (!BinaryLinkEnvironment.InputFiles.Contains(LinkInputFile))
+					// Compile each module.
+					Log.TraceVerbose("Compile module: " + ModuleName);
+					LinkInputFiles = Module.Compile(CompileEnvironment, BinaryCompileEnvironment, Config.bCompileMonolithic);
+
+					// NOTE: Because of 'Shared PCHs', in monolithic builds the same PCH file may appear as a link input
+					// multiple times for a single binary.  We'll check for that here, and only add it once.  This avoids
+					// a linker warning about redundant .obj files. 
+					foreach (var LinkInputFile in LinkInputFiles)
 					{
-						BinaryLinkEnvironment.InputFiles.Add(LinkInputFile);
+						if (!BinaryLinkEnvironment.InputFiles.Contains(LinkInputFile))
+						{
+							BinaryLinkEnvironment.InputFiles.Add(LinkInputFile);
+						}
 					}
 				}
+				else 
+				{
+					BinaryDependencies.Add(Module.Binary);
+				}
 
 				if (!BuildConfiguration.bRunUnrealCodeAnalyzer)
 				{
diff --git a/Engine/Source/Programs/UnrealBuildTool/Configuration/UEBuildModule.cs b/Engine/Source/Programs/UnrealBuildTool/Configuration/UEBuildModule.cs
--- a/Engine/Source/Programs/UnrealBuildTool/Configuration/UEBuildModule.cs
+++ b/Engine/Source/Programs/UnrealBuildTool/Configuration/UEBuildModule.cs
@@ -809,6 +809,15 @@ namespace UnrealBuildTool
 		{
 		}
 
+		/**
+		 * Gets all of the modules precompiled along with this module
+		 * 
+		 * @param	Modules	Set of all the precompiled modules
+		 */
+		public virtual void RecursivelyAddPrecompiledModules(List<UEBuildModule> Modules)
+		{
+		}
+
 		/**
 		 * Gathers and binds binaries for this module and all of it's dependent modules
 		 *
@@ -987,14 +996,8 @@ namespace UnrealBuildTool
 		/** Enable exception handling */
 		public bool bEnableExceptions = false;
 
-		/** Path to this module's redist static library */
-		public string[] RedistStaticLibraryPaths = null;
-
 		public List<string> IncludeSearchPaths = new List<string>();
 
-		/** Whether we're building the redist static library (as well as using it). */
-		public bool bBuildingRedistStaticLibrary = false;
-
 		public class ProcessedDependenciesClass
 		{
 			/** The file, if any, which is used as the unique PCH for this module */
@@ -1143,12 +1146,6 @@ namespace UnrealBuildTool
 				return LinkInputFiles;
 			}
 
-			if(RedistStaticLibraryPaths != null && !bBuildingRedistStaticLibrary)
-			{
-				// The redist static library will be added in SetupPrivateLinkEnvironment
-				return LinkInputFiles;
-			}
-
 			var ModuleCompileEnvironment = CreateModuleCompileEnvironment(CompileEnvironment);
 			IncludeSearchPaths = ModuleCompileEnvironment.Config.CPPIncludeInfo.IncludePaths.ToList();
 			IncludeSearchPaths.AddRange(ModuleCompileEnvironment.Config.CPPIncludeInfo.SystemIncludePaths.ToList());
@@ -1269,7 +1266,7 @@ namespace UnrealBuildTool
 					Log.TraceVerbose("Module '{0}' was not allowed to use Shared PCHs, because we're compiling to a library in monolithic mode", this.Name );
 				}
 
-				bool bUseSharedPCHFiles  = BuildConfiguration.bUseSharedPCHs && (bDisableSharedPCHFiles == false) && !UnrealBuildTool.BuildingRocket();
+				bool bUseSharedPCHFiles  = BuildConfiguration.bUseSharedPCHs && (bDisableSharedPCHFiles == false) && GlobalCompileEnvironment.SharedPCHHeaderFiles.Count > 0;
 
 				if( bUseSharedPCHFiles )
 				{
@@ -1525,23 +1522,6 @@ namespace UnrealBuildTool
 			// Compile MM files directly.
 			LinkInputFiles.AddRange(CPPCompileEnvironment.CompileFiles( Target, SourceFilesToBuild.MMFiles, Name).ObjectFiles);
 
-			// If we're building Rocket, generate a static library for this module
-			if(RedistStaticLibraryPaths != null)
-			{
-				// Create a link environment for it
-				LinkEnvironment RedistLinkEnvironment = new LinkEnvironment();
-				RedistLinkEnvironment.InputFiles.AddRange(LinkInputFiles);
-				RedistLinkEnvironment.Config.Target                = CompileEnvironment.Config.Target;
-				RedistLinkEnvironment.Config.bIsBuildingDLL        = false;
-				RedistLinkEnvironment.Config.bIsBuildingLibrary    = true;
-				RedistLinkEnvironment.Config.IntermediateDirectory = Binary.Config.IntermediateDirectory;
-				RedistLinkEnvironment.Config.OutputFilePaths       = RedistStaticLibraryPaths != null ? (string[])RedistStaticLibraryPaths.Clone() : null;
-
-				// Replace the items built so far with the library
-				RedistLinkEnvironment.LinkExecutable(false);
-				LinkInputFiles.Clear();
-			}
-
 			// Compile RC files.
 			LinkInputFiles.AddRange(CPPCompileEnvironment.CompileRCFiles(Target, SourceFilesToBuild.RCFiles).ObjectFiles);
 
@@ -1829,23 +1809,6 @@ namespace UnrealBuildTool
 			return Result;
 		}
 
-		public override void SetupPrivateLinkEnvironment(
-			ref LinkEnvironment LinkEnvironment,
-			ref List<UEBuildBinary> BinaryDependencies,
-			ref Dictionary<UEBuildModule, bool> VisitedModules
-			)
-		{
-			base.SetupPrivateLinkEnvironment(ref LinkEnvironment,ref BinaryDependencies,ref VisitedModules);
-
-			if(RedistStaticLibraryPaths != null)
-			{
-				foreach (string LibVar in RedistStaticLibraryPaths)
-				{
-					LinkEnvironment.Config.AdditionalLibraries.Add(Path.GetFullPath(LibVar));
-				}
-			}
-		}
-
 		public HashSet<FileItem> _AllClassesHeaders     = new HashSet<FileItem>();
 		public HashSet<FileItem> _PublicUObjectHeaders  = new HashSet<FileItem>();
 		public HashSet<FileItem> _PrivateUObjectHeaders = new HashSet<FileItem>();
@@ -1954,6 +1917,28 @@ namespace UnrealBuildTool
 			}
 		}
 
+		public override void RecursivelyAddPrecompiledModules(List<UEBuildModule> Modules)
+		{
+			if(!Modules.Contains(this))
+			{
+				Modules.Add(this);
+
+				// Get the dependent modules
+				List<string> DependentModuleNames = new List<string>();
+				DependentModuleNames.AddRange(PrivateDependencyModuleNames);
+				DependentModuleNames.AddRange(PublicDependencyModuleNames);
+				DependentModuleNames.AddRange(DynamicallyLoadedModuleNames);
+				DependentModuleNames.AddRange(PlatformSpecificDynamicallyLoadedModuleNames);
+
+				// Find modules for each of them, and add their dependencies too
+				foreach(string DependentModuleName in DependentModuleNames)
+				{
+					UEBuildModule DependentModule = Target.FindOrCreateModuleByName(DependentModuleName);
+					DependentModule.RecursivelyAddPrecompiledModules(Modules);
+				}
+			}
+		}
+
 		public override void RecursivelyAddIncludePathModules( UEBuildTarget Target, bool bPublicIncludesOnly )
 		{
 			var AllIncludePathModuleNames = new List<string>();
diff --git a/Engine/Source/Programs/UnrealBuildTool/Configuration/UEBuildTarget.cs b/Engine/Source/Programs/UnrealBuildTool/Configuration/UEBuildTarget.cs
--- a/Engine/Source/Programs/UnrealBuildTool/Configuration/UEBuildTarget.cs
+++ b/Engine/Source/Programs/UnrealBuildTool/Configuration/UEBuildTarget.cs
@@ -166,6 +166,8 @@ namespace UnrealBuildTool
 		public bool bIsEditorRecompile;
 		public string RemoteRoot;
 		public List<OnlyModule> OnlyModules;
+		public bool bPrecompileModules;
+		public bool bUsePrecompiledModules;
 	}
 
 
@@ -219,6 +221,9 @@ namespace UnrealBuildTool
 			// If true, the recompile was launched by the editor.
 			bool bIsEditorRecompile = false;
 
+			// Settings for creating/using static libraries for the engine
+			bool bPrecompileModules = UnrealBuildTool.BuildingRocket();
+			bool bUsePrecompiledModules = UnrealBuildTool.RunningRocket();
 
 			// Combine the two arrays of arguments
 			List<string> Arguments = new List<string>(SourceArguments.Length);
@@ -379,6 +384,20 @@ namespace UnrealBuildTool
 							}
 							break;
 
+						case "-PRECOMPILEMODULES":
+							{
+								// Make static libraries for all engine modules as intermediates for this target
+								bPrecompileModules = true;
+							}
+							break;
+							
+						case "-USEPRECOMPILEDMODULES":
+							{
+								// Use existing static libraries for all engine modules in this target
+								bUsePrecompiledModules = true;
+							}
+							break;
+
 						default:
 							PossibleTargetNames.Add(Arguments[ArgumentIndex]);
 							break;
@@ -438,7 +457,9 @@ namespace UnrealBuildTool
 							AdditionalDefinitions = AdditionalDefinitions,
 							bIsEditorRecompile = bIsEditorRecompile,
 							RemoteRoot = RemoteRoot,
-							OnlyModules = OnlyModules
+							OnlyModules = OnlyModules,
+							bPrecompileModules = bPrecompileModules,
+							bUsePrecompiledModules = bUsePrecompiledModules,
 						} );
 					break;
 				}
@@ -671,6 +692,12 @@ namespace UnrealBuildTool
 		/** Remote path of the binary if it is to be synced with CookerSync */
 		public string RemoteRoot;
 
+		/** Whether to build engine modules that can be used without compiling */
+		public bool bPrecompileModules;
+
+		/** Whether to use precompiled engine modules */
+		public bool bUsePrecompiledModules;
+
 		/** The C++ environment that all the environments used to compile UE-based modules are derived from. */
 		[NonSerialized]
 		public CPPEnvironment GlobalCompileEnvironment = new CPPEnvironment();
@@ -695,6 +722,10 @@ namespace UnrealBuildTool
 		[NonSerialized]
 		public List<string> ExtraModuleNames = new List<string>();
 
+		/** Extra engine module names which are compiled but not linked into any binary in precompiled engine distributions */
+		[NonSerialized]
+		public List<UEBuildBinary> PrecompiledBinaries = new List<UEBuildBinary>();
+
 		/** True if re-compiling this target from the editor */
 		public bool bEditorRecompile;
 
@@ -734,7 +765,8 @@ namespace UnrealBuildTool
 			Configuration = InDesc.Configuration;
 			Rules = InRulesObject;
 			bEditorRecompile = InDesc.bIsEditorRecompile;
-
+			bPrecompileModules = InDesc.bPrecompileModules;
+			bUsePrecompiledModules = InDesc.bUsePrecompiledModules;
 
 			{
 				bCompileMonolithic = (Rules != null) ? Rules.ShouldCompileMonolithic(InDesc.Platform, InDesc.Configuration) : false;
@@ -1287,7 +1319,7 @@ namespace UnrealBuildTool
 				// Don't add static library files to the manifest as we do not check them into perforce.
 				// However, add them to the manifest when cleaning the project as we do want to delete 
 				// them in that case.
-				if (UEBuildConfiguration.bCleanProject == false)
+				if (UEBuildConfiguration.bCleanProject == false && !bPrecompileModules)
 				{
                     if (Binary.Config.Type == UEBuildBinaryType.StaticLibrary)
 					{
@@ -1357,7 +1389,7 @@ namespace UnrealBuildTool
 					var AllReferencedModules = Binary.GetAllDependencyModules(bIncludeDynamicallyLoaded: true, bForceCircular: true);
 					foreach (var CurModule in AllReferencedModules)
 					{
-						if (CurModule.Binary == null || CurModule.Binary == ExecutableBinary)
+						if (CurModule.Binary == null || CurModule.Binary == ExecutableBinary || CurModule.Binary.Config.Type == UEBuildBinaryType.StaticLibrary)
 						{
 							ExecutableBinary.AddModule(CurModule.Name);
 						}
@@ -1398,25 +1430,6 @@ namespace UnrealBuildTool
 					string ProjectName = Path.GetFileNameWithoutExtension(UnrealBuildTool.GetUProjectFile());
 					GlobalCompileEnvironment.Config.Definitions.Add(String.Format("UE_PROJECT_NAME={0}", ProjectName));
 				}
-
-				// Generate static libraries for monolithic games in Rocket
-				if ((UnrealBuildTool.BuildingRocket() || UnrealBuildTool.RunningRocket()) && TargetRules.IsAGame(TargetType))
-				{
-					List<UEBuildModule> Modules = ExecutableBinary.GetAllDependencyModules(true, false);
-					foreach (UEBuildModuleCPP Module in Modules.OfType<UEBuildModuleCPP>())
-					{
-						if(Utils.IsFileUnderDirectory(Module.ModuleDirectory, BuildConfiguration.RelativeEnginePath) && Module.Binary == ExecutableBinary)
-						{
-							UnrealTargetConfiguration LibraryConfiguration = (Configuration == UnrealTargetConfiguration.DebugGame)? UnrealTargetConfiguration.Development : Configuration;
-							Module.RedistStaticLibraryPaths = MakeBinaryPaths("", "UE4Game-Redist-" + Module.Name, Platform, LibraryConfiguration, UEBuildBinaryType.StaticLibrary, TargetType, null, AppName);
-							Module.bBuildingRedistStaticLibrary = UnrealBuildTool.BuildingRocket();
-                            if (Module.bBuildingRedistStaticLibrary)
-                            {
-                                SpecialRocketLibFilesThatAreBuildProducts.AddRange(Module.RedistStaticLibraryPaths);
-                            }
-						}
-					}
-				}
 			}
 
 			// On Mac and Linux we have actions that should be executed after all the binaries are created
@@ -1453,8 +1466,11 @@ namespace UnrealBuildTool
 			{
 				var UObjectDiscoveryStartTime = DateTime.UtcNow;
 
+				// Reconstruct a full list of binaries. Binaries which aren't compiled are stripped out of AppBinaries, but we still need to scan them for UHT.
+				List<UEBuildBinary> AllAppBinaries = AppBinaries.Union(PrecompiledBinaries).Distinct().ToList();
+	
 				// Figure out which modules have UObjects that we may need to generate headers for
-				foreach( var Binary in AppBinaries )
+				foreach( var Binary in AllAppBinaries )
 				{
 					var LocalUObjectModules = UObjectModules;	// For lambda access
 					var DependencyModules = Binary.GetAllDependencyModules(bIncludeDynamicallyLoaded: false, bForceCircular: false);
@@ -1629,14 +1645,27 @@ namespace UnrealBuildTool
 			GlobalLinkEnvironment.Config.LocalShadowDirectory = GlobalLinkEnvironment.Config.OutputDirectory;
 
 			// Add all of the extra modules, including game modules, that need to be compiled along
-			// with this app.  These modules aren't necessarily statically linked against, but may
-			// still be required at runtime in order for the application to load and function properly!
+			// with this app.  These modules are always statically linked in monolithic targets, but not necessarily linked to anything in modular targets,
+			// and may still be required at runtime in order for the application to load and function properly!
 			AddExtraModules();
 
-			// Bind modules for all app binaries.
+			// Add all the precompiled modules to the target. In contrast to "Extra Modules", these modules are not compiled into monolithic targets by default.
+			AddPrecompiledModules();
+
+			// Bind modules for all app binaries. Static libraries can be linked into other binaries, so bound modules to those first.
+			foreach(UEBuildBinary Binary in AppBinaries)
+			{
+				if(Binary.Config.Type == UEBuildBinaryType.StaticLibrary)
+				{
+					Binary.BindModules();
+				}
+			}
 			foreach (var Binary in AppBinaries)
 			{
-				Binary.BindModules();
+				if(Binary.Config.Type != UEBuildBinaryType.StaticLibrary)
+				{
+					Binary.BindModules();
+				}
 			}
 
 			// Process all referenced modules and create new binaries for DLL dependencies if needed
@@ -1650,7 +1679,13 @@ namespace UnrealBuildTool
 					NewBinaries.AddRange(FoundBinaries);
 				}
 			}
-			AppBinaries.AddRange(NewBinaries);
+			foreach(var NewBinary in NewBinaries)
+			{
+				if(!AppBinaries.Contains(NewBinary))
+				{
+					AppBinaries.Add(NewBinary);
+				}
+			}
 
 			// On Mac AppBinaries paths for non-console targets need to be adjusted to be inside the app bundle
 			if (GlobalLinkEnvironment.Config.Target.Platform == CPPTargetPlatform.Mac && !GlobalLinkEnvironment.Config.bIsBuildingConsoleApplication)
@@ -2056,7 +2091,7 @@ namespace UnrealBuildTool
 						var CPPModule = GetModuleByName( ModuleName ) as UEBuildModuleCPP;
 						if( CPPModule != null )
 						{
-							if( !String.IsNullOrEmpty( CPPModule.SharedPCHHeaderFile ) )
+							if( !String.IsNullOrEmpty( CPPModule.SharedPCHHeaderFile ) && CPPModule.Binary.Config.bAllowCompilation )
 							{
 								// @todo SharedPCH: Ideally we could figure the PCH header name automatically, and simply use a boolean in the module
 								//     definition to opt into exposing a shared PCH.  Unfortunately we don't determine which private PCH header "goes with"
@@ -2248,6 +2283,130 @@ namespace UnrealBuildTool
 			}
 		}
 
+		/// <summary>
+		/// Adds all the precompiled modules into the target. Precompiled modules are compiled alongside the target, but not linked into it unless directly referenced.
+		/// </summary>
+		protected void AddPrecompiledModules()
+		{
+			if(bPrecompileModules || bUsePrecompiledModules)
+			{
+				// Find all the precompiled module names.
+				List<string> PrecompiledModuleNames = new List<string>();
+				Rules.GetModulesToPrecompile(new TargetInfo(Platform, Configuration, TargetTypeOrNull), PrecompiledModuleNames);
+
+				// Add all the enabled plugins to the precompiled module list. Plugins are always precompiled, even if bPrecompileModules is not set, so we should precompile their dependencies.
+				foreach(PluginInfo Plugin in BuildPlugins)
+				{
+					foreach(PluginInfo.PluginModuleInfo Module in Plugin.Modules)
+					{
+						if (ShouldIncludePluginModule(Plugin, Module))
+						{
+							PrecompiledModuleNames.Add(Module.Name);
+						}
+					}
+				}
+
+				// When running in Rocket, all engine modules have to be precompiled, but are precompiled using a different target (UE4Game) with a separate list of precompiled module names.
+				// Add every referenced engine module to the list instead.
+				if(UnrealBuildTool.RunningRocket())
+				{
+					List<UEBuildModule> TargetModules = new List<UEBuildModule>();
+					foreach(UEBuildBinary TargetBinary in AppBinaries)
+					{
+						foreach(string ModuleName in TargetBinary.Config.ModuleNames)
+						{
+							UEBuildModule TargetModule = FindOrCreateModuleByName(ModuleName);
+							TargetModule.RecursivelyAddPrecompiledModules(TargetModules);
+						}
+					}
+					foreach(UEBuildModuleCPP TargetModule in TargetModules)
+					{
+						if(TargetModule is UEBuildModuleCPP && !RulesCompiler.IsGameModule(TargetModule.Name))
+						{
+							PrecompiledModuleNames.Add(TargetModule.Name);
+						}
+					}
+				}
+
+				// Build the final list of precompiled modules from the list of module names
+				List<UEBuildModule> PrecompiledModules = new List<UEBuildModule>();
+				foreach(string PrecompiledModuleName in PrecompiledModuleNames)
+				{
+					UEBuildModule PrecompiledModule = FindOrCreateModuleByName(PrecompiledModuleName);
+					PrecompiledModule.RecursivelyAddPrecompiledModules(PrecompiledModules);
+				}
+
+				// Find all the module names which are already bound to one binary and cannot be bound to another.
+				// For monolithic builds, modules which are already bound to static libraries (eg. plugins) do not need new binaries. 
+				// For modular builds, modules which are already bound to DLLs or executables (eg. launch, extra modules, plugins) do not need new binaries.
+				HashSet<string> BoundModuleNames = new HashSet<string>();
+				foreach(UEBuildBinary Binary in AppBinaries)
+				{
+					foreach(string ModuleName in Binary.Config.ModuleNames)
+					{
+						if(!bCompileMonolithic || Binary.Config.Type == UEBuildBinaryType.StaticLibrary)
+						{
+							BoundModuleNames.Add(ModuleName);
+						}
+					}
+				}
+
+				// Create binaries for the of names to actual modules. Make sure every referenced module is also precompiled.
+				foreach(UEBuildModule PrecompiledModule in PrecompiledModules)
+				{
+					if(PrecompiledModule is UEBuildModuleCPP && !BoundModuleNames.Contains(PrecompiledModule.Name))
+					{
+						UEBuildBinary Binary = AddBinaryForModule(PrecompiledModule.Name, bCompileMonolithic? UEBuildBinaryType.StaticLibrary : UEBuildBinaryType.DynamicLinkLibrary, bUsePrecompiledModules);
+						PrecompiledBinaries.Add(Binary);
+						BoundModuleNames.Add(PrecompiledModule.Name);
+					}
+				}
+			}
+		}
+
+		protected UEBuildBinaryCPP AddBinaryForModule(string ModuleName, UEBuildBinaryType BinaryType, bool bAllowCompilation)
+		{
+			// Get the plugin info for the module
+			PluginInfo Plugin = Plugins.GetPluginInfoForModule(ModuleName);
+
+			// Get the output paths for it. 
+			// HACK: Rocket does not allow overriding the compile environment, so always use the UE4Game prefix. This should be fixed by allowing targets choose whether to use a custom build environment.
+			string[] OutputFilePaths;
+			if(TargetType == TargetRules.TargetType.Game && UnrealBuildTool.RunningRocket())
+			{
+				OutputFilePaths = MakeBinaryPaths(ModuleName, "UE4Game-" + ModuleName, BinaryType, TargetType, Plugin, AppName);
+			}
+			else
+			{
+				OutputFilePaths = MakeBinaryPaths(ModuleName, AppName + "-" + ModuleName, BinaryType, TargetType, Plugin, AppName);
+			}
+
+			// Get the intermediate path
+			string IntermediateDirectory = RulesCompiler.IsGameModule(ModuleName)? ProjectIntermediateDirectory : EngineIntermediateDirectory;
+
+			// If it's a plugin, check if it actually has source.
+			if(Plugin != null)
+			{
+				string ModuleFileName = RulesCompiler.GetModuleFilename(ModuleName);
+				if(String.IsNullOrEmpty(ModuleFileName))
+				{
+					bAllowCompilation = false;
+				}
+				else if(!Directory.EnumerateFiles(Path.GetDirectoryName(ModuleFileName), "*.cpp", SearchOption.AllDirectories).Any())
+				{
+					bAllowCompilation = false;
+				}
+			}
+
+			// Get the binary configuration
+			UEBuildBinaryConfiguration BinaryConfig = new UEBuildBinaryConfiguration(BinaryType, OutputFilePaths, IntermediateDirectory, true, bInAllowCompilation: bPrecompileModules, bInCompileMonolithic: bCompileMonolithic, InModuleNames: new List<string>{ ModuleName });
+
+			// Create the binary
+			UEBuildBinaryCPP Binary = new UEBuildBinaryCPP(this, BinaryConfig);
+			AppBinaries.Add(Binary);
+			return Binary;
+		}
+
 		/**
 		 * @return true if debug information is created, false otherwise.
 		 */
diff --git a/Engine/Source/Programs/UnrealBuildTool/IOS/IOSToolChain.cs b/Engine/Source/Programs/UnrealBuildTool/IOS/IOSToolChain.cs
--- a/Engine/Source/Programs/UnrealBuildTool/IOS/IOSToolChain.cs
+++ b/Engine/Source/Programs/UnrealBuildTool/IOS/IOSToolChain.cs
@@ -1126,25 +1126,6 @@ namespace UnrealBuildTool
 					BuiltBinaries.Add(Path.GetFullPath(Binary.ToString()));
 				}
 
-				// Generate static libraries for monolithic games in Rocket
-				if ((UnrealBuildTool.BuildingRocket() || UnrealBuildTool.RunningRocket()) && TargetRules.IsAGame(Target.TargetType))
-				{
-					List<UEBuildModule> Modules = Target.AppBinaries[0].GetAllDependencyModules(true, false);
-					foreach (UEBuildModuleCPP Module in Modules.OfType<UEBuildModuleCPP>())
-					{
-						if (Utils.IsFileUnderDirectory(Module.ModuleDirectory, BuildConfiguration.RelativeEnginePath) && Module.Binary == Target.AppBinaries[0])
-						{
-							if (Module.bBuildingRedistStaticLibrary)
-							{
-								foreach (var LibraryPath in Module.RedistStaticLibraryPaths)
-								{
-									BuiltBinaries.Add(Path.GetFullPath(LibraryPath));
-								}
-							}
-						}
-					}
-				}
-
 				// check to see if the DangerouslyFast mode is valid (in other words, a build has gone through since a Rebuild/Clean operation)
 				string DangerouslyFastValidFile = Path.Combine(Target.GlobalLinkEnvironment.Config.IntermediateDirectory, "DangerouslyFastIsNotDangerous");
 				bool bUseDangerouslyFastModeWasRequested = bUseDangerouslyFastMode;
diff --git a/Engine/Source/Programs/UnrealBuildTool/System/RulesCompiler.cs b/Engine/Source/Programs/UnrealBuildTool/System/RulesCompiler.cs
--- a/Engine/Source/Programs/UnrealBuildTool/System/RulesCompiler.cs
+++ b/Engine/Source/Programs/UnrealBuildTool/System/RulesCompiler.cs
@@ -656,6 +656,17 @@ namespace UnrealBuildTool
 			)
 		{
 		}
+
+		/// <summary>
+		/// Allows the target to specify modules which can be precompiled with the -PrecompileModules/-UsePrecompiledModules arguments to UBT. 
+		/// All dependencies of the specified modules will be included.
+		/// </summary>
+		/// <param name="Target">The target information, such as platform and configuration</param>
+		/// <param name="ModuleNames">List which receives module names to precompile</param>
+		public virtual void GetModulesToPrecompile(TargetInfo Target, List<string> ModuleNames)
+		{
+		}
+
         /// <summary>
         /// Return true if this target should always be built with the base editor. Usually programs like shadercompilerworker.
         /// </summary>
diff --git a/Engine/Source/UE4Editor.Target.cs b/Engine/Source/UE4Editor.Target.cs
--- a/Engine/Source/UE4Editor.Target.cs
+++ b/Engine/Source/UE4Editor.Target.cs
@@ -17,34 +17,40 @@ public class UE4EditorTarget : TargetRules
 		)
 	{
 		OutExtraModuleNames.Add("UE4Game");
-		OutExtraModuleNames.Add("GameMenuBuilder");
+	}
+
+	public override void SetupGlobalEnvironment(
+		TargetInfo Target,
+		ref LinkEnvironmentConfiguration OutLinkEnvironmentConfiguration,
+		ref CPPEnvironmentConfiguration OutCPPEnvironmentConfiguration
+		)
+	{
+	}
+
+	public override void GetModulesToPrecompile(TargetInfo Target, List<string> ModuleNames)
+	{
+		ModuleNames.Add("Launch");
+		ModuleNames.Add("GameMenuBuilder");
 		if ((Target.Platform == UnrealTargetPlatform.Win32) || (Target.Platform == UnrealTargetPlatform.Win64))
 		{
-			OutExtraModuleNames.Add("OnlineSubsystemNull");
-			OutExtraModuleNames.Add("OnlineSubsystemAmazon");
+			ModuleNames.Add("OnlineSubsystemNull");
+			ModuleNames.Add("OnlineSubsystemAmazon");
 			if (UEBuildConfiguration.bCompileSteamOSS == true)
 			{
-				OutExtraModuleNames.Add("OnlineSubsystemSteam");
+				ModuleNames.Add("OnlineSubsystemSteam");
 			}
-			OutExtraModuleNames.Add("OnlineSubsystemFacebook");
+			ModuleNames.Add("OnlineSubsystemFacebook");
 		}
 		else if (Target.Platform == UnrealTargetPlatform.Mac || Target.Platform == UnrealTargetPlatform.Linux)
 		{
-			OutExtraModuleNames.Add("OnlineSubsystemNull");
+			ModuleNames.Add("OnlineSubsystemNull");
 			if (UEBuildConfiguration.bCompileSteamOSS == true)
 			{
-				OutExtraModuleNames.Add("OnlineSubsystemSteam");
+				ModuleNames.Add("OnlineSubsystemSteam");
 			}
 		}
 	}
 
-	public override void SetupGlobalEnvironment(
-		TargetInfo Target,
-		ref LinkEnvironmentConfiguration OutLinkEnvironmentConfiguration,
-		ref CPPEnvironmentConfiguration OutCPPEnvironmentConfiguration
-		)
-	{
-	}
     public override GUBPProjectOptions GUBP_IncludeProjectInPromotedBuild_EditorTypeOnly(UnrealTargetPlatform HostPlatform)
     {
         var Result = new GUBPProjectOptions();
diff --git a/Engine/Source/UE4Game.Target.cs b/Engine/Source/UE4Game.Target.cs
--- a/Engine/Source/UE4Game.Target.cs
+++ b/Engine/Source/UE4Game.Target.cs
@@ -8,6 +8,7 @@ public class UE4GameTarget : TargetRules
 	public UE4GameTarget( TargetInfo Target )
 	{
 		Type = TargetType.Game;
+
 		// Output to Engine/Binaries/<PLATFORM> even if built as monolithic
 		bOutputToEngineBinaries = true;
 	}
@@ -29,46 +30,6 @@ public class UE4GameTarget : TargetRules
 		{
 			OutExtraModuleNames.Add("OnlineSubsystemNull");
 		}
-
-
-		if (UnrealBuildTool.UnrealBuildTool.BuildingRocket())
-		{
-			OutExtraModuleNames.Add("GameMenuBuilder");
-
-			if (Target.Platform == UnrealTargetPlatform.Win32 || Target.Platform == UnrealTargetPlatform.Win64)
-			{
-				OutExtraModuleNames.Add("OnlineSubsystemNull");
-				OutExtraModuleNames.Add("OnlineSubsystemAmazon");
-				if (UEBuildConfiguration.bCompileSteamOSS == true)
-				{
-					OutExtraModuleNames.Add("OnlineSubsystemSteam");
-				}
-				OutExtraModuleNames.Add("OnlineSubsystemFacebook");
-			}
-			else if (Target.Platform == UnrealTargetPlatform.Mac || Target.Platform == UnrealTargetPlatform.Linux)
-			{
-				OutExtraModuleNames.Add("OnlineSubsystemNull");
-				if (UEBuildConfiguration.bCompileSteamOSS == true)
-				{
-					OutExtraModuleNames.Add("OnlineSubsystemSteam");
-				}
-			}
-			else if (Target.Platform == UnrealTargetPlatform.IOS)
-			{
-				OutExtraModuleNames.Add("OnlineSubsystemFacebook");
-				OutExtraModuleNames.Add("OnlineSubsystemIOS");
-				OutExtraModuleNames.Add("IOSAdvertising");
-			}
-			else if (Target.Platform == UnrealTargetPlatform.Android)
-			{
-				// @todo android: Add Android online subsystem
-				OutExtraModuleNames.Add("AndroidAdvertising");
-			}
-			else if (Target.Platform == UnrealTargetPlatform.HTML5)
-			{
-				OutExtraModuleNames.Add("OnlineSubsystemNull");
-			}
-		}
 	}
 
 	public override void SetupGlobalEnvironment(
@@ -90,6 +51,47 @@ public class UE4GameTarget : TargetRules
 			UEBuildConfiguration.bCompileForSize = true;*/
 		}
 	}
+
+	public override void GetModulesToPrecompile(TargetInfo Target, List<string> ModuleNames)
+	{
+		// Add all the precompiled modules for this target
+		ModuleNames.Add("Launch");
+		ModuleNames.Add("GameMenuBuilder");
+		if (Target.Platform == UnrealTargetPlatform.Win32 || Target.Platform == UnrealTargetPlatform.Win64)
+		{
+			ModuleNames.Add("OnlineSubsystemNull");
+			ModuleNames.Add("OnlineSubsystemAmazon");
+			if (UEBuildConfiguration.bCompileSteamOSS == true)
+			{
+				ModuleNames.Add("OnlineSubsystemSteam");
+			}
+			ModuleNames.Add("OnlineSubsystemFacebook");
+		}
+		else if (Target.Platform == UnrealTargetPlatform.Mac || Target.Platform == UnrealTargetPlatform.Linux)
+		{
+			ModuleNames.Add("OnlineSubsystemNull");
+			if (UEBuildConfiguration.bCompileSteamOSS == true)
+			{
+				ModuleNames.Add("OnlineSubsystemSteam");
+			}
+		}
+		else if (Target.Platform == UnrealTargetPlatform.IOS)
+		{
+			ModuleNames.Add("OnlineSubsystemFacebook");
+			ModuleNames.Add("OnlineSubsystemIOS");
+			ModuleNames.Add("IOSAdvertising");
+		}
+		else if (Target.Platform == UnrealTargetPlatform.Android)
+		{
+			// @todo android: Add Android online subsystem
+			ModuleNames.Add("AndroidAdvertising");
+		}
+		else if (Target.Platform == UnrealTargetPlatform.HTML5)
+		{
+			ModuleNames.Add("OnlineSubsystemNull");
+		}
+	}
+
 	public override List<UnrealTargetPlatform> GUBP_GetPlatforms_MonolithicOnly(UnrealTargetPlatform HostPlatform)
     {
 		List<UnrealTargetPlatform> Platforms = null;
