Change ModuleManager to automatically scan NotForLicensees/NoRedist/CarefullyRedist subfolders when looking for modules, and add Mac support for putting dylibs in subfolders.

diff --git a/Engine/Source/Programs/UnrealBuildTool/Mac/MacToolChain.cs b/Engine/Source/Programs/UnrealBuildTool/Mac/MacToolChain.cs
--- a/Engine/Source/Programs/UnrealBuildTool/Mac/MacToolChain.cs
+++ b/Engine/Source/Programs/UnrealBuildTool/Mac/MacToolChain.cs
@@ -55,6 +55,8 @@ namespace UnrealBuildTool
 
 		private static List<FileItem> BundleDependencies = new List<FileItem>();
 
+		private static List<string> BundleDylibPaths = new List<string>();
+
 		public List<string> BuiltBinaries = new List<string>();
 
 		public override void SetUpGlobalEnvironment()
@@ -686,6 +688,12 @@ namespace UnrealBuildTool
 
 			if (!bIsBuildingLibrary || LinkEnvironment.Config.bIncludeDependentLibrariesInLibrary)
 			{
+				// Add RPaths for other bundle modules
+				foreach (string BundleDylibPath in BundleDylibPaths)
+				{
+					LinkCommand += String.Format(" -rpath @executable_path/{0}/", BundleDylibPath);
+				}
+
 				// Add the additional libraries to the argument list.
 				foreach (string AdditionalLibrary in LinkEnvironment.Config.AdditionalLibraries)
 				{
@@ -1279,11 +1287,13 @@ namespace UnrealBuildTool
 				BundleDependencies.Add(FileItem.GetItemByPath(Binary.ToString()));
 			}
 		}
-		
+			
 		public override void FixBundleBinariesPaths(UEBuildTarget Target, List<UEBuildBinary> Binaries)
 		{
 			base.FixBundleBinariesPaths(Target, Binaries);
 
+			BundleDylibPaths.Clear();
+
 			string BundleContentsPath = Target.OutputPath + ".app/Contents/";
 			foreach (UEBuildBinary Binary in Binaries)
 			{
@@ -1297,13 +1307,15 @@ namespace UnrealBuildTool
 					if (DylibDir.StartsWith(ExeDir))
 					{
 						// get the subdir, which is the DylibDir - ExeDir
-						// @todo: This is commented out because the code in ApplePlatformFile.cpp for finding frameworks doesn't handle subdirectories
-						// at all yet, and it doesn't seem that the subdirectory actually matters in the Mac case, because the JunkManifest doesn't
-						// run over every single .app directory to delete them anyway - and the platforms Mac supports aren't the type to need that
-						// kind of protection to be deleted if they aren't supported for a current user that has the dylibs locally.
-						// string SubDir = DylibDir.Replace(ExeDir, "");
-						string SubDir = "";
+						string SubDir = DylibDir.Replace(ExeDir, "");
 						Binary.Config.OutputFilePaths[0] = BundleContentsPath + "MacOS" + SubDir + "/" + BinaryFileName;
+
+						// Add the path to the list of search paths for the bundle
+						string TrimSubDir = SubDir.Trim('/', '\\');
+						if (TrimSubDir.Length > 0 && !BundleDylibPaths.Contains(TrimSubDir))
+						{
+							BundleDylibPaths.Add(TrimSubDir);
+						}
 					}
 				}
 				else if (!BinaryFileName.EndsWith(".a") && !Binary.Config.OutputFilePath.Contains(".app/Contents/MacOS/")) // Binaries can contain duplicates
diff --git a/Engine/Source/Runtime/Core/Private/Modules/ModuleManager.cpp b/Engine/Source/Runtime/Core/Private/Modules/ModuleManager.cpp
--- a/Engine/Source/Runtime/Core/Private/Modules/ModuleManager.cpp
+++ b/Engine/Source/Runtime/Core/Private/Modules/ModuleManager.cpp
@@ -866,6 +866,17 @@ void FModuleManager::FindModulePathsInDirectory(const FString& InDirectoryName,
 			}
 		}
 	}
+
+	// Also recurse into restricted sub-folders, if they exist
+	const TCHAR* RestrictedFolderNames[] = { TEXT("NoRedist"), TEXT("NotForLicensees"), TEXT("CarefullyRedist") };
+	for(const TCHAR* RestrictedFolderName: RestrictedFolderNames)
+	{
+		FString RestrictedFolder = InDirectoryName / RestrictedFolderName;
+		if(FPaths::DirectoryExists(RestrictedFolder))
+		{
+			FindModulePathsInDirectory(RestrictedFolder, bIsGameDirectory, NamePattern, OutModulePaths);
+		}
+	}
 }
 
 
