Disable stripping symbols on iOS. Will remove the dsym files once they're being output separately, or just strip binaries as part of the build process.

diff --git a/Engine/Source/Programs/UnrealBuildTool/IOS/IOSToolChain.cs b/Engine/Source/Programs/UnrealBuildTool/IOS/IOSToolChain.cs
--- a/Engine/Source/Programs/UnrealBuildTool/IOS/IOSToolChain.cs
+++ b/Engine/Source/Programs/UnrealBuildTool/IOS/IOSToolChain.cs
@@ -273,7 +273,7 @@ namespace UnrealBuildTool
 			}
 
 			// Create DWARF format debug info if wanted,
-			if (CompileEnvironment.Config.bCreateDebugInfo && !UnrealBuildTool.BuildingRocket())
+			if (CompileEnvironment.Config.bCreateDebugInfo)
 			{
 				Result += " -gdwarf-2";
 			}
@@ -997,7 +997,7 @@ namespace UnrealBuildTool
 				string FinalRemoteExecutablePath = String.Format("{0}/Payload/{1}.app/{1}", RemoteShadowDirectoryMac, AppName);
 
 				// strip the debug info from the executable if needed
-				if (BuildConfiguration.bStripSymbolsOnIOS || (Target.Configuration == UnrealTargetConfiguration.Shipping) || UnrealBuildTool.BuildingRocket())
+				if (BuildConfiguration.bStripSymbolsOnIOS || (Target.Configuration == UnrealTargetConfiguration.Shipping))
 				{
 					Process StripProcess = new Process();
 					StripProcess.StartInfo.WorkingDirectory = RemoteShadowDirectoryMac;
