Remove Android symbol stripping as part of Rocket build process. Will make stripping part of the build process instead.

diff --git a/Engine/Source/Programs/UnrealBuildTool/Android/AndroidToolChain.cs b/Engine/Source/Programs/UnrealBuildTool/Android/AndroidToolChain.cs
--- a/Engine/Source/Programs/UnrealBuildTool/Android/AndroidToolChain.cs
+++ b/Engine/Source/Programs/UnrealBuildTool/Android/AndroidToolChain.cs
@@ -308,7 +308,7 @@ namespace UnrealBuildTool
 			}
 
 			// debug info
-            if (CompileEnvironment.Config.bCreateDebugInfo && !UnrealBuildTool.BuildingRocket())
+            if (CompileEnvironment.Config.bCreateDebugInfo)
 			{
 				Result += " -g2 -gdwarf-2";
 			}
@@ -438,11 +438,6 @@ namespace UnrealBuildTool
 			Result += " -Wl,-shared,-Bsymbolic";
 			Result += " -Wl,--no-undefined";
 
-            if (UnrealBuildTool.BuildingRocket())
-            {
-                Result += " -Wl,--strip-debug";
-            }
-
             if (Architecture == "-armv7")
 			{
 				Result += ToolchainParamsArm;
