Don't disable PS4, XboxOne or WinRT when building Rocket builds - we can filter out build products by removing /PS4/, /XboxOne/, /WinRT/ directories as a post-build step instead.

diff --git a/Engine/Source/Programs/UnrealBuildTool/WinRT/UEBuildWinRT.cs b/Engine/Source/Programs/UnrealBuildTool/WinRT/UEBuildWinRT.cs
--- a/Engine/Source/Programs/UnrealBuildTool/WinRT/UEBuildWinRT.cs
+++ b/Engine/Source/Programs/UnrealBuildTool/WinRT/UEBuildWinRT.cs
@@ -54,7 +54,7 @@ namespace UnrealBuildTool
         protected override void RegisterBuildPlatformInternal()
         {
             [REDACTED_EMAIL]: Add platform support
-            if (UnrealBuildTool.RunningRocket() || UnrealBuildTool.BuildingRocket() || Utils.IsRunningOnMono)
+            if (UnrealBuildTool.RunningRocket() || Utils.IsRunningOnMono)
             {
                 return;
             }
