Always strip the full binary path name from PDBs; we always keep them side-by-side with the binary, and want them to be usable on different machines.

diff --git a/Engine/Source/Programs/UnrealBuildTool/Windows/VCToolChain.cs b/Engine/Source/Programs/UnrealBuildTool/Windows/VCToolChain.cs
--- a/Engine/Source/Programs/UnrealBuildTool/Windows/VCToolChain.cs
+++ b/Engine/Source/Programs/UnrealBuildTool/Windows/VCToolChain.cs
@@ -795,11 +795,8 @@ namespace UnrealBuildTool
 				}
 			}
 
-			// Don't embed the full PDB path in the binary when building Rocket executables; the location on disk won't match the user's install directory.
-			if(UnrealBuildTool.BuildingRocket())
-			{
-				Arguments.Append(" /PDBALTPATH:%_PDB%");
-			}
+			// Don't embed the full PDB path; we want to be able to move binaries elsewhere. They will always be side by side.
+			Arguments.Append(" /PDBALTPATH:%_PDB%");
 
 			//
 			//	Shipping & LTCG
