Remove filtering for plugins based on BuildingRocket(). We can filter out resulting binaries using directory name checks instead.

diff --git a/Engine/Source/Programs/UnrealBuildTool/System/Plugins.cs b/Engine/Source/Programs/UnrealBuildTool/System/Plugins.cs
--- a/Engine/Source/Programs/UnrealBuildTool/System/Plugins.cs
+++ b/Engine/Source/Programs/UnrealBuildTool/System/Plugins.cs
@@ -386,33 +386,30 @@ namespace UnrealBuildTool
 			var PluginsDirectoryInfo = new DirectoryInfo(PluginsDirectory);
 			foreach( var PossiblePluginDirectory in PluginsDirectoryInfo.EnumerateDirectories() )
 			{
-				if (!UnrealBuildTool.IsPathIgnoredForBuild(PossiblePluginDirectory.FullName))
+				// Do we have a plugin descriptor in this directory?
+				bool bFoundPlugin = false;
+				foreach (var PluginDescriptorFileName in Directory.GetFiles(PossiblePluginDirectory.FullName, "*" + PluginDescriptorFileExtension))
 				{
-					// Do we have a plugin descriptor in this directory?
-					bool bFoundPlugin = false;
-					foreach (var PluginDescriptorFileName in Directory.GetFiles(PossiblePluginDirectory.FullName, "*" + PluginDescriptorFileExtension))
+					// Found a plugin directory!  No need to recurse any further, but make sure it's unique.
+					if (!Plugins.Any(x => x.Directory == PossiblePluginDirectory.FullName))
 					{
-						// Found a plugin directory!  No need to recurse any further, but make sure it's unique.
-						if (!Plugins.Any(x => x.Directory == PossiblePluginDirectory.FullName))
-						{
-							// Load the plugin info and keep track of it
-							var PluginDescriptorFile = new FileInfo(PluginDescriptorFileName);
-							var PluginInfo = LoadPluginDescriptor(PluginDescriptorFile, LoadedFrom);
-
-							Plugins.Add(PluginInfo);
-							bFoundPlugin = true;
-							Log.TraceVerbose("Found plugin in: " + PluginInfo.Directory);
-						}
+						// Load the plugin info and keep track of it
+						var PluginDescriptorFile = new FileInfo(PluginDescriptorFileName);
+						var PluginInfo = LoadPluginDescriptor(PluginDescriptorFile, LoadedFrom);
 
-						// No need to search for more plugins
-						break;
+						Plugins.Add(PluginInfo);
+						bFoundPlugin = true;
+						Log.TraceVerbose("Found plugin in: " + PluginInfo.Directory);
 					}
 
-					if (!bFoundPlugin)
-					{
-						// Didn't find a plugin in this directory.  Continue to look in subfolders.
-						FindPluginsRecursively(PossiblePluginDirectory.FullName, LoadedFrom, ref Plugins);
-					}
+					// No need to search for more plugins
+					break;
+				}
+
+				if (!bFoundPlugin)
+				{
+					// Didn't find a plugin in this directory.  Continue to look in subfolders.
+					FindPluginsRecursively(PossiblePluginDirectory.FullName, LoadedFrom, ref Plugins);
 				}
 			}
 		}
diff --git a/Engine/Source/Programs/UnrealBuildTool/System/UnrealBuildTool.cs b/Engine/Source/Programs/UnrealBuildTool/System/UnrealBuildTool.cs
--- a/Engine/Source/Programs/UnrealBuildTool/System/UnrealBuildTool.cs
+++ b/Engine/Source/Programs/UnrealBuildTool/System/UnrealBuildTool.cs
@@ -418,24 +418,6 @@ namespace UnrealBuildTool
             }
         }
 
-		/// <summary>
-		/// Determines whether the given path is ignored for the active build.
-		/// </summary>
-		/// <param name="InString">The path to check</param>
-		/// <returns>true if the path should be ignored for the build</returns>
-		public static bool IsPathIgnoredForBuild(string InPath)
-		{
-			if (BuildingRocket())
-			{
-				string NormalizedPath = InPath.Replace('\\', '/').ToLowerInvariant();
-				if (NormalizedPath.Contains("/notforlicensees/") || NormalizedPath.Contains("/epicinternal/") || NormalizedPath.Contains("/noredist/"))
-				{
-					return true;
-				}
-			}
-			return false;
-		}
-
         /// <summary>
         /// See if the given string was pass in on the command line
         /// </summary>
