Build all plugins for a target when we're precompiling modules.

diff --git a/Engine/Source/Programs/UnrealBuildTool/Configuration/UEBuildTarget.cs b/Engine/Source/Programs/UnrealBuildTool/Configuration/UEBuildTarget.cs
--- a/Engine/Source/Programs/UnrealBuildTool/Configuration/UEBuildTarget.cs
+++ b/Engine/Source/Programs/UnrealBuildTool/Configuration/UEBuildTarget.cs
@@ -2621,7 +2621,7 @@ namespace UnrealBuildTool
 			DependentPlugins.AddRange(EnabledPlugins);
 
 			// Set the list of plugins that should be built
-			if (UnrealBuildTool.BuildingRocket() && TargetType != TargetRules.TargetType.Program)
+			if (bPrecompileModules && TargetType != TargetRules.TargetType.Program)
 			{
 				BuildPlugins.AddRange(ValidPlugins);
 			}
