Remove symbol stripping for Rocket from UBT. Will strip inside build process instead.

diff --git a/Engine/Source/Programs/UnrealBuildTool/Linux/LinuxToolChain.cs b/Engine/Source/Programs/UnrealBuildTool/Linux/LinuxToolChain.cs
--- a/Engine/Source/Programs/UnrealBuildTool/Linux/LinuxToolChain.cs
+++ b/Engine/Source/Programs/UnrealBuildTool/Linux/LinuxToolChain.cs
@@ -454,12 +454,6 @@ namespace UnrealBuildTool
                 Result += string.Format(" -Wl,--unresolved-symbols=ignore-in-shared-libs");
             }
 
-            if (UnrealBuildTool.BuildingRocket())
-            {
-                // strip symbols for Rocket in every configuration
-                Result += " -Wl,-s";
-            }
-
             // RPATH for third party libs
             Result += " -Wl,-rpath=${ORIGIN}";
             Result += " -Wl,-rpath-link=${ORIGIN}";
