Removing last few calls to BuildingRocket()!

* Import libraries are added to the build manifest whenever precompiling modules
* UHT now just generates headers to the project directory, and headers aren't included in the engine distribution. UHT was generating them anyway, just not writing them.
* Plugins binaries now just use UE4Game prefix in Rocket (since that's where they're built from), rather than using a special name.
* Platform-specific plugins are not explicitly excluded from builds - will be stripped out by directory instead
* XMPP is not manually excluded from Rocket builds.
* Removed BuildingRocket() function.

diff --git a/Engine/Source/Programs/AutomationTool/Scripts/GUBP.Automation.cs b/Engine/Source/Programs/AutomationTool/Scripts/GUBP.Automation.cs
--- a/Engine/Source/Programs/AutomationTool/Scripts/GUBP.Automation.cs
+++ b/Engine/Source/Programs/AutomationTool/Scripts/GUBP.Automation.cs
@@ -148,7 +148,7 @@ public class GUBP : BuildCommand
 		string Args = "";
         if (bBuildRocket)
         {
-            Args = " -NoSimplygon " + (bUseRocketInsteadOfBuildRocket ? "-Rocket" : "-BuildRocket");
+            Args = (bUseRocketInsteadOfBuildRocket ? "-Rocket" : "-BuildRocket -PrecompileModules");
         }
         return Args;
     }
diff --git a/Engine/Source/Programs/UnrealBuildTool/Configuration/UEBuildModule.cs b/Engine/Source/Programs/UnrealBuildTool/Configuration/UEBuildModule.cs
--- a/Engine/Source/Programs/UnrealBuildTool/Configuration/UEBuildModule.cs
+++ b/Engine/Source/Programs/UnrealBuildTool/Configuration/UEBuildModule.cs
@@ -2071,12 +2071,9 @@ namespace UnrealBuildTool
 		public static string GetGeneratedCodeDirectoryForModule(UEBuildTarget Target, string ModuleDirectory, string ModuleName)
 		{
 			string BaseDirectory = null;
-			if ((Target.ShouldCompileMonolithic() || Target.TargetType == TargetRules.TargetType.Program) && 
-				(!UnrealBuildTool.BuildingRocket()) &&
-				(!UnrealBuildTool.RunningRocket() || Utils.IsFileUnderDirectory(ModuleDirectory, UnrealBuildTool.GetUProjectPath())))
+			if (Target.ShouldCompileMonolithic() || Target.TargetType == TargetRules.TargetType.Program || UnrealBuildTool.RunningRocket())
 			{
-				// Monolithic configurations and programs have their intermediate headers stored under their
-				// respective project folders with the exception of rocket which always stores engine modules in the engine folder.
+				// Monolithic configurations, programs and Rocket games have their intermediate headers stored under their respective project folders.
 				string RootDirectory = UnrealBuildTool.GetUProjectPath();
 				if (String.IsNullOrEmpty(RootDirectory))
 				{
diff --git a/Engine/Source/Programs/UnrealBuildTool/Configuration/UEBuildTarget.cs b/Engine/Source/Programs/UnrealBuildTool/Configuration/UEBuildTarget.cs
--- a/Engine/Source/Programs/UnrealBuildTool/Configuration/UEBuildTarget.cs
+++ b/Engine/Source/Programs/UnrealBuildTool/Configuration/UEBuildTarget.cs
@@ -222,7 +222,7 @@ namespace UnrealBuildTool
 			bool bIsEditorRecompile = false;
 
 			// Settings for creating/using static libraries for the engine
-			bool bPrecompileModules = UnrealBuildTool.BuildingRocket();
+			bool bPrecompileModules = false;
 			bool bUsePrecompiledModules = UnrealBuildTool.RunningRocket();
 
 			// Combine the two arrays of arguments
@@ -2196,17 +2196,17 @@ namespace UnrealBuildTool
 
 			// Get the output path. Don't prefix the app name for Rocket
 			string[] OutputFilePaths;
-			if ((UnrealBuildTool.BuildingRocket() || UnrealBuildTool.RunningRocket()) && bCompileMonolithic)
+			if (UnrealBuildTool.RunningRocket() && bCompileMonolithic)
 			{
-				OutputFilePaths = MakeBinaryPaths(Module.Name, Module.Name, BinaryType, TargetType, Plugin, AppName);
-                if (UnrealBuildTool.BuildingRocket())
-                {
-                    SpecialRocketLibFilesThatAreBuildProducts.AddRange(OutputFilePaths);
-                }
+				OutputFilePaths = MakeBinaryPaths(Module.Name, "UE4Game-" + Module.Name, BinaryType, TargetType, Plugin, AppName);
 			}
 			else
 			{
 				OutputFilePaths = MakeBinaryPaths(Module.Name, GetAppName() + "-" + Module.Name, BinaryType, TargetType, Plugin, AppName);
+				if (bPrecompileModules)
+				{
+					SpecialRocketLibFilesThatAreBuildProducts.AddRange(OutputFilePaths);
+				}
 			}
 
 			// Try to determine if we have the rules file
@@ -2591,12 +2591,6 @@ namespace UnrealBuildTool
 				}
 			}
 
-			// For Rocket builds, remove plugin that's designed for a NDA console
-			if(UnrealBuildTool.BuildingRocket())
-			{
-				ValidPlugins.RemoveAll(x => x.Directory.IndexOf("\\PS4\\", StringComparison.InvariantCultureIgnoreCase) != -1 || x.Directory.IndexOf("\\XboxOne\\", StringComparison.InvariantCultureIgnoreCase) != -1);
-			}
-
 			// Build a list of enabled plugins
 			List<string> EnabledPluginNames = new List<string>(Rules.AdditionalPlugins);
 
diff --git a/Engine/Source/Programs/UnrealBuildTool/System/ExternalExecution.cs b/Engine/Source/Programs/UnrealBuildTool/System/ExternalExecution.cs
--- a/Engine/Source/Programs/UnrealBuildTool/System/ExternalExecution.cs
+++ b/Engine/Source/Programs/UnrealBuildTool/System/ExternalExecution.cs
@@ -117,8 +117,7 @@ namespace UnrealBuildTool
 				PrivateHeaders           = Info.PrivateUObjectHeaders      .Select((Header) => Header.AbsolutePath).ToList(),
 				PCH                      = Info.PCH,
 				GeneratedCPPFilenameBase = Info.GeneratedCPPFilenameBase,
-				[REDACTED_EMAIL]: This assumes Engine/Source is a 'safe' folder name to check for
-				SaveExportedHeaders = !UnrealBuildTool.RunningRocket() || !Info.ModuleDirectory.Contains("Engine\\Source\\")
+				SaveExportedHeaders      = true, // Rocket just regenerates all headers to the project directory now
 
 			}).ToList();
 		}
diff --git a/Engine/Source/Programs/UnrealBuildTool/System/UnrealBuildTool.cs b/Engine/Source/Programs/UnrealBuildTool/System/UnrealBuildTool.cs
--- a/Engine/Source/Programs/UnrealBuildTool/System/UnrealBuildTool.cs
+++ b/Engine/Source/Programs/UnrealBuildTool/System/UnrealBuildTool.cs
@@ -57,9 +57,6 @@ namespace UnrealBuildTool
         /** Are we running for Rocket */
         static public bool bRunningRocket = false;
 
-        /** Are we building Rocket */
-        static bool bBuildingRocket = false;
-
         /** The project file */
         static string UProjectFile = "";
 
@@ -144,15 +141,6 @@ namespace UnrealBuildTool
             return bRunningRocket;
         }
 
-        /// <summary>
-        /// Checks to see if we're building a "Rocket" target
-        /// </summary>
-        /// <returns>True if building a "Rocket" target, such as RocketEditor (not running with "-rocket" though)</returns>
-        static public bool BuildingRocket()
-        {
-            return bBuildingRocket;
-        }
-
 		/// <summary>
 		/// Returns true if UnrealBuildTool is running using installed Engine components
 		/// </summary>
@@ -556,10 +544,6 @@ namespace UnrealBuildTool
                 UProjectFile = InArg;
                 bSetupProject = true;
             }
-            else if (LowercaseArg == "-buildrocket")
-            {
-                bBuildingRocket = true;
-            }
             else
             {
                 return false;
diff --git a/Engine/Source/Programs/UnrealBuildTool/Windows/VCToolChain.cs b/Engine/Source/Programs/UnrealBuildTool/Windows/VCToolChain.cs
--- a/Engine/Source/Programs/UnrealBuildTool/Windows/VCToolChain.cs
+++ b/Engine/Source/Programs/UnrealBuildTool/Windows/VCToolChain.cs
@@ -1672,7 +1672,7 @@ namespace UnrealBuildTool
         public override void AddFilesToManifest(BuildManifest Manifest, UEBuildBinary Binary)
         {
             // ok, this is pretty awful, we want the import libraries that go with the editor, only on the PC
-            if (UnrealBuildTool.BuildingRocket() &&
+            if (Binary.Target.bPrecompileModules &&
                 Path.GetFileNameWithoutExtension(Binary.Config.OutputFilePath).StartsWith("UE4Editor-", StringComparison.InvariantCultureIgnoreCase) &&
                 Path.GetExtension(Binary.Config.OutputFilePath).EndsWith("dll", StringComparison.InvariantCultureIgnoreCase) &&
                 Binary.Config.Type == UEBuildBinaryType.DynamicLinkLibrary)
