Remove some UBT abuses of the "ref" keyword.

diff --git a/Engine/Source/Programs/UnrealBuildTool/Configuration/UEBuildBinary.cs b/Engine/Source/Programs/UnrealBuildTool/Configuration/UEBuildBinary.cs
--- a/Engine/Source/Programs/UnrealBuildTool/Configuration/UEBuildBinary.cs
+++ b/Engine/Source/Programs/UnrealBuildTool/Configuration/UEBuildBinary.cs
@@ -227,7 +227,7 @@ namespace UnrealBuildTool
 		/// a module that depends on a module in this binary. */
 		/// </summary>
 		/// <param name="DependentLinkEnvironment">The link environment of the dependency</param>
-		public virtual void SetupDependentLinkEnvironment(ref LinkEnvironment DependentLinkEnvironment) {}
+		public virtual void SetupDependentLinkEnvironment(LinkEnvironment DependentLinkEnvironment) {}
 
 		/// <summary>
 		/// Called to allow the binary to to determine if it matches the Only module "short module name".
@@ -408,7 +408,7 @@ namespace UnrealBuildTool
 					ReferencedModules[ ModuleName ] = Module;
 
 					bool bOnlyDirectDependencies = false;
-					Module.GetAllDependencyModules(ref ReferencedModules, ref OrderedModules, bIncludeDynamicallyLoaded, bForceCircular, bOnlyDirectDependencies);
+					Module.GetAllDependencyModules(ReferencedModules, OrderedModules, bIncludeDynamicallyLoaded, bForceCircular, bOnlyDirectDependencies);
 
 					OrderedModules.Add( Module );
 				}
@@ -431,7 +431,7 @@ namespace UnrealBuildTool
 				foreach (var ModuleName in ModuleNames)
 				{
 					var Module = Target.FindOrCreateModuleByName(ModuleName);
-					Module.RecursivelyProcessUnboundModules(Target, ref Binaries, ExecutableBinary);
+					Module.RecursivelyProcessUnboundModules(Target, Binaries, ExecutableBinary);
 				}
 			}
 			else
@@ -524,7 +524,7 @@ namespace UnrealBuildTool
 		/// a module that depends on a module in this binary.
 		/// </summary>
 		/// <param name="DependentLinkEnvironment">The link environment of the dependency</param>
-		public override void SetupDependentLinkEnvironment(ref LinkEnvironment DependentLinkEnvironment)
+		public override void SetupDependentLinkEnvironment(LinkEnvironment DependentLinkEnvironment)
 		{
 			foreach (string OutputFilePath in Config.OutputFilePaths)
 			{
@@ -639,7 +639,7 @@ namespace UnrealBuildTool
 				if (!BuildConfiguration.bRunUnrealCodeAnalyzer)
 				{
 					// Allow the module to modify the link environment for the binary.
-					Module.SetupPrivateLinkEnvironment(ref BinaryLinkEnvironment, ref BinaryDependencies, ref LinkEnvironmentVisitedModules);
+					Module.SetupPrivateLinkEnvironment(BinaryLinkEnvironment, BinaryDependencies, LinkEnvironmentVisitedModules);
 				}
 			}
 
@@ -653,7 +653,7 @@ namespace UnrealBuildTool
 			// Allow the binary dependencies to modify the link environment.
 			foreach (var BinaryDependency in BinaryDependencies)
 			{
-				BinaryDependency.SetupDependentLinkEnvironment(ref BinaryLinkEnvironment);
+				BinaryDependency.SetupDependentLinkEnvironment(BinaryLinkEnvironment);
 			}
 
 			// Set the link output file.
diff --git a/Engine/Source/Programs/UnrealBuildTool/Configuration/UEBuildModule.cs b/Engine/Source/Programs/UnrealBuildTool/Configuration/UEBuildModule.cs
--- a/Engine/Source/Programs/UnrealBuildTool/Configuration/UEBuildModule.cs
+++ b/Engine/Source/Programs/UnrealBuildTool/Configuration/UEBuildModule.cs
@@ -524,11 +524,11 @@ namespace UnrealBuildTool
 		protected virtual void SetupPublicCompileEnvironment(
 			UEBuildBinary SourceBinary,
 			bool bIncludePathsOnly,
-			ref HashSet<string> IncludePaths,
-			ref HashSet<string> SystemIncludePaths,
-			ref List<string> Definitions,
-			ref List<UEBuildFramework> AdditionalFrameworks,
-			ref Dictionary<UEBuildModule, bool> VisitedModules
+			HashSet<string> IncludePaths,
+			HashSet<string> SystemIncludePaths,
+			List<string> Definitions,
+			List<UEBuildFramework> AdditionalFrameworks,
+			Dictionary<UEBuildModule, bool> VisitedModules
 			)
 		{
 			// There may be circular dependencies in compile dependencies, so we need to avoid reentrance.
@@ -595,7 +595,7 @@ namespace UnrealBuildTool
 					foreach(var DependencyName in PublicDependencyModuleNames)
 					{
 						var DependencyModule = Target.GetModuleByName(DependencyName);
-						DependencyModule.SetupPublicCompileEnvironment(SourceBinary,bIncludePathsOnly,ref IncludePaths,ref SystemIncludePaths,ref Definitions,ref AdditionalFrameworks,ref VisitedModules);
+						DependencyModule.SetupPublicCompileEnvironment(SourceBinary,bIncludePathsOnly, IncludePaths, SystemIncludePaths, Definitions, AdditionalFrameworks, VisitedModules);
 					}
 				}
 
@@ -604,7 +604,7 @@ namespace UnrealBuildTool
 				{
 					bool bInnerIncludePathsOnly = true;
 					var IncludePathModule = Target.GetModuleByName(IncludePathModuleName);
-					IncludePathModule.SetupPublicCompileEnvironment( SourceBinary, bInnerIncludePathsOnly, ref IncludePaths, ref SystemIncludePaths, ref Definitions, ref AdditionalFrameworks, ref VisitedModules );
+					IncludePathModule.SetupPublicCompileEnvironment( SourceBinary, bInnerIncludePathsOnly, IncludePaths, SystemIncludePaths, Definitions, AdditionalFrameworks, VisitedModules );
 				}
 
 				// Add the module's directory to the include path, so we can root #includes to it
@@ -654,10 +654,10 @@ namespace UnrealBuildTool
 
 		/** Sets up the environment for compiling this module. */
 		protected virtual void SetupPrivateCompileEnvironment(
-			ref HashSet<string> IncludePaths,
-			ref HashSet<string> SystemIncludePaths,
-			ref List<string> Definitions,
-			ref List<UEBuildFramework> AdditionalFrameworks
+			HashSet<string> IncludePaths,
+			HashSet<string> SystemIncludePaths,
+			List<string> Definitions,
+			List<UEBuildFramework> AdditionalFrameworks
 			)
 		{
 			var VisitedModules = new Dictionary<UEBuildModule, bool>();
@@ -672,13 +672,13 @@ namespace UnrealBuildTool
 
 			// Allow the module's public dependencies to modify the compile environment.
 			bool bIncludePathsOnly = false;
-			SetupPublicCompileEnvironment(Binary,bIncludePathsOnly,ref IncludePaths,ref SystemIncludePaths,ref Definitions,ref AdditionalFrameworks,ref VisitedModules);
+			SetupPublicCompileEnvironment(Binary,bIncludePathsOnly, IncludePaths, SystemIncludePaths, Definitions, AdditionalFrameworks, VisitedModules);
 
 			// Also allow the module's private dependencies to modify the compile environment.
 			foreach(var DependencyName in PrivateDependencyModuleNames)
 			{
 				var DependencyModule = Target.GetModuleByName(DependencyName);
-				DependencyModule.SetupPublicCompileEnvironment(Binary,bIncludePathsOnly,ref IncludePaths,ref SystemIncludePaths,ref Definitions,ref AdditionalFrameworks,ref VisitedModules);
+				DependencyModule.SetupPublicCompileEnvironment(Binary,bIncludePathsOnly, IncludePaths, SystemIncludePaths, Definitions, AdditionalFrameworks, VisitedModules);
 			}
 
 			// Add include paths from modules with header files that our private files need access to, but won't necessarily be importing
@@ -686,23 +686,23 @@ namespace UnrealBuildTool
 			{
 				bool bInnerIncludePathsOnly = true;
 				var IncludePathModule = Target.GetModuleByName(IncludePathModuleName);
-				IncludePathModule.SetupPublicCompileEnvironment(Binary, bInnerIncludePathsOnly, ref IncludePaths, ref SystemIncludePaths, ref Definitions, ref AdditionalFrameworks, ref VisitedModules);
+				IncludePathModule.SetupPublicCompileEnvironment(Binary, bInnerIncludePathsOnly, IncludePaths, SystemIncludePaths, Definitions, AdditionalFrameworks, VisitedModules);
 			}
 		}
 
 		/** Sets up the environment for linking any module that includes the public interface of this module. */ 
 		protected virtual void SetupPublicLinkEnvironment(
 			UEBuildBinary SourceBinary,
-			ref List<string> LibraryPaths,
-			ref List<string> AdditionalLibraries,
-			ref List<string> Frameworks,
-			ref List<string> WeakFrameworks,
-			ref List<UEBuildFramework> AdditionalFrameworks,
-			ref List<string> AdditionalShadowFiles,
-			ref List<UEBuildBundleResource> AdditionalBundleResources,
-			ref List<string> DelayLoadDLLs,
-			ref List<UEBuildBinary> BinaryDependencies,
-			ref Dictionary<UEBuildModule, bool> VisitedModules
+			List<string> LibraryPaths,
+			List<string> AdditionalLibraries,
+			List<string> Frameworks,
+			List<string> WeakFrameworks,
+			List<UEBuildFramework> AdditionalFrameworks,
+			List<string> AdditionalShadowFiles,
+			List<UEBuildBundleResource> AdditionalBundleResources,
+			List<string> DelayLoadDLLs,
+			List<UEBuildBinary> BinaryDependencies,
+			Dictionary<UEBuildModule, bool> VisitedModules
 			)
 		{
 			// There may be circular dependencies in compile dependencies, so we need to avoid reentrance.
@@ -739,8 +739,8 @@ namespace UnrealBuildTool
 							bool bIsInStaticLibrary = (DependencyModule.Binary != null && DependencyModule.Binary.Config.Type == UEBuildBinaryType.StaticLibrary);
 							if (bIsExternalModule || bIsInStaticLibrary)
 							{
-								DependencyModule.SetupPublicLinkEnvironment(SourceBinary, ref LibraryPaths, ref AdditionalLibraries, ref Frameworks, ref WeakFrameworks,
-									ref AdditionalFrameworks, ref AdditionalShadowFiles, ref AdditionalBundleResources, ref DelayLoadDLLs, ref BinaryDependencies, ref VisitedModules);
+								DependencyModule.SetupPublicLinkEnvironment(SourceBinary, LibraryPaths, AdditionalLibraries, Frameworks, WeakFrameworks,
+									AdditionalFrameworks, AdditionalShadowFiles, AdditionalBundleResources, DelayLoadDLLs, BinaryDependencies, VisitedModules);
 							}
 						}
 					}
@@ -765,14 +765,14 @@ namespace UnrealBuildTool
 
 		/** Sets up the environment for linking this module. */
 		public virtual void SetupPrivateLinkEnvironment(
-			ref LinkEnvironment LinkEnvironment,
-			ref List<UEBuildBinary> BinaryDependencies,
-			ref Dictionary<UEBuildModule, bool> VisitedModules
+			LinkEnvironment LinkEnvironment,
+			List<UEBuildBinary> BinaryDependencies,
+			Dictionary<UEBuildModule, bool> VisitedModules
 			)
 		{
 			// Allow the module's public dependencies to add library paths and additional libraries to the link environment.
-			SetupPublicLinkEnvironment(Binary,ref LinkEnvironment.Config.LibraryPaths,ref LinkEnvironment.Config.AdditionalLibraries,ref LinkEnvironment.Config.Frameworks,ref LinkEnvironment.Config.WeakFrameworks,
-				ref LinkEnvironment.Config.AdditionalFrameworks,ref LinkEnvironment.Config.AdditionalShadowFiles, ref LinkEnvironment.Config.AdditionalBundleResources,ref LinkEnvironment.Config.DelayLoadDLLs,ref BinaryDependencies,ref VisitedModules);
+			SetupPublicLinkEnvironment(Binary, LinkEnvironment.Config.LibraryPaths, LinkEnvironment.Config.AdditionalLibraries, LinkEnvironment.Config.Frameworks, LinkEnvironment.Config.WeakFrameworks,
+				LinkEnvironment.Config.AdditionalFrameworks,LinkEnvironment.Config.AdditionalShadowFiles, LinkEnvironment.Config.AdditionalBundleResources, LinkEnvironment.Config.DelayLoadDLLs, BinaryDependencies, VisitedModules);
 
 			// Also allow the module's public and private dependencies to modify the link environment.
 			List<string> AllDependencyModuleNames = new List<string>(PrivateDependencyModuleNames);
@@ -781,8 +781,8 @@ namespace UnrealBuildTool
 			foreach (var DependencyName in AllDependencyModuleNames)
 			{
 				var DependencyModule = Target.GetModuleByName(DependencyName);
-				DependencyModule.SetupPublicLinkEnvironment(Binary,ref LinkEnvironment.Config.LibraryPaths,ref LinkEnvironment.Config.AdditionalLibraries,ref LinkEnvironment.Config.Frameworks,ref LinkEnvironment.Config.WeakFrameworks,
-					ref LinkEnvironment.Config.AdditionalFrameworks,ref LinkEnvironment.Config.AdditionalShadowFiles, ref LinkEnvironment.Config.AdditionalBundleResources,ref LinkEnvironment.Config.DelayLoadDLLs,ref BinaryDependencies,ref VisitedModules);
+				DependencyModule.SetupPublicLinkEnvironment(Binary, LinkEnvironment.Config.LibraryPaths, LinkEnvironment.Config.AdditionalLibraries, LinkEnvironment.Config.Frameworks, LinkEnvironment.Config.WeakFrameworks,
+					LinkEnvironment.Config.AdditionalFrameworks, LinkEnvironment.Config.AdditionalShadowFiles, LinkEnvironment.Config.AdditionalBundleResources, LinkEnvironment.Config.DelayLoadDLLs, BinaryDependencies, VisitedModules);
 			}
 		}
 
@@ -805,7 +805,7 @@ namespace UnrealBuildTool
 		 * @param	bForceCircular	True if circular dependencies should be processed
 		 * @param	bOnlyDirectDependencies	True to return only this module's direct dependencies
 		 */
-		public virtual void GetAllDependencyModules(ref Dictionary<string, UEBuildModule> ReferencedModules, ref List<UEBuildModule> OrderedModules, bool bIncludeDynamicallyLoaded, bool bForceCircular, bool bOnlyDirectDependencies)
+		public virtual void GetAllDependencyModules(Dictionary<string, UEBuildModule> ReferencedModules, List<UEBuildModule> OrderedModules, bool bIncludeDynamicallyLoaded, bool bForceCircular, bool bOnlyDirectDependencies)
 		{
 		}
 
@@ -827,7 +827,7 @@ namespace UnrealBuildTool
 		 * @param	bBuildOnlyModules	True to build only specific modules, false for all
 		 * @param	ModulesToBuild		The specific modules to build
 		 */
-		public virtual void RecursivelyProcessUnboundModules(UEBuildTarget Target, ref Dictionary<string, UEBuildBinary> Binaries, UEBuildBinary ExecutableBinary)
+		public virtual void RecursivelyProcessUnboundModules(UEBuildTarget Target, Dictionary<string, UEBuildBinary> Binaries, UEBuildBinary ExecutableBinary)
 		{
 		}
 
@@ -1286,7 +1286,7 @@ namespace UnrealBuildTool
 						var DirectDependencyModules = new List<UEBuildModule>();
 						{ 
 							var ReferencedModules = new Dictionary<string, UEBuildModule>( StringComparer.InvariantCultureIgnoreCase );
-							this.GetAllDependencyModules( ref ReferencedModules, ref DirectDependencyModules, bIncludeDynamicallyLoaded:false, bForceCircular:false, bOnlyDirectDependencies:true );
+							this.GetAllDependencyModules( ReferencedModules, DirectDependencyModules, bIncludeDynamicallyLoaded:false, bForceCircular:false, bOnlyDirectDependencies:true );
 						}
 
 						int LargestSharedPCHHeaderFileIndex = -1;
@@ -1797,7 +1797,7 @@ namespace UnrealBuildTool
 			Result.Config.Definitions.AddRange(Definitions);
 
 			// Setup the compile environment for the module.
-			SetupPrivateCompileEnvironment(ref Result.Config.CPPIncludeInfo.IncludePaths, ref Result.Config.CPPIncludeInfo.SystemIncludePaths, ref Result.Config.Definitions, ref Result.Config.AdditionalFrameworks);
+			SetupPrivateCompileEnvironment(Result.Config.CPPIncludeInfo.IncludePaths, Result.Config.CPPIncludeInfo.SystemIncludePaths, Result.Config.Definitions, Result.Config.AdditionalFrameworks);
 
 			// @hack to skip adding definitions to compile environment, they will be baked into source code files
 			if (bSkipDefinitionsForCompileEnvironment)
@@ -1883,7 +1883,7 @@ namespace UnrealBuildTool
 			return CachedModuleUHTInfo;
 		}
 	
-		public override void GetAllDependencyModules( ref Dictionary<string, UEBuildModule> ReferencedModules, ref List<UEBuildModule> OrderedModules, bool bIncludeDynamicallyLoaded, bool bForceCircular, bool bOnlyDirectDependencies )
+		public override void GetAllDependencyModules( Dictionary<string, UEBuildModule> ReferencedModules, List<UEBuildModule> OrderedModules, bool bIncludeDynamicallyLoaded, bool bForceCircular, bool bOnlyDirectDependencies )
 		{
 			var AllModuleNames = new List<string>();
 			AllModuleNames.AddRange(PrivateDependencyModuleNames);
@@ -1908,7 +1908,7 @@ namespace UnrealBuildTool
 						if( !bOnlyDirectDependencies )
 						{ 
 							// Recurse into dependent modules first
-							Module.GetAllDependencyModules(ref ReferencedModules, ref OrderedModules, bIncludeDynamicallyLoaded, bForceCircular, bOnlyDirectDependencies);
+							Module.GetAllDependencyModules(ReferencedModules, OrderedModules, bIncludeDynamicallyLoaded, bForceCircular, bOnlyDirectDependencies);
 						}
 
 						OrderedModules.Add( Module );
@@ -1960,7 +1960,7 @@ namespace UnrealBuildTool
 			}
 		}
 
-		public override void RecursivelyProcessUnboundModules(UEBuildTarget Target, ref Dictionary<string, UEBuildBinary> Binaries, UEBuildBinary ExecutableBinary)
+		public override void RecursivelyProcessUnboundModules(UEBuildTarget Target, Dictionary<string, UEBuildBinary> Binaries, UEBuildBinary ExecutableBinary)
 		{
 			try
 			{
@@ -2032,7 +2032,7 @@ namespace UnrealBuildTool
 							DependencyModule.bIncludedInTarget = true;
 
 							// Also add binaries for this module's dependencies
-							DependencyModule.RecursivelyProcessUnboundModules( Target, ref Binaries, ExecutableBinary );
+							DependencyModule.RecursivelyProcessUnboundModules( Target, Binaries, ExecutableBinary );
 						}
 					}
 
@@ -2196,12 +2196,12 @@ namespace UnrealBuildTool
 		}
 
 		public override void SetupPrivateLinkEnvironment(
-			ref LinkEnvironment LinkEnvironment,
-			ref List<UEBuildBinary> BinaryDependencies,
-			ref Dictionary<UEBuildModule, bool> VisitedModules
+			LinkEnvironment LinkEnvironment,
+			List<UEBuildBinary> BinaryDependencies,
+			Dictionary<UEBuildModule, bool> VisitedModules
 			)
 		{
-			base.SetupPrivateLinkEnvironment(ref LinkEnvironment,ref BinaryDependencies,ref VisitedModules);
+			base.SetupPrivateLinkEnvironment(LinkEnvironment, BinaryDependencies, VisitedModules);
 
 			// Setup the link environment for linking a CLR binary.
 			LinkEnvironment.Config.CLRMode = CPPCLRMode.CLREnabled;
diff --git a/Engine/Source/Programs/UnrealBuildTool/Configuration/UEBuildTarget.cs b/Engine/Source/Programs/UnrealBuildTool/Configuration/UEBuildTarget.cs
--- a/Engine/Source/Programs/UnrealBuildTool/Configuration/UEBuildTarget.cs
+++ b/Engine/Source/Programs/UnrealBuildTool/Configuration/UEBuildTarget.cs
@@ -2103,7 +2103,7 @@ namespace UnrealBuildTool
 									var ModuleDependencies = new Dictionary<string, UEBuildModule>();
 									var ModuleList = new List<UEBuildModule>();
 									bool bIncludeDynamicallyLoaded = false;
-									CPPModule.GetAllDependencyModules(ref ModuleDependencies, ref ModuleList, bIncludeDynamicallyLoaded, bForceCircular: false, bOnlyDirectDependencies:false);
+									CPPModule.GetAllDependencyModules(ModuleDependencies, ModuleList, bIncludeDynamicallyLoaded, bForceCircular: false, bOnlyDirectDependencies:false);
 
 									// Figure out where to insert the shared PCH into our list, based off the module dependency ordering
 									int InsertAtIndex = SharedPCHHeaderFiles.Count;
diff --git a/Engine/Source/Programs/UnrealBuildTool/Mac/MacToolChain.cs b/Engine/Source/Programs/UnrealBuildTool/Mac/MacToolChain.cs
--- a/Engine/Source/Programs/UnrealBuildTool/Mac/MacToolChain.cs
+++ b/Engine/Source/Programs/UnrealBuildTool/Mac/MacToolChain.cs
@@ -1346,7 +1346,7 @@ namespace UnrealBuildTool
 			var LinkEnvironmentVisitedModules = new Dictionary<UEBuildModule, bool>();
 			foreach (var Module in Modules)
 			{
-				Module.SetupPrivateLinkEnvironment(ref BinaryLinkEnvironment, ref BinaryDependencies, ref LinkEnvironmentVisitedModules);
+				Module.SetupPrivateLinkEnvironment(BinaryLinkEnvironment, BinaryDependencies, LinkEnvironmentVisitedModules);
 			}
 
 			foreach (string AdditionalLibrary in BinaryLinkEnvironment.Config.AdditionalLibraries)
