Prevent trying to precompile plugins for platforms that are not available.

diff --git a/Engine/Source/Programs/UnrealBuildTool/Configuration/UEBuildTarget.cs b/Engine/Source/Programs/UnrealBuildTool/Configuration/UEBuildTarget.cs
--- a/Engine/Source/Programs/UnrealBuildTool/Configuration/UEBuildTarget.cs
+++ b/Engine/Source/Programs/UnrealBuildTool/Configuration/UEBuildTarget.cs
@@ -2591,6 +2591,16 @@ namespace UnrealBuildTool
 				}
 			}
 
+			// Remove any plugins for platforms we don't have
+			foreach(UnrealTargetPlatform TargetPlatform in Enum.GetValues(typeof(UnrealTargetPlatform)))
+			{
+				if(UEBuildPlatform.GetBuildPlatform(TargetPlatform, true) == null)
+				{
+					string DirectoryFragment = String.Format("/{0}/", TargetPlatform.ToString());
+					ValidPlugins.RemoveAll(x => x.Directory.Replace('\\', '/').Contains(DirectoryFragment));
+				}
+			}
+
 			// Build a list of enabled plugins
 			List<string> EnabledPluginNames = new List<string>(Rules.AdditionalPlugins);
 
