Prevent trying to stage bundle resources for static libraries on Mac.

diff --git a/Engine/Source/Programs/UnrealBuildTool/Mac/MacToolChain.cs b/Engine/Source/Programs/UnrealBuildTool/Mac/MacToolChain.cs
--- a/Engine/Source/Programs/UnrealBuildTool/Mac/MacToolChain.cs
+++ b/Engine/Source/Programs/UnrealBuildTool/Mac/MacToolChain.cs
@@ -1447,6 +1447,11 @@ namespace UnrealBuildTool
 		{
 			var OutputFiles = base.PostBuild(Executable, BinaryLinkEnvironment);
 
+			if (BinaryLinkEnvironment.Config.bIsBuildingLibrary)
+			{
+				return OutputFiles;
+			}
+
 			foreach (UEBuildBundleResource Resource in BinaryLinkEnvironment.Config.AdditionalBundleResources)
 			{
 				OutputFiles.Add(CopyBundleResource(Resource, Executable));
