Fix MacToolChain trying to add frameworks to static libraries.

diff --git a/Engine/Source/Programs/UnrealBuildTool/Mac/MacToolChain.cs b/Engine/Source/Programs/UnrealBuildTool/Mac/MacToolChain.cs
--- a/Engine/Source/Programs/UnrealBuildTool/Mac/MacToolChain.cs
+++ b/Engine/Source/Programs/UnrealBuildTool/Mac/MacToolChain.cs
@@ -792,10 +792,13 @@ namespace UnrealBuildTool
 				}
 			}
 
-			foreach (var Framework in AllFrameworks)
+			if (!bIsBuildingLibrary)
 			{
-				LinkCommand += AddFrameworkToLinkCommand(Framework.Key, Framework.Value ? "-weak_framework" : "-framework");
-				AddLibraryPathToRPaths(Framework.Key, AbsolutePath, ref RPaths, ref LinkCommand, bIsBuildingAppBundle);
+				foreach (var Framework in AllFrameworks)
+				{
+					LinkCommand += AddFrameworkToLinkCommand(Framework.Key, Framework.Value ? "-weak_framework" : "-framework");
+					AddLibraryPathToRPaths(Framework.Key, AbsolutePath, ref RPaths, ref LinkCommand, bIsBuildingAppBundle);
+				}
 			}
 
 			// Add the input files to a response file, and pass the response file on the command-line.
