Don't add stub files for static libraries to manifest on iOS. We only generate them from executables.

diff --git a/Engine/Source/Programs/UnrealBuildTool/IOS/IOSToolChain.cs b/Engine/Source/Programs/UnrealBuildTool/IOS/IOSToolChain.cs
--- a/Engine/Source/Programs/UnrealBuildTool/IOS/IOSToolChain.cs
+++ b/Engine/Source/Programs/UnrealBuildTool/IOS/IOSToolChain.cs
@@ -172,7 +172,7 @@ namespace UnrealBuildTool
 
 		public override void AddFilesToManifest(BuildManifest Manifest, UEBuildBinary Binary)
 		{
-			if (BuildConfiguration.bCreateStubIPA)
+			if (BuildConfiguration.bCreateStubIPA && Binary.Config.Type != UEBuildBinaryType.StaticLibrary)
 			{
 				string StubFile = Path.Combine (Path.GetDirectoryName (Binary.Config.OutputFilePath), Path.GetFileNameWithoutExtension (Binary.Config.OutputFilePath) + ".stub");
 				Manifest.AddBuildProduct(StubFile);
