Don't add mem files to manifest for HTML5

diff --git a/Engine/Source/Programs/UnrealBuildTool/HTML5/HTML5ToolChain.cs b/Engine/Source/Programs/UnrealBuildTool/HTML5/HTML5ToolChain.cs
--- a/Engine/Source/Programs/UnrealBuildTool/HTML5/HTML5ToolChain.cs
+++ b/Engine/Source/Programs/UnrealBuildTool/HTML5/HTML5ToolChain.cs
@@ -529,7 +529,10 @@ namespace UnrealBuildTool
         public override void AddFilesToManifest(BuildManifest Manifest, UEBuildBinary Binary)
         {
             // we need to include the generated .mem file.  
-            Manifest.AddBuildProduct(Binary.Config.OutputFilePath + ".mem");
+			if(Binary.Config.Type != UEBuildBinaryType.StaticLibrary)
+			{
+	            Manifest.AddBuildProduct(Binary.Config.OutputFilePath + ".mem");
+			}
         }
 
 		public override UnrealTargetPlatform GetPlatform()
