Fix linking against binaries relative to the engine root on Android.

diff --git a/Engine/Source/Programs/UnrealBuildTool/Android/AndroidToolChain.cs b/Engine/Source/Programs/UnrealBuildTool/Android/AndroidToolChain.cs
--- a/Engine/Source/Programs/UnrealBuildTool/Android/AndroidToolChain.cs
+++ b/Engine/Source/Programs/UnrealBuildTool/Android/AndroidToolChain.cs
@@ -917,7 +917,7 @@ namespace UnrealBuildTool
 								else
 								{
 									// full pathed libs are compiled by us, so we depend on linking them
-									LinkAction.CommandArguments += string.Format(" \"{0}\"", AdditionalLibrary);
+									LinkAction.CommandArguments += string.Format(" \"{0}\"", Path.GetFullPath(AdditionalLibrary));
 									LinkAction.PrerequisiteItems.Add(FileItem.GetItemByPath(AdditionalLibrary));
 								}
 							}
