Don't add APKs to manifest for static libraries.

diff --git a/Engine/Source/Programs/UnrealBuildTool/Android/AndroidToolChain.cs b/Engine/Source/Programs/UnrealBuildTool/Android/AndroidToolChain.cs
--- a/Engine/Source/Programs/UnrealBuildTool/Android/AndroidToolChain.cs
+++ b/Engine/Source/Programs/UnrealBuildTool/Android/AndroidToolChain.cs
@@ -940,10 +940,13 @@ namespace UnrealBuildTool
 		public override void AddFilesToManifest(BuildManifest Manifest, UEBuildBinary Binary)
 		{
 			// the binary will have all of the .so's in the output files, we need to trim down to the shared apk (which is what needs to go into the manifest)
-			foreach (string BinaryPath in Binary.Config.OutputFilePaths)
+			if (Binary.Config.Type != UEBuildBinaryType.StaticLibrary)
 			{
-				string ApkFile = Path.ChangeExtension(BinaryPath, ".apk");
-				Manifest.AddBuildProduct(ApkFile);
+				foreach (string BinaryPath in Binary.Config.OutputFilePaths)
+				{
+					string ApkFile = Path.ChangeExtension(BinaryPath, ".apk");
+					Manifest.AddBuildProduct(ApkFile);
+				}
 			}
 		}
 
