Add a "InShippingBuildsOnly" setting for code optimization, so it's possible to disable optimization for engine modules in Development/DebugGame configurations.

diff --git a/Engine/Source/Programs/UnrealBuildTool/System/RulesCompiler.cs b/Engine/Source/Programs/UnrealBuildTool/System/RulesCompiler.cs
--- a/Engine/Source/Programs/UnrealBuildTool/System/RulesCompiler.cs
+++ b/Engine/Source/Programs/UnrealBuildTool/System/RulesCompiler.cs
@@ -105,6 +105,9 @@ namespace UnrealBuildTool
 			/// Code should only be optimized in non-debug builds (not in Debug).
 			InNonDebugBuilds,
 
+			/// Code should only be optimized in shipping builds (not in Debug, DebugGame, Development)
+			InShippingBuildsOnly,
+
 			/// Code should always be optimized if possible.
 			Always,
 
diff --git a/Engine/Source/Programs/UnrealBuildTool/Windows/VCToolChain.cs b/Engine/Source/Programs/UnrealBuildTool/Windows/VCToolChain.cs
--- a/Engine/Source/Programs/UnrealBuildTool/Windows/VCToolChain.cs
+++ b/Engine/Source/Programs/UnrealBuildTool/Windows/VCToolChain.cs
@@ -274,21 +274,32 @@ namespace UnrealBuildTool
 			{
 				if( WindowsPlatform.bUseVCCompilerArgs )
 				{ 
-					// Maximum optimizations if desired.
-					if( CompileEnvironment.Config.OptimizeCode >= ModuleRules.CodeOptimization.InNonDebugBuilds )
+					if(CompileEnvironment.Config.OptimizeCode == ModuleRules.CodeOptimization.InShippingBuildsOnly && CompileEnvironment.Config.Target.Configuration != CPPTargetConfiguration.Shipping)
 					{
-						Arguments.Append(" /Ox");
+						// Disable compiler optimization.
+						Arguments.Append(" /Od");
+
+						// Favor code size (especially useful for embedded platforms).
+						Arguments.Append(" /Os");
 					}
+					else
+					{
+						// Maximum optimizations if desired.
+						if( CompileEnvironment.Config.OptimizeCode >= ModuleRules.CodeOptimization.InNonDebugBuilds )
+						{
+							Arguments.Append(" /Ox");
+						}
 				
-					// Favor code speed.
-					Arguments.Append(" /Ot");
+						// Favor code speed.
+						Arguments.Append(" /Ot");
 
-					// Only omit frame pointers on the PC (which is implied by /Ox) if wanted.
-					if ( BuildConfiguration.bOmitFramePointers == false
-					&& ((CompileEnvironment.Config.Target.Platform == CPPTargetPlatform.Win32) ||
-						(CompileEnvironment.Config.Target.Platform == CPPTargetPlatform.Win64)))
-					{
-						Arguments.Append(" /Oy-");
+						// Only omit frame pointers on the PC (which is implied by /Ox) if wanted.
+						if ( BuildConfiguration.bOmitFramePointers == false
+						&& ((CompileEnvironment.Config.Target.Platform == CPPTargetPlatform.Win32) ||
+							(CompileEnvironment.Config.Target.Platform == CPPTargetPlatform.Win64)))
+						{
+							Arguments.Append(" /Oy-");
+						}
 					}
 
 					// Allow inline method expansion
@@ -308,10 +319,13 @@ namespace UnrealBuildTool
 				}
 				else
 				{
-					// Maximum optimizations if desired.
-					if( CompileEnvironment.Config.OptimizeCode >= ModuleRules.CodeOptimization.InNonDebugBuilds )
+					if(CompileEnvironment.Config.OptimizeCode != ModuleRules.CodeOptimization.InShippingBuildsOnly || CompileEnvironment.Config.Target.Configuration == CPPTargetConfiguration.Shipping)
 					{
-						Arguments.Append( " -O3");
+						// Maximum optimizations if desired.
+						if( CompileEnvironment.Config.OptimizeCode >= ModuleRules.CodeOptimization.InNonDebugBuilds )
+						{
+							Arguments.Append( " -O3");
+						}
 					}
 				}
 			}
