Switch UBT to use clang 3.6.0 Linux cross-toolchain.

-  -Wno-inconsistent-missing-override suppressed, pending discussion with the Core team.

diff --git a/Engine/Source/Programs/UnrealBuildTool/Linux/LinuxToolChain.cs b/Engine/Source/Programs/UnrealBuildTool/Linux/LinuxToolChain.cs
--- a/Engine/Source/Programs/UnrealBuildTool/Linux/LinuxToolChain.cs
+++ b/Engine/Source/Programs/UnrealBuildTool/Linux/LinuxToolChain.cs
@@ -319,6 +319,12 @@ namespace UnrealBuildTool
 				{
 					Result += " -Wno-undefined-bool-conversion";	// hides checking if 'this' pointer is null
 				}
+
+				if (CompilerVersionGreaterOrEqual(3, 6, 0))
+				{
+					Result += " -Wno-inconsistent-missing-override";	// we may still want to fix those - tracked as UE-12341
+					Result += " -Wno-unused-local-typedef";	// clang is being overly strict here? PhysX headers trigger this.
+				}
             }
 
             Result += " -Wno-unused-variable";
diff --git a/Engine/Source/Programs/UnrealBuildTool/Linux/UEBuildLinux.cs b/Engine/Source/Programs/UnrealBuildTool/Linux/UEBuildLinux.cs
--- a/Engine/Source/Programs/UnrealBuildTool/Linux/UEBuildLinux.cs
+++ b/Engine/Source/Programs/UnrealBuildTool/Linux/UEBuildLinux.cs
@@ -13,7 +13,7 @@ namespace UnrealBuildTool
         /** This is the SDK version we support */
 		static private Dictionary<string, string> ExpectedSDKVersions = new Dictionary<string, string>()
 		{
-			{ "x86_64-unknown-linux-gnu", "v4_clang-3.5.0_ld-2.24_glibc-2.12.2" },
+			{ "x86_64-unknown-linux-gnu", "v6_clang-3.6.0_ld-2.24_glibc-2.12.2" },
 			{ "arm-unknown-linux-gnueabihf", "arm-unknown-linux-gnueabihf_v5_clang-3.5.0-ld-2.23.1-glibc-2.13" },
 		};
 
