Replacement method for harvesting environment variables.  Should handle non-ASCII and avoid other problems with set, e.g.:

https://answers.unrealengine.com/questions/193930/vc-directories-not-set-on-start-causing-2-errors.html

diff --git a/Engine/Source/Programs/EnvVarsToXML/EnvVarsToXML/App.config b/Engine/Source/Programs/EnvVarsToXML/EnvVarsToXML/App.config
new file mode 100644
--- /dev/null
+++ b/Engine/Source/Programs/EnvVarsToXML/EnvVarsToXML/App.config
@@ -0,0 +1,6 @@
+﻿<?xml version="1.0" encoding="utf-8" ?>
+<configuration>
+    <startup> 
+        <supportedRuntime version="v4.0" sku=".NETFramework,Version=v4.5" />
+    </startup>
+</configuration>
\ No newline at end of file
diff --git a/Engine/Source/Programs/EnvVarsToXML/EnvVarsToXML/EnvVarsToXML.csproj b/Engine/Source/Programs/EnvVarsToXML/EnvVarsToXML/EnvVarsToXML.csproj
new file mode 100644
--- /dev/null
+++ b/Engine/Source/Programs/EnvVarsToXML/EnvVarsToXML/EnvVarsToXML.csproj
@@ -0,0 +1,55 @@
+﻿<?xml version="1.0" encoding="utf-8"?>
+<Project ToolsVersion="12.0" DefaultTargets="Build" xmlns="http://schemas.microsoft.com/developer/msbuild/2003">
+  <Import Project="$(MSBuildExtensionsPath)\$(MSBuildToolsVersion)\Microsoft.Common.props" Condition="Exists('$(MSBuildExtensionsPath)\$(MSBuildToolsVersion)\Microsoft.Common.props')" />
+  <PropertyGroup>
+    <Configuration Condition=" '$(Configuration)' == '' ">Debug</Configuration>
+    <Platform Condition=" '$(Platform)' == '' ">AnyCPU</Platform>
+    <ProjectGuid>{66E423EB-FFA3-4CE4-BCA3-75BDB6A8318D}</ProjectGuid>
+    <OutputType>Exe</OutputType>
+    <AppDesignerFolder>Properties</AppDesignerFolder>
+    <RootNamespace>EnvVarsToXML</RootNamespace>
+    <AssemblyName>EnvVarsToXML</AssemblyName>
+    <TargetFrameworkVersion>v4.5</TargetFrameworkVersion>
+    <FileAlignment>512</FileAlignment>
+  </PropertyGroup>
+  <PropertyGroup Condition=" '$(Configuration)|$(Platform)' == 'Debug|AnyCPU' ">
+    <PlatformTarget>AnyCPU</PlatformTarget>
+    <DebugSymbols>true</DebugSymbols>
+    <DebugType>full</DebugType>
+    <Optimize>false</Optimize>
+    <OutputPath>..\..\..\Binaries\DotNET\</OutputPath>
+    <DefineConstants>DEBUG;TRACE</DefineConstants>
+    <ErrorReport>prompt</ErrorReport>
+    <WarningLevel>4</WarningLevel>
+  </PropertyGroup>
+  <PropertyGroup Condition=" '$(Configuration)|$(Platform)' == 'Release|AnyCPU' ">
+    <PlatformTarget>AnyCPU</PlatformTarget>
+    <DebugType>pdbonly</DebugType>
+    <Optimize>true</Optimize>
+    <OutputPath>..\..\..\Binaries\DotNET\</OutputPath>
+    <DefineConstants>TRACE</DefineConstants>
+    <ErrorReport>prompt</ErrorReport>
+    <WarningLevel>4</WarningLevel>
+  </PropertyGroup>
+  <ItemGroup>
+    <Reference Include="System" />
+    <Reference Include="System.Core" />
+    <Reference Include="Microsoft.CSharp" />
+    <Reference Include="System.XML" />
+  </ItemGroup>
+  <ItemGroup>
+    <Compile Include="Program.cs" />
+    <Compile Include="Properties\AssemblyInfo.cs" />
+  </ItemGroup>
+  <ItemGroup>
+    <None Include="App.config" />
+  </ItemGroup>
+  <Import Project="$(MSBuildToolsPath)\Microsoft.CSharp.targets" />
+  <!-- To modify your build process, add your task inside one of the targets below and uncomment it. 
+       Other similar extension points exist, see Microsoft.Common.targets.
+  <Target Name="BeforeBuild">
+  </Target>
+  <Target Name="AfterBuild">
+  </Target>
+  -->
+</Project>
\ No newline at end of file
diff --git a/Engine/Source/Programs/EnvVarsToXML/EnvVarsToXML/Program.cs b/Engine/Source/Programs/EnvVarsToXML/EnvVarsToXML/Program.cs
new file mode 100644
--- /dev/null
+++ b/Engine/Source/Programs/EnvVarsToXML/EnvVarsToXML/Program.cs
@@ -0,0 +1,49 @@
+﻿using System;
+using System.Collections;
+using System.Collections.Generic;
+using System.IO;
+using System.Xml;
+using System.Xml.Serialization;
+using System.Text;
+
+namespace EnvVarsToXML
+{
+	[Serializable]
+	public struct EnvVar
+	{
+		[XmlAttribute("Key")]
+		public string Key;
+
+		[XmlAttribute("Value")]
+		public string Value;
+	}
+
+	class Program
+	{
+		static void Main(string[] args)
+		{
+			if (args.Length != 1)
+			{
+				Console.WriteLine("Dumps the environment variables to an XML file.");
+				Console.WriteLine("");
+				Console.WriteLine("EnvVarsToXML <filename>");
+				return;
+			}
+
+			var OutputFilename = args[0];
+
+			var ConvertedEnvVars = new List<EnvVar>();
+			var EnvVars = Environment.GetEnvironmentVariables();
+			foreach (DictionaryEntry  Var in EnvVars)
+			{
+				ConvertedEnvVars.Add(new EnvVar{ Key = (string)Var.Key, Value = (string)Var.Value });
+			}
+
+			var Serializer = new XmlSerializer(typeof(List<EnvVar>));
+			using (var Stream = new StreamWriter(OutputFilename, false, Encoding.Unicode))
+			{
+				Serializer.Serialize(Stream, ConvertedEnvVars);
+			}
+		}
+	}
+}
diff --git a/Engine/Source/Programs/EnvVarsToXML/EnvVarsToXML/Properties/AssemblyInfo.cs b/Engine/Source/Programs/EnvVarsToXML/EnvVarsToXML/Properties/AssemblyInfo.cs
new file mode 100644
--- /dev/null
+++ b/Engine/Source/Programs/EnvVarsToXML/EnvVarsToXML/Properties/AssemblyInfo.cs
@@ -0,0 +1,36 @@
+﻿using System.Reflection;
+using System.Runtime.CompilerServices;
+using System.Runtime.InteropServices;
+
+// General Information about an assembly is controlled through the following 
+// set of attributes. Change these attribute values to modify the information
+// associated with an assembly.
+[assembly: AssemblyTitle("EnvVarsToXML")]
+[assembly: AssemblyDescription("")]
+[assembly: AssemblyConfiguration("")]
+[assembly: AssemblyCompany("")]
+[assembly: AssemblyProduct("EnvVarsToXML")]
+[assembly: AssemblyCopyright("Copyright ©  2015")]
+[assembly: AssemblyTrademark("")]
+[assembly: AssemblyCulture("")]
+
+// Setting ComVisible to false makes the types in this assembly not visible 
+// to COM components.  If you need to access a type in this assembly from 
+// COM, set the ComVisible attribute to true on that type.
+[assembly: ComVisible(false)]
+
+// The following GUID is for the ID of the typelib if this project is exposed to COM
+[assembly: Guid("f929e0dc-50b8-452d-b9f9-84b4b6bded97")]
+
+// Version information for an assembly consists of the following four values:
+//
+//      Major Version
+//      Minor Version 
+//      Build Number
+//      Revision
+//
+// You can specify all the values or you can default the Build and Revision Numbers 
+// by using the '*' as shown below:
+// [assembly: AssemblyVersion("1.0.*")]
+[assembly: AssemblyVersion("1.0.0.0")]
+[assembly: AssemblyFileVersion("1.0.0.0")]
diff --git a/Engine/Source/Programs/UnrealBuildTool/System/Utils.cs b/Engine/Source/Programs/UnrealBuildTool/System/Utils.cs
--- a/Engine/Source/Programs/UnrealBuildTool/System/Utils.cs
+++ b/Engine/Source/Programs/UnrealBuildTool/System/Utils.cs
@@ -162,6 +162,16 @@ namespace UnrealBuildTool
 			return Default;
 		}
 
+		[Serializable]
+		public struct EnvVar
+		{
+			[XmlAttribute("Key")]
+			public string Key;
+
+			[XmlAttribute("Value")]
+			public string Value;
+		}
+
 		/**
 		 * Sets the environment variables from the passed in batch file
 		 * 
@@ -183,8 +193,11 @@ namespace UnrealBuildTool
 					// Run 'vcvars32.bat' (or similar x64 version) to set environment variables
 					EnvReaderBatchFileContent.Add( String.Format( "call \"{0}\"", BatchFileName ) );
 
-					// Pipe all environment variables to a file where we can read them in
-					EnvReaderBatchFileContent.Add( String.Format( "set >\"{0}\"", EnvOutputFileName ) );
+					// Pipe all environment variables to a file where we can read them in.
+					// We use a separate executable which runs after the batch file because we want to capture
+					// the environment after it has been set, and there's no easy way of doing this, and parsing
+					// the output of the set command is problematic when the vars contain non-ASCII characters.
+					EnvReaderBatchFileContent.Add( String.Format( Path.Combine( GetExecutingAssemblyDirectory(), "EnvVarsToXML.exe" ) + " \"{0}\"", EnvOutputFileName ) );
 
 					ResponseFile.Create( EnvReaderBatchFileName, EnvReaderBatchFileContent );
 				}
@@ -196,6 +209,7 @@ namespace UnrealBuildTool
 				{
 					// Run the batch file using cmd.exe with the /U option, to force Unicode output. Many locales have non-ANSI characters in system paths.
 					var StartInfo = BatchFileProcess.StartInfo;
+
 					StartInfo.FileName = Path.Combine(Environment.SystemDirectory, "cmd.exe");
 					StartInfo.Arguments = String.Format("/U /C \"{0}\"", EnvReaderBatchFileName);
 					StartInfo.CreateNoWindow = true;
@@ -221,23 +235,16 @@ namespace UnrealBuildTool
 					Log.TraceVerbose( "Finished launching {0}.", StartInfo.FileName );
 				}
 
-				// Load environment variables (the file will be encoded without a BOM, so we need to manually specify the encoding)
-				var EnvStringsFromFile = File.ReadAllLines( EnvOutputFileName, Encoding.Unicode ).Where( Line => !String.IsNullOrEmpty( Line ) );
-				foreach( var EnvString in EnvStringsFromFile )
+				List<EnvVar> EnvVars;
+				var Serializer = new XmlSerializer(typeof(List<EnvVar>));
+				using (var Stream = new StreamReader(EnvOutputFileName))
 				{
-					// Parse the environment variable name and value from the string ("name=value")
-					int EqualSignPos = EnvString.IndexOf( '=' );
-					// Guard against strange cases where File.ReadAllLines splits a line in half
-					if (EqualSignPos >= 0)
-					{
-						var EnvironmentVariableName = EnvString.Substring(0, EqualSignPos);
-						var EnvironmentVariableValue = EnvString.Substring(EqualSignPos + 1);
-
-						Log.TraceVerbose("Setting environment variable: {0}={1}", EnvironmentVariableName, EnvironmentVariableValue);
+					EnvVars = (List<EnvVar>)Serializer.Deserialize(Stream);
+				}
 
-						// Set the environment variable
-						Environment.SetEnvironmentVariable(EnvironmentVariableName, EnvironmentVariableValue);
-					}
+				foreach (var EnvVar in EnvVars)
+				{
+					Environment.SetEnvironmentVariable(EnvVar.Key, EnvVar.Value);
 				}
 
 				// Clean up the temporary files we created earlier on, so the temp directory doesn't fill up
