Quotes around the EnvVarsToXML.exe call for when there are spaces in the path.

diff --git a/Engine/Source/Programs/UnrealBuildTool/System/Utils.cs b/Engine/Source/Programs/UnrealBuildTool/System/Utils.cs
--- a/Engine/Source/Programs/UnrealBuildTool/System/Utils.cs
+++ b/Engine/Source/Programs/UnrealBuildTool/System/Utils.cs
@@ -197,7 +197,7 @@ namespace UnrealBuildTool
 					// We use a separate executable which runs after the batch file because we want to capture
 					// the environment after it has been set, and there's no easy way of doing this, and parsing
 					// the output of the set command is problematic when the vars contain non-ASCII characters.
-					EnvReaderBatchFileContent.Add( String.Format( Path.Combine( GetExecutingAssemblyDirectory(), "EnvVarsToXML.exe" ) + " \"{0}\"", EnvOutputFileName ) );
+					EnvReaderBatchFileContent.Add( String.Format( "\"{0}\" \"{1}\"", Path.Combine( GetExecutingAssemblyDirectory(), "EnvVarsToXML.exe" ), EnvOutputFileName ) );
 
 					ResponseFile.Create( EnvReaderBatchFileName, EnvReaderBatchFileContent );
 				}
