Don't throw an exception on non-XML chars.

diff --git a/Engine/Source/Programs/UnrealBuildTool/System/Utils.cs b/Engine/Source/Programs/UnrealBuildTool/System/Utils.cs
--- a/Engine/Source/Programs/UnrealBuildTool/System/Utils.cs
+++ b/Engine/Source/Programs/UnrealBuildTool/System/Utils.cs
@@ -235,11 +235,16 @@ namespace UnrealBuildTool
 					Log.TraceVerbose( "Finished launching {0}.", StartInfo.FileName );
 				}
 
+				// Accept chars which are technically not valid XML - they were written out by XmlSerializer anyway!
+				var Settings = new XmlReaderSettings();
+				Settings.CheckCharacters  = false;
+
 				List<EnvVar> EnvVars;
-				var Serializer = new XmlSerializer(typeof(List<EnvVar>));
 				using (var Stream = new StreamReader(EnvOutputFileName))
+				using (var Reader = XmlReader.Create(Stream, Settings))
 				{
-					EnvVars = (List<EnvVar>)Serializer.Deserialize(Stream);
+					var Serializer = new XmlSerializer(typeof(List<EnvVar>));
+					EnvVars = (List<EnvVar>)Serializer.Deserialize(Reader);
 				}
 
 				foreach (var EnvVar in EnvVars)
