Rolled back Linux toolchain to clang 3.5.0 until Shooter compilation is fixed.

diff --git a/Engine/Source/Programs/UnrealBuildTool/Linux/UEBuildLinux.cs b/Engine/Source/Programs/UnrealBuildTool/Linux/UEBuildLinux.cs
--- a/Engine/Source/Programs/UnrealBuildTool/Linux/UEBuildLinux.cs
+++ b/Engine/Source/Programs/UnrealBuildTool/Linux/UEBuildLinux.cs
@@ -13,7 +13,7 @@ namespace UnrealBuildTool
         /** This is the SDK version we support */
 		static private Dictionary<string, string> ExpectedSDKVersions = new Dictionary<string, string>()
 		{
-			{ "x86_64-unknown-linux-gnu", "v6_clang-3.6.0_ld-2.24_glibc-2.12.2" },
+			{ "x86_64-unknown-linux-gnu", "v4_clang-3.5.0_ld-2.24_glibc-2.12.2" },
 			{ "arm-unknown-linux-gnueabihf", "arm-unknown-linux-gnueabihf_v5_clang-3.5.0-ld-2.23.1-glibc-2.13" },
 		};
 
