EnvVarsToXML is now a dependency of UBT so that it gets built when UBT gets build.

* The EnvVarsToXML project gets added to generated .sln.
* The EnvVarsToXML project is now referenced by UBT.
* The generated .sln shows the correct dependencies in the Configuration Manager.
* EnvVarsToXML now targets .NET 4.0.
* Fixes to output path and config names.

diff --git a/Engine/Source/Programs/EnvVarsToXML/EnvVarsToXML/App.config b/Engine/Source/Programs/EnvVarsToXML/EnvVarsToXML/App.config
--- a/Engine/Source/Programs/EnvVarsToXML/EnvVarsToXML/App.config
+++ b/Engine/Source/Programs/EnvVarsToXML/EnvVarsToXML/App.config
@@ -1,6 +1,6 @@
-﻿<?xml version="1.0" encoding="utf-8" ?>
+<?xml version="1.0" encoding="utf-8"?>
 <configuration>
     <startup> 
-        <supportedRuntime version="v4.0" sku=".NETFramework,Version=v4.5" />
+        <supportedRuntime version="v4.0" sku=".NETFramework,Version=v4.0"/>
     </startup>
-</configuration>
\ No newline at end of file
+</configuration>
diff --git a/Engine/Source/Programs/EnvVarsToXML/EnvVarsToXML/EnvVarsToXML.csproj b/Engine/Source/Programs/EnvVarsToXML/EnvVarsToXML/EnvVarsToXML.csproj
--- a/Engine/Source/Programs/EnvVarsToXML/EnvVarsToXML/EnvVarsToXML.csproj
+++ b/Engine/Source/Programs/EnvVarsToXML/EnvVarsToXML/EnvVarsToXML.csproj
@@ -9,7 +9,7 @@
     <AppDesignerFolder>Properties</AppDesignerFolder>
     <RootNamespace>EnvVarsToXML</RootNamespace>
     <AssemblyName>EnvVarsToXML</AssemblyName>
-    <TargetFrameworkVersion>v4.5</TargetFrameworkVersion>
+    <TargetFrameworkVersion>v4.0</TargetFrameworkVersion>
     <FileAlignment>512</FileAlignment>
   </PropertyGroup>
   <PropertyGroup Condition=" '$(Configuration)|$(Platform)' == 'Debug|AnyCPU' ">
@@ -17,7 +17,7 @@
     <DebugSymbols>true</DebugSymbols>
     <DebugType>full</DebugType>
     <Optimize>false</Optimize>
-    <OutputPath>..\..\..\Binaries\DotNET\</OutputPath>
+    <OutputPath>..\..\..\..\Binaries\DotNET\</OutputPath>
     <DefineConstants>DEBUG;TRACE</DefineConstants>
     <ErrorReport>prompt</ErrorReport>
     <WarningLevel>4</WarningLevel>
@@ -26,7 +26,7 @@
     <PlatformTarget>AnyCPU</PlatformTarget>
     <DebugType>pdbonly</DebugType>
     <Optimize>true</Optimize>
-    <OutputPath>..\..\..\Binaries\DotNET\</OutputPath>
+    <OutputPath>..\..\..\..\Binaries\DotNET\</OutputPath>
     <DefineConstants>TRACE</DefineConstants>
     <ErrorReport>prompt</ErrorReport>
     <WarningLevel>4</WarningLevel>
diff --git a/Engine/Source/Programs/UnrealBuildTool/System/ProjectFileGenerator.cs b/Engine/Source/Programs/UnrealBuildTool/System/ProjectFileGenerator.cs
--- a/Engine/Source/Programs/UnrealBuildTool/System/ProjectFileGenerator.cs
+++ b/Engine/Source/Programs/UnrealBuildTool/System/ProjectFileGenerator.cs
@@ -492,8 +492,12 @@ namespace UnrealBuildTool
 				{
 					MasterProjectFolder ProgramsFolder = RootFolder.AddSubFolder( "Programs" );
 
+					// Add EnvVarsToXML to the master project
+					var EnvVarsToXMLProjectFile = AddSimpleCSharpProject("EnvVarsToXML/EnvVarsToXML", bShouldBuildForAllSolutionTargets: true, bForceDevelopmentConfiguration: true);
+					ProgramsFolder.ChildProjects.Add(EnvVarsToXMLProjectFile);
+
 					// Add UnrealBuildTool to the master project
-					AddUnrealBuildToolProject( ProgramsFolder );
+					AddUnrealBuildToolProject( ProgramsFolder, new ProjectFile[] { EnvVarsToXMLProjectFile } );
 
 					// Add AutomationTool to the master project
 					ProgramsFolder.ChildProjects.Add(AddSimpleCSharpProject("AutomationTool", bShouldBuildForAllSolutionTargets: true, bForceDevelopmentConfiguration: true));
@@ -1157,12 +1161,17 @@ namespace UnrealBuildTool
 		/// <summary>
 		/// Adds UnrealBuildTool to the master project
 		/// </summary>
-		private void AddUnrealBuildToolProject( MasterProjectFolder ProgramsFolder )
+		private void AddUnrealBuildToolProject( MasterProjectFolder ProgramsFolder, IEnumerable<ProjectFile> Dependencies )
 		{
 			var ProjectFileName = Utils.MakePathRelativeTo( Path.Combine( Path.Combine( EngineRelativePath, "Source" ), "Programs", "UnrealBuildTool", "UnrealBuildTool.csproj" ), MasterProjectRelativePath );
 			var UnrealBuildToolProject = new VCSharpProjectFile( ProjectFileName );
 			UnrealBuildToolProject.ShouldBuildForAllSolutionTargets = true;
 
+			foreach (var Dependent in Dependencies)
+			{
+				UnrealBuildToolProject.AddDependsOnProject( Dependent );
+			}
+
 			// Store it off as we need it when generating target projects.
 			UBTProject = UnrealBuildToolProject;
 
diff --git a/Engine/Source/Programs/UnrealBuildTool/System/VCProjectFileGenerator.cs b/Engine/Source/Programs/UnrealBuildTool/System/VCProjectFileGenerator.cs
--- a/Engine/Source/Programs/UnrealBuildTool/System/VCProjectFileGenerator.cs
+++ b/Engine/Source/Programs/UnrealBuildTool/System/VCProjectFileGenerator.cs
@@ -360,29 +360,26 @@ namespace UnrealBuildTool
 					// freshly compiled before kicking off any build operations on this target project
 					if ( !CurProject.IsStubProject )
 					{
-						// Don't add self as a project dependency!
-						if ((CurProject != UBTProject) && (CurProject.IsGeneratedProject || (CurProject.DependsOnProjects.Count > 0)))
+						var Dependencies = new List<ProjectFile>();
+						if (CurProject.IsGeneratedProject && UBTProject != null && CurProject != UBTProject)
 						{
-							VCSolutionFileContent.Append(
-									"	ProjectSection(ProjectDependencies) = postProject" + ProjectFileGenerator.NewLine);
+							Dependencies.Add(UBTProject);
+							Dependencies.AddRange(UBTProject.DependsOnProjects);
+						}
+						Dependencies.AddRange(CurProject.DependsOnProjects);
 
-							if (CurProject.IsGeneratedProject && UBTProject != null)
-							{
-								var UBTProjectGUID = ((MSBuildProjectFile)UBTProject).ProjectGUID.ToString("B").ToUpperInvariant();
-								VCSolutionFileContent.Append(
-										"		" + UBTProjectGUID + " = " + UBTProjectGUID + ProjectFileGenerator.NewLine);
-							}
+						if (Dependencies.Count > 0)
+						{
+							VCSolutionFileContent.Append("\tProjectSection(ProjectDependencies) = postProject" + ProjectFileGenerator.NewLine);
 
 							// Setup any addition dependencies this project has...
-							foreach (var DependsOnProject in CurProject.DependsOnProjects)
+							foreach (var DependsOnProject in Dependencies)
 							{
 								var DependsOnProjectGUID = ((MSBuildProjectFile)DependsOnProject).ProjectGUID.ToString("B").ToUpperInvariant();
-								VCSolutionFileContent.Append(
-									"		" + DependsOnProjectGUID + " = " + DependsOnProjectGUID + ProjectFileGenerator.NewLine);
+								VCSolutionFileContent.Append("\t\t" + DependsOnProjectGUID + " = " + DependsOnProjectGUID + ProjectFileGenerator.NewLine);
 							}
 
-							VCSolutionFileContent.Append(
-									"	EndProjectSection" + ProjectFileGenerator.NewLine);
+							VCSolutionFileContent.Append("\tEndProjectSection" + ProjectFileGenerator.NewLine);
 						}
 					}
 
diff --git a/Engine/Source/Programs/UnrealBuildTool/UnrealBuildTool.csproj b/Engine/Source/Programs/UnrealBuildTool/UnrealBuildTool.csproj
--- a/Engine/Source/Programs/UnrealBuildTool/UnrealBuildTool.csproj
+++ b/Engine/Source/Programs/UnrealBuildTool/UnrealBuildTool.csproj
@@ -123,16 +123,16 @@
     </When>
   </Choose>
   <ItemGroup>
-	<Compile Include="Mac\MacToolChain.cs" />
-	<Compile Include="Mac\UEBuildMac.cs" />
-	<Compile Include="Mac\MacProjectGenerator.cs" />
-	<Compile Include="Mac\UEDeployMac.cs" />
+    <Compile Include="Mac\MacToolChain.cs" />
+    <Compile Include="Mac\UEBuildMac.cs" />
+    <Compile Include="Mac\MacProjectGenerator.cs" />
+    <Compile Include="Mac\UEDeployMac.cs" />
   </ItemGroup>
   <ItemGroup>
-	<Compile Include="Android\UEBuildAndroid.cs" />
-	<Compile Include="Android\UEDeployAndroid.cs" />
-	<Compile Include="Android\AndroidToolChain.cs" />
-	<Compile Include="Android\AndroidProjectGenerator.cs" />
+    <Compile Include="Android\UEBuildAndroid.cs" />
+    <Compile Include="Android\UEDeployAndroid.cs" />
+    <Compile Include="Android\AndroidToolChain.cs" />
+    <Compile Include="Android\AndroidProjectGenerator.cs" />
   </ItemGroup>
   <Choose>
     <When Condition="Exists('WinRT\UEBuildWinRT.cs')">
@@ -145,21 +145,21 @@
     </When>
   </Choose>
   <ItemGroup>
-	<Compile Include="IOS\IOSProjectGenerator.cs" />
-	<Compile Include="IOS\IOSToolChain.cs" />
-	<Compile Include="IOS\UEBuildIOS.cs" />
-	<Compile Include="IOS\UEDeployIOS.cs" />
+    <Compile Include="IOS\IOSProjectGenerator.cs" />
+    <Compile Include="IOS\IOSToolChain.cs" />
+    <Compile Include="IOS\UEBuildIOS.cs" />
+    <Compile Include="IOS\UEDeployIOS.cs" />
   </ItemGroup>
   <ItemGroup>
-	<Compile Include="HTML5\HTML5ProjectGenerator.cs" />
-	<Compile Include="HTML5\HTML5ToolChain.cs" />
-	<Compile Include="HTML5\UEBuildHTML5.cs" />
-	<Compile Include="HTML5\HTML5SDKInfo.cs" />
+    <Compile Include="HTML5\HTML5ProjectGenerator.cs" />
+    <Compile Include="HTML5\HTML5ToolChain.cs" />
+    <Compile Include="HTML5\UEBuildHTML5.cs" />
+    <Compile Include="HTML5\HTML5SDKInfo.cs" />
   </ItemGroup>
   <ItemGroup>
-	<Compile Include="Linux\LinuxProjectGenerator.cs" />
-	<Compile Include="Linux\LinuxToolChain.cs" />
-	<Compile Include="Linux\UEBuildLinux.cs" />
+    <Compile Include="Linux\LinuxProjectGenerator.cs" />
+    <Compile Include="Linux\LinuxToolChain.cs" />
+    <Compile Include="Linux\UEBuildLinux.cs" />
   </ItemGroup>
   <Choose>
     <When Condition="Exists('System\NotForLicensees\TelemetryProviderET.cs')">
@@ -283,6 +283,12 @@
   <ItemGroup>
     <None Include="app.config" />
   </ItemGroup>
+  <ItemGroup>
+    <ProjectReference Include="..\EnvVarsToXML\EnvVarsToXML\EnvVarsToXML.csproj">
+      <Project>{66e423eb-ffa3-4ce4-bca3-75bdb6a8318d}</Project>
+      <Name>EnvVarsToXML</Name>
+    </ProjectReference>
+  </ItemGroup>
   <Import Project="$(MSBuildBinPath)\Microsoft.CSharp.targets" />
   <!-- To modify your build process, add your task inside one of the targets below and uncomment it. 
        Other similar extension points exist, see Microsoft.Common.targets.
@@ -297,4 +303,4 @@
     <PreBuildEvent>
     </PreBuildEvent>
   </PropertyGroup>
-</Project>
+</Project>
\ No newline at end of file
