Disable constructor member init order warning for HTML5

diff --git a/Engine/Source/Programs/UnrealBuildTool/HTML5/HTML5ToolChain.cs b/Engine/Source/Programs/UnrealBuildTool/HTML5/HTML5ToolChain.cs
--- a/Engine/Source/Programs/UnrealBuildTool/HTML5/HTML5ToolChain.cs
+++ b/Engine/Source/Programs/UnrealBuildTool/HTML5/HTML5ToolChain.cs
@@ -63,7 +63,6 @@ namespace UnrealBuildTool
             Result += " -Wno-logical-op-parentheses"; // appErrorf triggers this
             Result += " -Wno-array-bounds"; // some VectorLoads go past the end of the array, but it's okay in that case
             Result += " -Wno-invalid-offsetof"; // too many warnings kills windows clang. 
-            
 
             // JavsScript option overrides (see src/settings.js)
 
@@ -107,6 +106,7 @@ namespace UnrealBuildTool
 							}*/
 
 				Result += " -Wno-warn-absolute-paths ";
+				Result += " -Wno-reorder"; // we disable constructor order warnings.
 
 				if (CompileEnvironment.Config.Target.Configuration == CPPTargetConfiguration.Debug)
 				{
