Added checks on distribution level for modules and the direct dependencies of binaries
Changed Module Manager so that it adds restricted directories to modified PATH if they exist, removed unnecessary recursive use of FindModulePathsInDirectory
Updated setup of any projects that were flagged up using this system
Added old locations of restricted output to junk manifest
Added Util function to get min/max of IComparable types (used for enums)

diff --git a/Engine/Build/JunkManifest.txt b/Engine/Build/JunkManifest.txt
--- a/Engine/Build/JunkManifest.txt
+++ b/Engine/Build/JunkManifest.txt
@@ -36,10 +36,13 @@ Engine\Binaries\Mac\UE4Editor.app\Contents\MacOS\*-CameraDebugger.*
 Engine\Binaries\Mac\UE4Editor.app\Contents\MacOS\*-CameraEditor.*
 Engine\Binaries\Mac\UE4Editor.app\Contents\MacOS\*-SessionLauncher.*
 Engine\Binaries\Mac\*-MessagingDebugger.*
-Engine\Binaries\Win64\*-DDCUtils.*
+Engine\Binaries\Win64\*-DDCUtils*.*
 Engine\Binaries\Win64\*-MessagingDebugger.*
+Engine\Binaries\Win64\*-OnlineSubsystemMcp*.*
+Engine\Binaries\Win64\*-SimplygonMeshReduction*.*
 Engine\Binaries\Win64\*-TextureXboxOneFormat.*
 Engine\Binaries\Win64\*-TextureFormatATITC.*
+Engine\Binaries\Win64\*-XMPP*.*
 Shadow\Binaries\Win64\*ShadowLoadingScreen.*
 Engine\Binaries\Win64\UnrealRemoteAgent.*
 Engine\Binaries\DotNET\AutomationScripts\NoRedist\FotniteGame.Automation.*
@@ -55,6 +58,10 @@ Engine\Plugins\Script\ScriptPlugin\Binaries\Win64\*
 Engine\Plugins\Experimental\SteamVR\Binaries\Win64\*
 Engine\Plugins\Experimental\Paper2D\Binaries
 Engine\Plugins\Experimental\Paper2D\Intermediate
+Engine\Plugins\Runtime\TwitchLiveStreaming\Binaries\Win32\TwitchLiveStreaming*.*
+Engine\Plugins\Runtime\TwitchLiveStreaming\Binaries\Win32\*-TwitchLiveStreaming*.*
+Engine\Plugins\Runtime\TwitchLiveStreaming\Binaries\Win64\TwitchLiveStreaming*.*
+Engine\Plugins\Runtime\TwitchLiveStreaming\Binaries\Win64\*-TwitchLiveStreaming*.*
 Engine\Binaries\Mac\UE4Editor.app\Contents\MacOS\UE4Editor-Android_ES31TargetPlatform.dylib
 Engine\Binaries\Mac\UnrealFrontend.app\Contents\MacOS\UnrealFrontend-Android_ES31TargetPlatform.dylib
 Engine\Binaries\Win64\Android\*Android_ES31*.*
diff --git a/Engine/Source/Developer/SimplygonMeshReduction/SimplygonMeshReduction.Build.cs b/Engine/Source/Developer/SimplygonMeshReduction/SimplygonMeshReduction.Build.cs
--- a/Engine/Source/Developer/SimplygonMeshReduction/SimplygonMeshReduction.Build.cs
+++ b/Engine/Source/Developer/SimplygonMeshReduction/SimplygonMeshReduction.Build.cs
@@ -7,6 +7,8 @@ public class SimplygonMeshReduction : ModuleRules
 {
 	public SimplygonMeshReduction(TargetInfo Target)
 	{
+		BinariesSubFolder = "NotForLicensees";
+
 		PublicIncludePaths.Add("Developer/SimplygonMeshReduction/Public");
 		PrivateDependencyModuleNames.Add("Core");
 		PrivateDependencyModuleNames.Add("CoreUObject");
diff --git a/Engine/Source/Programs/BuildPatchTool/BuildPatchTool.Target.cs b/Engine/Source/Programs/BuildPatchTool/BuildPatchTool.Target.cs
--- a/Engine/Source/Programs/BuildPatchTool/BuildPatchTool.Target.cs
+++ b/Engine/Source/Programs/BuildPatchTool/BuildPatchTool.Target.cs
@@ -8,6 +8,7 @@ public class BuildPatchToolTarget : TargetRules
 	public BuildPatchToolTarget( TargetInfo Target )
 	{
 		Type = TargetType.Program;
+		bOutputPubliclyDistributable = true;
 	}
 
 	//
diff --git a/Engine/Source/Programs/UnrealBuildTool/Configuration/UEBuildBinary.cs b/Engine/Source/Programs/UnrealBuildTool/Configuration/UEBuildBinary.cs
--- a/Engine/Source/Programs/UnrealBuildTool/Configuration/UEBuildBinary.cs
+++ b/Engine/Source/Programs/UnrealBuildTool/Configuration/UEBuildBinary.cs
@@ -41,7 +41,7 @@ namespace UnrealBuildTool
 		public string[] OutputFilePaths;
 
 		/// <summary>
-		/// Returns the OutputFilePath is there is only one entry in OutputFilePaths
+		/// Returns the OutputFilePath if there is only one entry in OutputFilePaths
 		/// </summary>
 		public string OutputFilePath
 		{
@@ -61,17 +61,17 @@ namespace UnrealBuildTool
 		public string[] OriginalOutputFilePaths;
 
 		/// <summary>
-		/// Returns the OutputFilePath is there is only one entry in OutputFilePaths
+		/// Returns the OriginalOutputFilePath if there is only one entry in OriginalOutputFilePaths
 		/// </summary>
 		public string OriginalOutputFilePath
 		{
 			get
 			{
-				if (OutputFilePaths.Length != 1)
+				if (OriginalOutputFilePaths.Length != 1)
 				{
-					throw new BuildException("Attempted to use UEBuildBinaryConfiguration.OutputFilePath property, but there are multiple (or no) OutputFilePaths. You need to handle multiple in the code that called this (size = {0})", OutputFilePaths.Length);
+					throw new BuildException("Attempted to use UEBuildBinaryConfiguration.OriginalOutputFilePath property, but there are multiple (or no) OriginalOutputFilePaths. You need to handle multiple in the code that called this (size = {0})", OriginalOutputFilePaths.Length);
 				}
-				return OutputFilePaths[0];
+				return OriginalOutputFilePaths[0];
 			}
 		}
 
@@ -304,6 +304,51 @@ namespace UnrealBuildTool
 		{
 			return Path.Combine(Path.GetDirectoryName(BinaryPath), Path.GetFileNameWithoutExtension(BinaryPath) + "-Cmd" + Path.GetExtension(BinaryPath));
 		}
+
+		/**
+		 * Checks whether the binary output paths are appropriate for the distribution
+		 * level of its direct module dependencies
+		 */
+		public void CheckOutputDistributionLevelAgainstDependencies()
+		{
+			// Find maximum distribution level of its direct dependencies
+			var DistributionLevel = UEBuildModuleDistribution.Public;
+			var DependantModules = GetAllDependencyModules(false, false);
+			List<string>[] DependantModuleNames = new List<string>[Enum.GetNames(typeof(UEBuildModuleDistribution)).Length];
+			foreach (var Module in DependantModules)
+			{
+				if (Module.DistributionLevel != UEBuildModuleDistribution.Public)
+				{
+					// Make a list of non-public dependant modules so that exception
+					// message can be more helpful
+					int DistributionIndex = (int)Module.DistributionLevel;
+					if (DependantModuleNames[DistributionIndex] == null)
+					{
+						DependantModuleNames[DistributionIndex] = new List<string>();
+					}
+					DependantModuleNames[DistributionIndex].Add(Module.Name);
+
+					DistributionLevel = Utils.Max(DistributionLevel, Module.DistributionLevel);
+				}
+			}
+
+			// Check Output Paths if dependencies shouldn't be distributed to everyone
+			if (DistributionLevel != UEBuildModuleDistribution.Public)
+			{
+				foreach (var OutputFilePath in Config.OutputFilePaths)
+				{
+					var OutputDistributionLevel = UEBuildModule.GetModuleDistributionLevelBasedOnLocation(OutputFilePath);
+
+					// Throw exception if output path is not appropriate
+					if (OutputDistributionLevel < DistributionLevel)
+					{
+						var JoinedModuleNames = String.Join(",", DependantModuleNames[(int)DistributionLevel]);
+						throw new BuildException("Output file \"{0}\" has distribution level of \"{1}\" but has direct dependencies on modules with distribution level of \"{2}\" ({3}).\nEither change to dynamic dependencies, set BinariesSubFolder/ExeBinariesSubFolder to \"{2}\" or set bOutputPubliclyDistributable to true in the target.cs file.",
+							OutputFilePath, OutputDistributionLevel.ToString(), DistributionLevel.ToString(), JoinedModuleNames);
+					}
+				}
+			}
+		}
 	};
 
 	/// <summary>
@@ -758,7 +803,7 @@ namespace UnrealBuildTool
 
 			Action LinkAction = new Action(ActionType.Compile);
 			LinkAction.WorkingDirectory = Path.GetFullPath(".");
-			LinkAction.CommandPath = System.IO.Path.Combine(LinkAction.WorkingDirectory, @"..", @"Binaries", @"Win32", @"UnrealCodeAnalyzer.exe");
+			LinkAction.CommandPath = System.IO.Path.Combine(LinkAction.WorkingDirectory, @"..", @"Binaries", @"Win32", @"NotForLicensees", @"UnrealCodeAnalyzer.exe");
 			LinkAction.bIsVCCompiler = false;
 			LinkAction.ProducedItems.Add(OutputFile);
 			LinkAction.PrerequisiteItems.AddRange(BinaryLinkEnvironment.InputFiles);
diff --git a/Engine/Source/Programs/UnrealBuildTool/Configuration/UEBuildModule.cs b/Engine/Source/Programs/UnrealBuildTool/Configuration/UEBuildModule.cs
--- a/Engine/Source/Programs/UnrealBuildTool/Configuration/UEBuildModule.cs
+++ b/Engine/Source/Programs/UnrealBuildTool/Configuration/UEBuildModule.cs
@@ -29,6 +29,25 @@ namespace UnrealBuildTool
 			return ModuleType != UEBuildModuleType.Game;
 		}
 	}
+	/**
+	 * Distribution level of module.
+	 * Note: The name of each entry is used to search for/create folders
+	 */
+	public enum UEBuildModuleDistribution
+	{
+		/// Binaries can be distributed to everyone
+		Public,
+
+		/// Can be used by UE4 but not required
+		CarefullyRedist,
+
+		/// Epic Employees and Contractors
+		NotForLicensees,
+
+		/// Epic Employees only
+		NoRedist,
+	}
+
 	
 	/** A unit of code compilation and linking. */
 	public abstract class UEBuildModule
@@ -98,6 +117,61 @@ namespace UnrealBuildTool
 			return ModuleType;
 		}
 
+		/**
+		 * Checks what distribution level a particular path should have by checking for key folders anywhere in it
+		 */
+		public static UEBuildModuleDistribution GetModuleDistributionLevelBasedOnLocation(string FilePath)
+		{
+			// Get full path to ensure all folder separators are the same
+			FilePath = Path.GetFullPath(FilePath);
+
+			// Check from highest to lowest (don't actually need to check 'Everyone')
+			for (var DistributionLevel = UEBuildModuleDistribution.NoRedist; DistributionLevel > UEBuildModuleDistribution.Public; DistributionLevel--)
+			{
+				var DistributionFolderName = String.Format("{0}{1}{0}", Path.DirectorySeparatorChar, DistributionLevel.ToString());
+				if (FilePath.IndexOf(DistributionFolderName, StringComparison.InvariantCultureIgnoreCase) >= 0)
+				{
+					return DistributionLevel;
+				}
+
+				if (DistributionLevel == UEBuildModuleDistribution.NotForLicensees)
+				{
+					// Extra checks for PS4 and XboxOne folders, which are equivalent to NotForLicensees
+					var PS4FolderName = String.Format("{0}ps4{0}", Path.DirectorySeparatorChar);
+					var XboxFolderName = String.Format("{0}xboxone{0}", Path.DirectorySeparatorChar);
+					if (FilePath.IndexOf(PS4FolderName, StringComparison.InvariantCultureIgnoreCase) >= 0
+					|| FilePath.IndexOf(XboxFolderName, StringComparison.InvariantCultureIgnoreCase) >= 0)
+					{
+						return UEBuildModuleDistribution.NotForLicensees;
+					}
+				}
+			}
+
+			return UEBuildModuleDistribution.Public;
+		}
+
+		/**
+		 * Determines the distribution level of a module based on its directory and includes.
+		 */
+		private void SetupModuleDistributionLevel()
+		{
+			List<string> PathsToCheck = new List<string>();
+			PathsToCheck.Add(ModuleDirectory);
+			PathsToCheck.AddRange(PublicIncludePaths);
+			PathsToCheck.AddRange(PrivateIncludePaths);
+			// Not sure if these two are necessary as paths will usually be in basic includes too
+			PathsToCheck.AddRange(PublicSystemIncludePaths);
+			PathsToCheck.AddRange(PublicLibraryPaths);
+
+			DistributionLevel = UEBuildModuleDistribution.Public;
+
+			// Keep checking as long as we haven't reached the maximum level
+			for (int PathIndex = 0; PathIndex < PathsToCheck.Count && DistributionLevel != UEBuildModuleDistribution.NoRedist; ++PathIndex)
+			{
+				DistributionLevel = Utils.Max(DistributionLevel, UEBuildModule.GetModuleDistributionLevelBasedOnLocation(PathsToCheck[PathIndex]));
+			}
+		}
+
 		/** Converts an optional string list parameter to a well-defined hash set. */
 		protected static HashSet<string> HashSetFromOptionalEnumerableStringParameter(IEnumerable<string> InEnumerableStrings)
 		{
@@ -113,6 +187,9 @@ namespace UnrealBuildTool
 		/** The type of module being built. Used to switch between debug/development and precompiled/source configurations. */
 		public UEBuildModuleType Type;
 
+		/** The distribution level of the module being built. Used to check where build products are placed. */
+		public UEBuildModuleDistribution DistributionLevel;
+
 		/** Path to the module directory */
 		public readonly string ModuleDirectory;
 
@@ -234,6 +311,8 @@ namespace UnrealBuildTool
             PlatformSpecificDynamicallyLoadedModuleNames = HashSetFromOptionalEnumerableStringParameter(InPlatformSpecificDynamicallyLoadedModuleNames);
 			IsRedistributableOverride = InIsRedistributableOverride;
 
+			SetupModuleDistributionLevel();
+
 			Target.RegisterModule(this);
 		}
 
diff --git a/Engine/Source/Programs/UnrealBuildTool/Configuration/UEBuildTarget.cs b/Engine/Source/Programs/UnrealBuildTool/Configuration/UEBuildTarget.cs
--- a/Engine/Source/Programs/UnrealBuildTool/Configuration/UEBuildTarget.cs
+++ b/Engine/Source/Programs/UnrealBuildTool/Configuration/UEBuildTarget.cs
@@ -1785,6 +1785,15 @@ namespace UnrealBuildTool
 				}
 			}
 
+			if (ShouldCheckOutputDistributionLevel())
+			{
+				// Check the distribution level of all binaries based on the dependencies they have
+				foreach (var Binary in AppBinaries)
+				{
+					Binary.CheckOutputDistributionLevelAgainstDependencies();
+				}
+			}
+
 			return SpecialRocketLibFilesThatAreBuildProducts;
 		}
 
@@ -2456,7 +2465,7 @@ namespace UnrealBuildTool
 			}
 			
 			string ModuleBinariesSubDir = "";
-            if (BinaryType == UEBuildBinaryType.DynamicLinkLibrary && (string.IsNullOrEmpty(ModuleName) == false))
+			if ((BinaryType == UEBuildBinaryType.DynamicLinkLibrary || BinaryType == UEBuildBinaryType.StaticLibrary) && (string.IsNullOrEmpty(ModuleName) == false))
 			{
 				// Allow for modules to specify sub-folders in the Binaries folder
 				var RulesFilename = RulesCompiler.GetModuleFilename(ModuleName);
@@ -3355,5 +3364,23 @@ namespace UnrealBuildTool
 
 			return FilteredFileItems;
 		}
+
+		/// <summary>
+		/// Determines whether we should check binary output paths for this target are
+		/// appropriate for the distribution level of their direct module dependencies
+		/// </summary>
+		private bool ShouldCheckOutputDistributionLevel()
+		{
+			if (Rules != null)
+			{
+				// We don't want to check this for true game projects (but do want to include UE4Game)
+				bool bIsUProjectGame = UnrealBuildTool.HasUProjectFile() && Rules.Type == TargetRules.TargetType.Game;
+				if (!bIsUProjectGame && !Rules.bOutputPubliclyDistributable)
+				{
+					return true;
+				}
+			}
+			return false;
+		}
 	}
 }
diff --git a/Engine/Source/Programs/UnrealBuildTool/System/RulesCompiler.cs b/Engine/Source/Programs/UnrealBuildTool/System/RulesCompiler.cs
--- a/Engine/Source/Programs/UnrealBuildTool/System/RulesCompiler.cs
+++ b/Engine/Source/Programs/UnrealBuildTool/System/RulesCompiler.cs
@@ -431,6 +431,13 @@ namespace UnrealBuildTool
         /// </summary>
         public bool bDebugBuildsActuallyUseDebugCRT = false;
 
+		/// <summary>
+		/// Whether the output from this target can be publicly distributed, even if it has
+		/// dependencies on modules that are not (i.e. CarefullyRedist, NotForLicensees, NoRedist).
+		/// This should be used when you plan to release binaries but not source.
+		/// </summary>
+		public bool bOutputPubliclyDistributable = false;
+
 		/// <summary>
 		/// A list of additional plugins which need to be built for this target. Game and editor targets can use the EnabledPlugins 
 		/// setting in their config files to control this.
diff --git a/Engine/Source/Programs/UnrealBuildTool/System/Utils.cs b/Engine/Source/Programs/UnrealBuildTool/System/Utils.cs
--- a/Engine/Source/Programs/UnrealBuildTool/System/Utils.cs
+++ b/Engine/Source/Programs/UnrealBuildTool/System/Utils.cs
@@ -894,6 +894,18 @@ namespace UnrealBuildTool
 				return Environment.GetFolderPath(Environment.SpecialFolder.LocalApplicationData);
 			}
 		}
+
+		/** Helper function to get the minimum of two comparable objects */
+		public static T Min<T>(T a, T b) where T : IComparable
+		{
+			return a.CompareTo(b) < 0 ? a : b;
+		}
+
+		/** Helper function to get the maximum of two comparable objects */
+		public static T Max<T>(T a, T b) where T : IComparable
+		{
+			return a.CompareTo(b) > 0 ? a : b;
+		}
 	}
 
 	/// <summary>
diff --git a/Engine/Source/Runtime/Core/Private/Modules/ModuleManager.cpp b/Engine/Source/Runtime/Core/Private/Modules/ModuleManager.cpp
--- a/Engine/Source/Runtime/Core/Private/Modules/ModuleManager.cpp
+++ b/Engine/Source/Runtime/Core/Private/Modules/ModuleManager.cpp
@@ -33,6 +33,18 @@ FModuleManager& FModuleManager::Get()
 	if( ModuleManager == NULL )
 	{
 		ModuleManager = new FModuleManager();
+
+		// Ensure that dependency dlls can be found in restricted sub directories
+		const TCHAR* RestrictedFolderNames[] = { TEXT("NoRedist"), TEXT("NotForLicensees"), TEXT("CarefullyRedist") };
+		FString ModuleDir = FPlatformProcess::GetModulesDirectory();
+		for (const TCHAR* RestrictedFolderName : RestrictedFolderNames)
+		{
+			FString RestrictedFolder = ModuleDir / RestrictedFolderName;
+			if (FPaths::DirectoryExists(RestrictedFolder))
+			{
+				ModuleManager->AddBinariesDirectory(*RestrictedFolder, false);
+			}
+		}
 	}
 
 	return *ModuleManager;
@@ -866,17 +878,6 @@ void FModuleManager::FindModulePathsInDirectory(const FString& InDirectoryName,
 			}
 		}
 	}
-
-	// Also recurse into restricted sub-folders, if they exist
-	const TCHAR* RestrictedFolderNames[] = { TEXT("NoRedist"), TEXT("NotForLicensees"), TEXT("CarefullyRedist") };
-	for(const TCHAR* RestrictedFolderName: RestrictedFolderNames)
-	{
-		FString RestrictedFolder = InDirectoryName / RestrictedFolderName;
-		if(FPaths::DirectoryExists(RestrictedFolder))
-		{
-			FindModulePathsInDirectory(RestrictedFolder, bIsGameDirectory, NamePattern, OutModulePaths);
-		}
-	}
 }
 
 
@@ -985,6 +986,17 @@ void FModuleManager::AddBinariesDirectory(const TCHAR *InDirectory, bool bIsGame
 	}
 
 	FPlatformProcess::AddDllDirectory(InDirectory);
+
+	// Also recurse into restricted sub-folders, if they exist
+	const TCHAR* RestrictedFolderNames[] = { TEXT("NoRedist"), TEXT("NotForLicensees"), TEXT("CarefullyRedist") };
+	for (const TCHAR* RestrictedFolderName : RestrictedFolderNames)
+	{
+		FString RestrictedFolder = FPaths::Combine(InDirectory, RestrictedFolderName);
+		if (FPaths::DirectoryExists(RestrictedFolder))
+		{
+			AddBinariesDirectory(*RestrictedFolder, bIsGameDirectory);
+		}
+	}
 }
 
 
