PR #929: Explicitly skip engine targets in rocket solution (Contributed by [REDACTED_PERSON])

diff --git a/Engine/Source/Programs/UnrealBuildTool/System/ProjectFileGenerator.cs b/Engine/Source/Programs/UnrealBuildTool/System/ProjectFileGenerator.cs
--- a/Engine/Source/Programs/UnrealBuildTool/System/ProjectFileGenerator.cs
+++ b/Engine/Source/Programs/UnrealBuildTool/System/ProjectFileGenerator.cs
@@ -1673,6 +1673,12 @@ namespace UnrealBuildTool
 					IsEngineTarget = true;
 				}
 
+				if (IsEngineTarget && bGeneratingRocketProjectFiles)
+				{
+					// Rocket project file must never include engine targets.
+					continue;
+				}
+
 				bool WantProjectFileForTarget = true;
 				if(TargetFileRelativeToEngineDirectory.StartsWith(Path.Combine("Source", "Programs"), StringComparison.InvariantCultureIgnoreCase))
 				{
