Add -useprecompiledmodules to build command lines if specified on the command line to GenerateProjectFiles.bat.

diff --git a/Engine/Source/Programs/UnrealBuildTool/System/ProjectFileGenerator.cs b/Engine/Source/Programs/UnrealBuildTool/System/ProjectFileGenerator.cs
--- a/Engine/Source/Programs/UnrealBuildTool/System/ProjectFileGenerator.cs
+++ b/Engine/Source/Programs/UnrealBuildTool/System/ProjectFileGenerator.cs
@@ -141,6 +141,9 @@ namespace UnrealBuildTool
 		/// True if all documentation languages should be included in generated projects, otherwise only "INT" will be included
 		bool bAllDocumentationLanguages = false;
 
+		/// True if build targets should pass the -useprecompiledmodules argument
+		public static bool bUsePrecompiledModules = false;
+
 		/// True if we should include engine source in the generated solution
 		protected bool IncludeEngineSource = true;
 
@@ -756,6 +759,10 @@ namespace UnrealBuildTool
 						case "-ALLLANGUAGES":
 							bAllDocumentationLanguages = true;
 							break;
+
+						case "-USEPRECOMPILEDMODULES":
+							bUsePrecompiledModules = true;
+							break;
 					}
 				}
 			}
diff --git a/Engine/Source/Programs/UnrealBuildTool/System/VCProject.cs b/Engine/Source/Programs/UnrealBuildTool/System/VCProject.cs
--- a/Engine/Source/Programs/UnrealBuildTool/System/VCProject.cs
+++ b/Engine/Source/Programs/UnrealBuildTool/System/VCProject.cs
@@ -1024,34 +1024,30 @@ namespace UnrealBuildTool
 					//	..\..\Build\BatchFiles\Build.bat <TARGETNAME> <PLATFORM> <CONFIGURATION>
 					//	ie ..\..\Build\BatchFiles\Build.bat BlankProgram Win64 Debug
 
-					string ProjectPlatformConfiguration = " " + TargetName + " " + UBTPlatformName + " " + UBTConfigurationName;
+					string BuildArguments = " " + TargetName + " " + UBTPlatformName + " " + UBTConfigurationName;
+					if(ProjectFileGenerator.bUsePrecompiledModules)
+					{
+						BuildArguments += " -useprecompiledmodules";
+					}
+					if (IsForeignProject)
+					{
+						BuildArguments += " " + UProjectPath + (UnrealBuildTool.RunningRocket() ? " -rocket" : "");
+					}
 					string BatchFilesDirectoryName = Path.Combine(ProjectFileGenerator.EngineRelativePath, "Build", "BatchFiles");
 
 					// NMake Build command line
 					VCProjectFileContent.Append("		<NMakeBuildCommandLine>");
-					VCProjectFileContent.Append(EscapePath(NormalizeProjectPath(Path.Combine(BatchFilesDirectoryName, "Build.bat"))) + ProjectPlatformConfiguration.ToString());
-					if (IsForeignProject)
-					{
-						VCProjectFileContent.Append(" " + UProjectPath + (UnrealBuildTool.RunningRocket() ? " -rocket" : ""));
-					}
+					VCProjectFileContent.Append(EscapePath(NormalizeProjectPath(Path.Combine(BatchFilesDirectoryName, "Build.bat"))) + BuildArguments.ToString());
 					VCProjectFileContent.Append("</NMakeBuildCommandLine>" + ProjectFileGenerator.NewLine);
 
 					// NMake ReBuild command line
 					VCProjectFileContent.Append("		<NMakeReBuildCommandLine>");
-					VCProjectFileContent.Append(EscapePath(NormalizeProjectPath(Path.Combine(BatchFilesDirectoryName, "Rebuild.bat"))) + ProjectPlatformConfiguration.ToString());
-					if (IsForeignProject)
-					{
-						VCProjectFileContent.Append(" " + UProjectPath + (UnrealBuildTool.RunningRocket() ? " -rocket" : ""));
-					}
+					VCProjectFileContent.Append(EscapePath(NormalizeProjectPath(Path.Combine(BatchFilesDirectoryName, "Rebuild.bat"))) + BuildArguments.ToString());
 					VCProjectFileContent.Append("</NMakeReBuildCommandLine>" + ProjectFileGenerator.NewLine);
 
 					// NMake Clean command line
 					VCProjectFileContent.Append("		<NMakeCleanCommandLine>");
-					VCProjectFileContent.Append(EscapePath(NormalizeProjectPath(Path.Combine(BatchFilesDirectoryName, "Clean.bat"))) + ProjectPlatformConfiguration.ToString());
-					if (IsForeignProject)
-					{
-						VCProjectFileContent.Append(" " + UProjectPath + (UnrealBuildTool.RunningRocket() ? " -rocket" : ""));
-					}
+					VCProjectFileContent.Append(EscapePath(NormalizeProjectPath(Path.Combine(BatchFilesDirectoryName, "Clean.bat"))) + BuildArguments.ToString());
 					VCProjectFileContent.Append("</NMakeCleanCommandLine>" + ProjectFileGenerator.NewLine);
 
 					VCProjectFileContent.Append("		<NMakeOutput>");
