Shorten the -precompilemodules argument to -precompile, and -useprecompiledmodules to -useprecompiled.

diff --git a/Engine/Source/Programs/AutomationTool/Scripts/GUBP.Automation.cs b/Engine/Source/Programs/AutomationTool/Scripts/GUBP.Automation.cs
--- a/Engine/Source/Programs/AutomationTool/Scripts/GUBP.Automation.cs
+++ b/Engine/Source/Programs/AutomationTool/Scripts/GUBP.Automation.cs
@@ -148,7 +148,7 @@ public class GUBP : BuildCommand
 		string Args = "";
         if (bBuildRocket)
         {
-            Args = (bUseRocketInsteadOfBuildRocket ? "-Rocket" : "-BuildRocket -PrecompileModules");
+            Args = (bUseRocketInsteadOfBuildRocket ? "-Rocket" : "-BuildRocket -Precompile");
         }
         return Args;
     }
@@ -826,7 +826,7 @@ public class GUBP : BuildCommand
         {
             var Agenda = new UE4Build.BuildAgenda();
 
-            string AddArgs = "-nobuilduht -precompilemodules";
+            string AddArgs = "-nobuilduht -precompile";
             if (bp.bOrthogonalizeEditorPlatforms)
             {
                 AddArgs += " -skipnonhostplatforms";
@@ -1504,7 +1504,7 @@ public class GUBP : BuildCommand
         BranchInfo.BranchUProject GameProj;
         UnrealTargetPlatform TargetPlatform;
 		bool WithXp;
-		bool Precompiled; // If true, just builds targets which generate static libraries for the -UsePrecompiledModules option to UBT. If false, just build those that don't.
+		bool Precompiled; // If true, just builds targets which generate static libraries for the -UsePrecompiled option to UBT. If false, just build those that don't.
 
         public GamePlatformMonolithicsNode(GUBP bp, UnrealTargetPlatform InHostPlatform, BranchInfo.BranchUProject InGameProj, UnrealTargetPlatform InTargetPlatform, bool InWithXp = false, bool InPrecompiled = false)
             : base(InHostPlatform)
@@ -1677,7 +1677,7 @@ public class GUBP : BuildCommand
 
 			if(Precompiled)
 			{
-				Args += " -precompilemodules";
+				Args += " -precompile";
 			}
 
 			if (WithXp)
diff --git a/Engine/Source/Programs/AutomationTool/Scripts/RocketBuild.Automation.cs b/Engine/Source/Programs/AutomationTool/Scripts/RocketBuild.Automation.cs
--- a/Engine/Source/Programs/AutomationTool/Scripts/RocketBuild.Automation.cs
+++ b/Engine/Source/Programs/AutomationTool/Scripts/RocketBuild.Automation.cs
@@ -599,8 +599,8 @@ namespace Rocket
 
 			// Add the game dependencies for UE4Game on all the target platforms. We don't link libraries with their dependencies, so we can use UE4Game and -buildrocket 
 			// to pick up all the default plugins.
-			Agenda.AddTargets(new string[] { "UE4Game" }, TargetPlatform, UnrealTargetConfiguration.Development, null, false, false, false, "-precompilemodules");
-			Agenda.AddTargets(new string[] { "UE4Game" }, TargetPlatform, UnrealTargetConfiguration.Shipping, null, false, false, false, "-precompilemodules");
+			Agenda.AddTargets(new string[] { "UE4Game" }, TargetPlatform, UnrealTargetConfiguration.Development, null, false, false, false, "-precompile");
+			Agenda.AddTargets(new string[] { "UE4Game" }, TargetPlatform, UnrealTargetConfiguration.Shipping, null, false, false, false, "-precompile");
 
 			// Read the file list
 			string FileListPath = new UE4Build(bp).GenerateExternalFileList(Agenda);
diff --git a/Engine/Source/Programs/UnrealBuildTool/Configuration/UEBuildTarget.cs b/Engine/Source/Programs/UnrealBuildTool/Configuration/UEBuildTarget.cs
--- a/Engine/Source/Programs/UnrealBuildTool/Configuration/UEBuildTarget.cs
+++ b/Engine/Source/Programs/UnrealBuildTool/Configuration/UEBuildTarget.cs
@@ -166,8 +166,8 @@ namespace UnrealBuildTool
 		public bool bIsEditorRecompile;
 		public string RemoteRoot;
 		public List<OnlyModule> OnlyModules;
-		public bool bPrecompileModules;
-		public bool bUsePrecompiledModules;
+		public bool bPrecompile;
+		public bool bUsePrecompiled;
 	}
 
 
@@ -222,8 +222,8 @@ namespace UnrealBuildTool
 			bool bIsEditorRecompile = false;
 
 			// Settings for creating/using static libraries for the engine
-			bool bPrecompileModules = false;
-			bool bUsePrecompiledModules = UnrealBuildTool.RunningRocket();
+			bool bPrecompile = false;
+			bool bUsePrecompiled = UnrealBuildTool.RunningRocket();
 
 			// Combine the two arrays of arguments
 			List<string> Arguments = new List<string>(SourceArguments.Length);
@@ -384,17 +384,17 @@ namespace UnrealBuildTool
 							}
 							break;
 
-						case "-PRECOMPILEMODULES":
+						case "-PRECOMPILE":
 							{
 								// Make static libraries for all engine modules as intermediates for this target
-								bPrecompileModules = true;
+								bPrecompile = true;
 							}
 							break;
-							
-						case "-USEPRECOMPILEDMODULES":
+
+						case "-USEPRECOMPILED":
 							{
 								// Use existing static libraries for all engine modules in this target
-								bUsePrecompiledModules = true;
+								bUsePrecompiled = true;
 							}
 							break;
 
@@ -458,8 +458,8 @@ namespace UnrealBuildTool
 							bIsEditorRecompile = bIsEditorRecompile,
 							RemoteRoot = RemoteRoot,
 							OnlyModules = OnlyModules,
-							bPrecompileModules = bPrecompileModules,
-							bUsePrecompiledModules = bUsePrecompiledModules,
+							bPrecompile = bPrecompile,
+							bUsePrecompiled = bUsePrecompiled,
 						} );
 					break;
 				}
@@ -692,11 +692,11 @@ namespace UnrealBuildTool
 		/** Remote path of the binary if it is to be synced with CookerSync */
 		public string RemoteRoot;
 
-		/** Whether to build engine modules that can be used without compiling */
-		public bool bPrecompileModules;
+		/** Whether to build target modules that can be reused for future builds */
+		public bool bPrecompile;
 
-		/** Whether to use precompiled engine modules */
-		public bool bUsePrecompiledModules;
+		/** Whether to use precompiled target modules */
+		public bool bUsePrecompiled;
 
 		/** The C++ environment that all the environments used to compile UE-based modules are derived from. */
 		[NonSerialized]
@@ -765,8 +765,8 @@ namespace UnrealBuildTool
 			Configuration = InDesc.Configuration;
 			Rules = InRulesObject;
 			bEditorRecompile = InDesc.bIsEditorRecompile;
-			bPrecompileModules = InDesc.bPrecompileModules;
-			bUsePrecompiledModules = InDesc.bUsePrecompiledModules;
+			bPrecompile = InDesc.bPrecompile;
+			bUsePrecompiled = InDesc.bUsePrecompiled;
 
 			{
 				bCompileMonolithic = (Rules != null) ? Rules.ShouldCompileMonolithic(InDesc.Platform, InDesc.Configuration) : false;
@@ -1319,7 +1319,7 @@ namespace UnrealBuildTool
 				// Don't add static library files to the manifest as we do not check them into perforce.
 				// However, add them to the manifest when cleaning the project as we do want to delete 
 				// them in that case.
-				if (UEBuildConfiguration.bCleanProject == false && !bPrecompileModules)
+				if (UEBuildConfiguration.bCleanProject == false && !bPrecompile)
 				{
                     if (Binary.Config.Type == UEBuildBinaryType.StaticLibrary)
 					{
@@ -2212,7 +2212,7 @@ namespace UnrealBuildTool
 			else
 			{
 				OutputFilePaths = MakeBinaryPaths(Module.Name, GetAppName() + "-" + Module.Name, BinaryType, TargetType, Plugin, AppName);
-				if (bPrecompileModules)
+				if (bPrecompile)
 				{
 					SpecialRocketLibFilesThatAreBuildProducts.AddRange(OutputFilePaths);
 				}
@@ -2297,13 +2297,13 @@ namespace UnrealBuildTool
 		/// </summary>
 		protected void AddPrecompiledModules()
 		{
-			if(bPrecompileModules || bUsePrecompiledModules)
+			if(bPrecompile || bUsePrecompiled)
 			{
 				// Find all the precompiled module names.
 				List<string> PrecompiledModuleNames = new List<string>();
 				Rules.GetModulesToPrecompile(new TargetInfo(Platform, Configuration, TargetTypeOrNull), PrecompiledModuleNames);
 
-				// Add all the enabled plugins to the precompiled module list. Plugins are always precompiled, even if bPrecompileModules is not set, so we should precompile their dependencies.
+				// Add all the enabled plugins to the precompiled module list. Plugins are always precompiled, even if bPrecompile is not set, so we should precompile their dependencies.
 				foreach(PluginInfo Plugin in BuildPlugins)
 				{
 					foreach(PluginInfo.PluginModuleInfo Module in Plugin.Modules)
@@ -2365,7 +2365,7 @@ namespace UnrealBuildTool
 				{
 					if(PrecompiledModule is UEBuildModuleCPP && !BoundModuleNames.Contains(PrecompiledModule.Name))
 					{
-						UEBuildBinary Binary = AddBinaryForModule(PrecompiledModule.Name, bCompileMonolithic? UEBuildBinaryType.StaticLibrary : UEBuildBinaryType.DynamicLinkLibrary, bUsePrecompiledModules);
+						UEBuildBinary Binary = AddBinaryForModule(PrecompiledModule.Name, bCompileMonolithic? UEBuildBinaryType.StaticLibrary : UEBuildBinaryType.DynamicLinkLibrary, bUsePrecompiled);
 						PrecompiledBinaries.Add(Binary);
 						BoundModuleNames.Add(PrecompiledModule.Name);
 					}
@@ -2408,7 +2408,7 @@ namespace UnrealBuildTool
 			}
 
 			// Get the binary configuration
-			UEBuildBinaryConfiguration BinaryConfig = new UEBuildBinaryConfiguration(BinaryType, OutputFilePaths, IntermediateDirectory, true, bInAllowCompilation: bPrecompileModules, bInCompileMonolithic: bCompileMonolithic, InModuleNames: new List<string>{ ModuleName });
+			UEBuildBinaryConfiguration BinaryConfig = new UEBuildBinaryConfiguration(BinaryType, OutputFilePaths, IntermediateDirectory, true, bInAllowCompilation: bPrecompile, bInCompileMonolithic: bCompileMonolithic, InModuleNames: new List<string>{ ModuleName });
 
 			// Create the binary
 			UEBuildBinaryCPP Binary = new UEBuildBinaryCPP(this, BinaryConfig);
@@ -2634,7 +2634,7 @@ namespace UnrealBuildTool
 			DependentPlugins.AddRange(EnabledPlugins);
 
 			// Set the list of plugins that should be built
-			if (bPrecompileModules && TargetType != TargetRules.TargetType.Program)
+			if (bPrecompile && TargetType != TargetRules.TargetType.Program)
 			{
 				BuildPlugins.AddRange(ValidPlugins);
 			}
diff --git a/Engine/Source/Programs/UnrealBuildTool/System/ProjectFileGenerator.cs b/Engine/Source/Programs/UnrealBuildTool/System/ProjectFileGenerator.cs
--- a/Engine/Source/Programs/UnrealBuildTool/System/ProjectFileGenerator.cs
+++ b/Engine/Source/Programs/UnrealBuildTool/System/ProjectFileGenerator.cs
@@ -141,8 +141,8 @@ namespace UnrealBuildTool
 		/// True if all documentation languages should be included in generated projects, otherwise only "INT" will be included
 		bool bAllDocumentationLanguages = false;
 
-		/// True if build targets should pass the -useprecompiledmodules argument
-		public static bool bUsePrecompiledModules = false;
+		/// True if build targets should pass the -useprecompiled argument
+		public static bool bUsePrecompiled = false;
 
 		/// True if we should include engine source in the generated solution
 		protected bool IncludeEngineSource = true;
@@ -760,8 +760,8 @@ namespace UnrealBuildTool
 							bAllDocumentationLanguages = true;
 							break;
 
-						case "-USEPRECOMPILEDMODULES":
-							bUsePrecompiledModules = true;
+						case "-USEPRECOMPILED":
+							bUsePrecompiled = true;
 							break;
 					}
 				}
diff --git a/Engine/Source/Programs/UnrealBuildTool/System/RulesCompiler.cs b/Engine/Source/Programs/UnrealBuildTool/System/RulesCompiler.cs
--- a/Engine/Source/Programs/UnrealBuildTool/System/RulesCompiler.cs
+++ b/Engine/Source/Programs/UnrealBuildTool/System/RulesCompiler.cs
@@ -668,7 +668,7 @@ namespace UnrealBuildTool
 		}
 
 		/// <summary>
-		/// Allows the target to specify modules which can be precompiled with the -PrecompileModules/-UsePrecompiledModules arguments to UBT. 
+		/// Allows the target to specify modules which can be precompiled with the -Precompile/-UsePrecompiled arguments to UBT. 
 		/// All dependencies of the specified modules will be included.
 		/// </summary>
 		/// <param name="Target">The target information, such as platform and configuration</param>
diff --git a/Engine/Source/Programs/UnrealBuildTool/System/VCProject.cs b/Engine/Source/Programs/UnrealBuildTool/System/VCProject.cs
--- a/Engine/Source/Programs/UnrealBuildTool/System/VCProject.cs
+++ b/Engine/Source/Programs/UnrealBuildTool/System/VCProject.cs
@@ -1025,9 +1025,9 @@ namespace UnrealBuildTool
 					//	ie ..\..\Build\BatchFiles\Build.bat BlankProgram Win64 Debug
 
 					string BuildArguments = " " + TargetName + " " + UBTPlatformName + " " + UBTConfigurationName;
-					if(ProjectFileGenerator.bUsePrecompiledModules)
+					if(ProjectFileGenerator.bUsePrecompiled)
 					{
-						BuildArguments += " -useprecompiledmodules";
+						BuildArguments += " -useprecompiled";
 					}
 					if (IsForeignProject)
 					{
diff --git a/Engine/Source/Programs/UnrealBuildTool/Windows/VCToolChain.cs b/Engine/Source/Programs/UnrealBuildTool/Windows/VCToolChain.cs
--- a/Engine/Source/Programs/UnrealBuildTool/Windows/VCToolChain.cs
+++ b/Engine/Source/Programs/UnrealBuildTool/Windows/VCToolChain.cs
@@ -1686,7 +1686,7 @@ namespace UnrealBuildTool
         public override void AddFilesToManifest(BuildManifest Manifest, UEBuildBinary Binary)
         {
             // ok, this is pretty awful, we want the import libraries that go with the editor, only on the PC
-            if (Binary.Target.bPrecompileModules &&
+            if (Binary.Target.bPrecompile &&
                 Path.GetFileNameWithoutExtension(Binary.Config.OutputFilePath).StartsWith("UE4Editor-", StringComparison.InvariantCultureIgnoreCase) &&
                 Path.GetExtension(Binary.Config.OutputFilePath).EndsWith("dll", StringComparison.InvariantCultureIgnoreCase) &&
                 Binary.Config.Type == UEBuildBinaryType.DynamicLinkLibrary)
