Fixed UnrealBuildTool printing "Performing full C++ include scan" more than once

diff --git a/Engine/Source/Programs/UnrealBuildTool/System/FlatCPPIncludeDepencencyCache.cs b/Engine/Source/Programs/UnrealBuildTool/System/FlatCPPIncludeDepencencyCache.cs
--- a/Engine/Source/Programs/UnrealBuildTool/System/FlatCPPIncludeDepencencyCache.cs
+++ b/Engine/Source/Programs/UnrealBuildTool/System/FlatCPPIncludeDepencencyCache.cs
@@ -67,7 +67,7 @@ namespace UnrealBuildTool
 				}
 			}
 
-			if( !ProjectFileGenerator.bGenerateProjectFiles )
+			if( !ProjectFileGenerator.bGenerateProjectFiles && !UnrealBuildTool.bNeedsFullCPPIncludeRescan )
 			{ 
 				UnrealBuildTool.bNeedsFullCPPIncludeRescan = true;
 				Log.TraceInformation( "Performing full C++ include scan (no include cache file)" );
