Changed Distribution level checking so that it is not performed if you have a uproject.

Removed use of bOutputPubliclyDistributable from all targets that no longer need it.

diff --git a/Engine/Source/Programs/UnrealBuildTool/Configuration/UEBuildTarget.cs b/Engine/Source/Programs/UnrealBuildTool/Configuration/UEBuildTarget.cs
--- a/Engine/Source/Programs/UnrealBuildTool/Configuration/UEBuildTarget.cs
+++ b/Engine/Source/Programs/UnrealBuildTool/Configuration/UEBuildTarget.cs
@@ -3373,9 +3373,8 @@ namespace UnrealBuildTool
 		{
 			if (Rules != null)
 			{
-				// We don't want to check this for true game projects (but do want to include UE4Game)
-				bool bIsUProjectGame = UnrealBuildTool.HasUProjectFile() && Rules.Type == TargetRules.TargetType.Game;
-				if (!bIsUProjectGame && !Rules.bOutputPubliclyDistributable)
+				// We don't want to check this for uprojects (but do want to include UE4Game)
+				if (!UnrealBuildTool.HasUProjectFile() && !Rules.bOutputPubliclyDistributable)
 				{
 					return true;
 				}
