PR #993: Resolving a crash when building project files with a full temp folder (Contributed by [REDACTED_PERSON])

diff --git a/.gitignore b/.gitignore
--- a/.gitignore
+++ b/.gitignore
@@ -99,6 +99,7 @@
 /Engine/Documentation/**/*.cur
 /Engine/Documentation/**/*.gif
 /Engine/Documentation/**/*.jpg
+/Engine/Documentation/CHM/**
 
 # Ignore all binary types in the extras folder
 /Engine/Extras/*/**
@@ -107,6 +108,11 @@
 /Engine/Extras/Maya_AnimationRiggingTools/**/*.pyd
 /Engine/Extras/Maya_AnimationRiggingTools/MayaTools/General/ART/**
 /Engine/Extras/Maya_AnimationRiggingTools/MayaTools/General/Icons/**
+/Engine/Extras/Maya_AnimationRiggingTools/ArtToolsOSX/MayaTools/General/ART/
+/Engine/Extras/Maya_AnimationRiggingTools/ArtToolsOSX/MayaTools/General/Icons/
+/Engine/Extras/Maya_AnimationRiggingTools/ArtToolsOSX/MayaTools/General/Scripts/projectSettings.txt
+/Engine/Extras/Maya_AnimationRiggingTools/ArtToolsOSX/MayaTools/General/Scripts/settings.txt
+/Engine/Extras/Maya_AnimationRiggingTools/ArtToolsOSX/README.md
 !/Engine/Extras/VisualStudioSnippets/**
 
 # Ignore any saved local files
@@ -144,6 +150,11 @@
 # Ignore LinuxNativeDialogs build files (all extensions)
 /Engine/Source/ThirdParty/LinuxNativeDialogs/UELinuxNativeDialogs/build/**
 
+# Ignore PhysX binaries
+/Engine/Source/ThirdParty/PhysX/APEX-1.3/externals/CUDA_6.5.19/bin/
+/Engine/Source/ThirdParty/PhysX/PhysX-3.3/externals/clang/3.3.2/linux32/bin/
+/Engine/Source/ThirdParty/PhysX/PhysX-3.3/externals/clang/3.3.2/osx/bin/
+
 # Ignore sample and template projects
 /Samples/**
 /Templates/**
diff --git a/Engine/Source/Programs/UnrealBuildTool/System/Utils.cs b/Engine/Source/Programs/UnrealBuildTool/System/Utils.cs
--- a/Engine/Source/Programs/UnrealBuildTool/System/Utils.cs
+++ b/Engine/Source/Programs/UnrealBuildTool/System/Utils.cs
@@ -183,11 +183,15 @@ namespace UnrealBuildTool
 			if( File.Exists( BatchFileName ) )
 			{
 				// Create a wrapper batch file that echoes environment variables to a text file
-				var EnvOutputFileName = Path.GetTempFileName();
-				var EnvReaderBatchFileName = EnvOutputFileName + ".bat";
-				Log.TraceVerbose( "Creating .bat file {0} for harvesting environment variables.", EnvReaderBatchFileName );
+                string EnvOutputFileName;
+                string EnvReaderBatchFileName;
+                try
+                {
+					EnvOutputFileName = Path.GetTempFileName();
+					EnvReaderBatchFileName = EnvOutputFileName + ".bat";
+
+					Log.TraceVerbose( "Creating .bat file {0} for harvesting environment variables.", EnvReaderBatchFileName );
 
-				{
 					var EnvReaderBatchFileContent = new List<string>();
 
 					// Run 'vcvars32.bat' (or similar x64 version) to set environment variables
@@ -200,7 +204,11 @@ namespace UnrealBuildTool
 					EnvReaderBatchFileContent.Add( String.Format( "\"{0}\" \"{1}\"", Path.Combine( GetExecutingAssemblyDirectory(), "EnvVarsToXML.exe" ), EnvOutputFileName ) );
 
 					ResponseFile.Create( EnvReaderBatchFileName, EnvReaderBatchFileContent );
-				}
+                }
+                catch (Exception Ex)
+                {
+                    throw new BuildException(Ex, "Failed to create temporary batch file to harvest environment variables (\"{0}\")", Ex.Message);
+                }
 
 				Log.TraceVerbose( "Finished creating .bat file.  Environment variables will be written to {0}.", EnvOutputFileName );
 
