Eliminate SpecialRocketLibFilesThatAreBuildProducts - everything that gets precompiled into static libs is now a normal UBT binary; we don't need to special case anything.

diff --git a/Engine/Source/Programs/UnrealBuildTool/Configuration/UEBuildTarget.cs b/Engine/Source/Programs/UnrealBuildTool/Configuration/UEBuildTarget.cs
--- a/Engine/Source/Programs/UnrealBuildTool/Configuration/UEBuildTarget.cs
+++ b/Engine/Source/Programs/UnrealBuildTool/Configuration/UEBuildTarget.cs
@@ -1288,7 +1288,7 @@ namespace UnrealBuildTool
 		}
 
 		/** Generates a public manifest file for writing out */
-        public void GenerateManifest(IUEToolChain ToolChain, List<UEBuildBinary> Binaries, CPPTargetPlatform Platform, List<string> SpecialRocketLibFilesThatAreBuildProducts)
+        public void GenerateManifest(IUEToolChain ToolChain, List<UEBuildBinary> Binaries, CPPTargetPlatform Platform)
 		{
 			string ManifestPath;
 			if (UnrealBuildTool.RunningRocket())
@@ -1346,14 +1346,6 @@ namespace UnrealBuildTool
 
                 ToolChain.AddFilesToManifest(Manifest, Binary);
 			}
-			{
-				string DebugInfoExtension = BuildPlatform.GetDebugInfoExtension(UEBuildBinaryType.StaticLibrary);
-				foreach (var RedistLib in SpecialRocketLibFilesThatAreBuildProducts)
-				{
-					Manifest.AddBuildProduct(RedistLib, DebugInfoExtension);
-				}
-			}
-
 
 			if (UEBuildConfiguration.bCleanProject)
 			{
@@ -1371,7 +1363,7 @@ namespace UnrealBuildTool
 			OutputItems = new List<FileItem>();
 			UObjectModules = new List<UHTModuleInfo>();
 
-			var SpecialRocketLibFilesThatAreBuildProducts = PreBuildSetup(TargetToolChain);
+			PreBuildSetup(TargetToolChain);
 
 			EULAViolationWarning = !ProjectFileGenerator.bGenerateProjectFiles
 				? CheckForEULAViolation()
@@ -1448,7 +1440,7 @@ namespace UnrealBuildTool
 			// If we're only generating the manifest, return now
 			if (UEBuildConfiguration.bGenerateManifest || UEBuildConfiguration.bCleanProject)
 			{
-                GenerateManifest(TargetToolChain, AppBinaries, GlobalLinkEnvironment.Config.Target.Platform, SpecialRocketLibFilesThatAreBuildProducts);
+                GenerateManifest(TargetToolChain, AppBinaries, GlobalLinkEnvironment.Config.Target.Platform);
                 if (!BuildConfiguration.bXGEExport)
                 {
                     return ECompilationResult.Succeeded;
@@ -1609,7 +1601,7 @@ namespace UnrealBuildTool
 		/// Setup target before build. This method finds dependencies, sets up global environment etc.
 		/// </summary>
 		/// <returns>Special Rocket lib files that are build products.</returns>
-		public List<string> PreBuildSetup(IUEToolChain TargetToolChain)
+		public void PreBuildSetup(IUEToolChain TargetToolChain)
 		{
 			// Set up the global compile and link environment in GlobalCompileEnvironment and GlobalLinkEnvironment.
 			SetupGlobalEnvironment();
@@ -1623,12 +1615,10 @@ namespace UnrealBuildTool
 			// Setup the target's plugins
 			SetupPlugins();
 
-			var SpecialRocketLibFilesThatAreBuildProducts = new List<string>();
-
 			// Add the enabled plugins to the build
 			foreach (PluginInfo BuildPlugin in BuildPlugins)
 			{
-				SpecialRocketLibFilesThatAreBuildProducts.AddRange(AddPlugin(BuildPlugin));
+				AddPlugin(BuildPlugin);
 			}
 
 			// Allow the platform to setup binaries/plugins/modules
@@ -1793,8 +1783,6 @@ namespace UnrealBuildTool
 					Binary.CheckOutputDistributionLevelAgainstDependencies();
 				}
 			}
-
-			return SpecialRocketLibFilesThatAreBuildProducts;
 		}
 
 		private static string AddModuleFilenameSuffix(string ModuleName, string FilePath, string Suffix)
@@ -2178,26 +2166,22 @@ namespace UnrealBuildTool
 		/// <summary>
 		/// Include the given plugin in the target. It may be included as a separate binary, or compiled into a monolithic executable.
 		/// </summary>
-        public List<string> AddPlugin(PluginInfo Plugin)
+        public void AddPlugin(PluginInfo Plugin)
 		{
-            var SpecialRocketLibFilesThatAreBuildProducts = new List<string>();
 			foreach(PluginInfo.PluginModuleInfo Module in Plugin.Modules)
 			{
 				if (ShouldIncludePluginModule(Plugin, Module))
 				{
-                    SpecialRocketLibFilesThatAreBuildProducts.AddRange(AddPluginModule(Plugin, Module));
+                    AddPluginModule(Plugin, Module);
 				}
 			}
-            return SpecialRocketLibFilesThatAreBuildProducts;
 		}
 
 		/// <summary>
 		/// Include the given plugin module in the target. Will be built in the appropriate subfolder under the plugin directory.
 		/// </summary>
-        public List<string> AddPluginModule(PluginInfo Plugin, PluginInfo.PluginModuleInfo Module)
+        public void AddPluginModule(PluginInfo Plugin, PluginInfo.PluginModuleInfo Module)
 		{
-            var SpecialRocketLibFilesThatAreBuildProducts = new List<string>();
-
 			bool bCompileMonolithic = ShouldCompileMonolithic();
 
 			// Get the binary type to build
@@ -2212,10 +2196,6 @@ namespace UnrealBuildTool
 			else
 			{
 				OutputFilePaths = MakeBinaryPaths(Module.Name, GetAppName() + "-" + Module.Name, BinaryType, TargetType, Plugin, AppName);
-				if (bPrecompile)
-				{
-					SpecialRocketLibFilesThatAreBuildProducts.AddRange(OutputFilePaths);
-				}
 			}
 
 			// Try to determine if we have the rules file
@@ -2251,7 +2231,6 @@ namespace UnrealBuildTool
 																				bInHasModuleRules:       bHasModuleRules,
 																				InModuleNames:           new List<string> { Module.Name } );
 			AppBinaries.Add(new UEBuildBinaryCPP(this, Config));
-            return SpecialRocketLibFilesThatAreBuildProducts;
 		}
 
 		/// When building a target, this is called to add any additional modules that should be compiled along
