Deprecate bBuildAdditionalConsoleApplication from LinkEnvironment, and move it to UEBuildBinaryConfiguration instead - it's not actually used by anything that does linking; it's a higher level behavior. The default is now false, which mirrors the only use case we had for it (and removes a hack for checking UEBuildConfiguration.bBuildEditor).

diff --git a/Engine/Source/Programs/BuildPatchTool/BuildPatchTool.Target.cs b/Engine/Source/Programs/BuildPatchTool/BuildPatchTool.Target.cs
--- a/Engine/Source/Programs/BuildPatchTool/BuildPatchTool.Target.cs
+++ b/Engine/Source/Programs/BuildPatchTool/BuildPatchTool.Target.cs
@@ -48,7 +48,6 @@ public class BuildPatchToolTarget : TargetRules
 		OutCPPEnvironmentConfiguration.Definitions.Add("WITH_BUILDPATCHGENERATION=1");
 
 		OutLinkEnvironmentConfiguration.bIsBuildingConsoleApplication = true;
-		OutLinkEnvironmentConfiguration.bBuildAdditionalConsoleApplication = false;
 		OutLinkEnvironmentConfiguration.bHasExports = false;
 	}
     public override bool GUBP_AlwaysBuildWithTools(UnrealTargetPlatform InHostPlatform, out bool bInternalToolOnly, out bool SeparateNode, out bool CrossCompile)
diff --git a/Engine/Source/Programs/CrashReporter/CrashReportClient/CrashReportClient.Target.cs b/Engine/Source/Programs/CrashReporter/CrashReportClient/CrashReportClient.Target.cs
--- a/Engine/Source/Programs/CrashReporter/CrashReportClient/CrashReportClient.Target.cs
+++ b/Engine/Source/Programs/CrashReporter/CrashReportClient/CrashReportClient.Target.cs
@@ -92,9 +92,6 @@ public class CrashReportClientTarget : TargetRules
 		// the linker.
 		OutLinkEnvironmentConfiguration.bHasExports = false;
 
-		// Do NOT produce additional console app exe
-		OutLinkEnvironmentConfiguration.bBuildAdditionalConsoleApplication = false;
-
 		OutCPPEnvironmentConfiguration.Definitions.Add( "USE_CHECKS_IN_SHIPPING=1" );
 	}
     public override bool GUBP_AlwaysBuildWithTools(UnrealTargetPlatform InHostPlatform, out bool bInternalToolOnly, out bool SeparateNode, out bool CrossCompile)
diff --git a/Engine/Source/Programs/IOS/UnrealLaunchDaemon/UnrealLaunchDaemon.Target.cs b/Engine/Source/Programs/IOS/UnrealLaunchDaemon/UnrealLaunchDaemon.Target.cs
--- a/Engine/Source/Programs/IOS/UnrealLaunchDaemon/UnrealLaunchDaemon.Target.cs
+++ b/Engine/Source/Programs/IOS/UnrealLaunchDaemon/UnrealLaunchDaemon.Target.cs
@@ -51,6 +51,5 @@ public class UnrealLaunchDaemonTarget : TargetRules
 		UEBuildConfiguration.bCompileAgainstEngine = false;
 
 		OutLinkEnvironmentConfiguration.bHasExports = false;
-		OutLinkEnvironmentConfiguration.bBuildAdditionalConsoleApplication = false;
 	}
 }
diff --git a/Engine/Source/Programs/SlateViewer/SlateViewer.Target.cs b/Engine/Source/Programs/SlateViewer/SlateViewer.Target.cs
--- a/Engine/Source/Programs/SlateViewer/SlateViewer.Target.cs
+++ b/Engine/Source/Programs/SlateViewer/SlateViewer.Target.cs
@@ -64,8 +64,5 @@ public class SlateViewerTarget : TargetRules
 		// SlateViewer.exe has no exports, so no need to verify that a .lib and .exp file was emitted by
 		// the linker.
 		OutLinkEnvironmentConfiguration.bHasExports = false;
-
-		// Do NOT produce additional console app exe
-		OutLinkEnvironmentConfiguration.bBuildAdditionalConsoleApplication = false;
 	}
 }
diff --git a/Engine/Source/Programs/SymbolDebugger/SymbolDebugger.Target.cs b/Engine/Source/Programs/SymbolDebugger/SymbolDebugger.Target.cs
--- a/Engine/Source/Programs/SymbolDebugger/SymbolDebugger.Target.cs
+++ b/Engine/Source/Programs/SymbolDebugger/SymbolDebugger.Target.cs
@@ -61,8 +61,5 @@ public class SymbolDebuggerTarget : TargetRules
 		// SymbolDebugger.exe has no exports, so no need to verify that a .lib and .exp file was emitted by
 		// the linker.
 		OutLinkEnvironmentConfiguration.bHasExports = false;
-
-		// Do NOT produce additional console app exe
-		OutLinkEnvironmentConfiguration.bBuildAdditionalConsoleApplication = false;
 	}
 }
diff --git a/Engine/Source/Programs/UnrealBuildTool/Configuration/UEBuildBinary.cs b/Engine/Source/Programs/UnrealBuildTool/Configuration/UEBuildBinary.cs
--- a/Engine/Source/Programs/UnrealBuildTool/Configuration/UEBuildBinary.cs
+++ b/Engine/Source/Programs/UnrealBuildTool/Configuration/UEBuildBinary.cs
@@ -115,6 +115,12 @@ namespace UnrealBuildTool
 		/// </summary>
 		public bool bCompileMonolithic = false;
 
+        /// <summary>
+		/// If true, creates an additional console application. Hack for Windows, where it's not possible to conditionally inherit a parent's console Window depending on how
+		/// the application is invoked; you have to link the same executable with a different subsystem setting.
+		/// </summary>
+		public bool bBuildAdditionalConsoleApp = false;
+
 		/// <summary>
 		/// The build target configuration being compiled
 		/// </summary>
@@ -768,7 +774,7 @@ namespace UnrealBuildTool
 			OutputFiles.AddRange(Executables);
 
 			// Produce additional console app if requested
-			if (BinaryLinkEnvironment.Config.CanProduceAdditionalConsoleApp && UEBuildConfiguration.bBuildEditor)
+			if (Config.bBuildAdditionalConsoleApp)
 			{
 				// Produce additional binary but link it as a console app
 				var ConsoleAppLinkEvironment = BinaryLinkEnvironment.DeepCopy();
diff --git a/Engine/Source/Programs/UnrealBuildTool/Configuration/UEBuildEditor.cs b/Engine/Source/Programs/UnrealBuildTool/Configuration/UEBuildEditor.cs
--- a/Engine/Source/Programs/UnrealBuildTool/Configuration/UEBuildEditor.cs
+++ b/Engine/Source/Programs/UnrealBuildTool/Configuration/UEBuildEditor.cs
@@ -31,6 +31,11 @@ namespace UnrealBuildTool
 																					bInAllowExports:!ShouldCompileMonolithic(),
 																					InModuleNames: new List<string>() { "Launch" } );
 
+				if(Platform == UnrealTargetPlatform.Win64 && Configuration != UnrealTargetConfiguration.Shipping)
+				{
+					Config.bBuildAdditionalConsoleApp = true;
+				}
+
 				AppBinaries.Add( new UEBuildBinaryCPP( this, Config ) );
 			}
 
diff --git a/Engine/Source/Programs/UnrealBuildTool/Configuration/UEBuildTarget.cs b/Engine/Source/Programs/UnrealBuildTool/Configuration/UEBuildTarget.cs
--- a/Engine/Source/Programs/UnrealBuildTool/Configuration/UEBuildTarget.cs
+++ b/Engine/Source/Programs/UnrealBuildTool/Configuration/UEBuildTarget.cs
@@ -1334,9 +1334,7 @@ namespace UnrealBuildTool
 					Manifest.AddBuildProduct(OutputFilePath, DebugInfoExtension);
 				}
 
-				if (Binary.Config.Type == UEBuildBinaryType.Executable &&
-					  GlobalLinkEnvironment.Config.CanProduceAdditionalConsoleApp &&
-					  UEBuildConfiguration.bBuildEditor)
+				if (Binary.Config.Type == UEBuildBinaryType.Executable && Binary.Config.bBuildAdditionalConsoleApp)
 				{
 					foreach (string OutputFilePath in Binary.Config.OutputFilePaths)
 					{
diff --git a/Engine/Source/Programs/UnrealBuildTool/System/LinkEnvironment.cs b/Engine/Source/Programs/UnrealBuildTool/System/LinkEnvironment.cs
--- a/Engine/Source/Programs/UnrealBuildTool/System/LinkEnvironment.cs
+++ b/Engine/Source/Programs/UnrealBuildTool/System/LinkEnvironment.cs
@@ -84,8 +84,9 @@ namespace UnrealBuildTool
 		/** True if this is a console application that's being build */
 		public bool bIsBuildingConsoleApplication = false;
 
-		/** If bIsBuildingConsoleApplication is false and this is true, additional console app exe will be built. */
-		public bool bBuildAdditionalConsoleApplication = true;
+		/** This setting is replaced by UEBuildBinaryConfiguration.bBuildAdditionalConsoleApp. */
+		[Obsolete("This setting is replaced by UEBuildBinaryConfiguration.bBuildAdditionalConsoleApp. It is explicitly set to true for editor targets, and defaults to false otherwise.")]
+		public bool bBuildAdditionalConsoleApplication { set { } }
 
 		/** If set, overrides the program entry function on Windows platform.  This is used by the base UE4
 		    program so we can link in either command-line mode or windowed mode without having to recompile the Launch module. */
@@ -133,25 +134,11 @@ namespace UnrealBuildTool
 			bIsBuildingLibrary = InCopyEnvironment.bIsBuildingLibrary;
 			bIsBuildingDLL = InCopyEnvironment.bIsBuildingDLL;
 			bIsBuildingConsoleApplication = InCopyEnvironment.bIsBuildingConsoleApplication;
-			bBuildAdditionalConsoleApplication = InCopyEnvironment.bBuildAdditionalConsoleApplication;
 			WindowsEntryPointOverride = InCopyEnvironment.WindowsEntryPointOverride;
 			bIsCrossReferenced = InCopyEnvironment.bIsCrossReferenced;
 			bHasExports = InCopyEnvironment.bHasExports;
 			bIsBuildingDotNetAssembly = InCopyEnvironment.bIsBuildingDotNetAssembly;
 		}
-
-		/** Whether this configuration is suitable to produce additional console app. */
-		public bool CanProduceAdditionalConsoleApp
-		{
-			get
-			{
-				return Target.Platform == CPPTargetPlatform.Win64 &&
-						bIsBuildingConsoleApplication == false &&
-						bBuildAdditionalConsoleApplication == true &&
-						bIsBuildingDLL == false &&
-						(Target.Configuration == CPPTargetConfiguration.Debug || Target.Configuration == CPPTargetConfiguration.Development);
-			}
-		}
 	}
 
 	/** Encapsulates the environment that is used to link object files. */
diff --git a/Engine/Source/Programs/UnrealCEFSubProcess/UnrealCEFSubProcess.Target.cs b/Engine/Source/Programs/UnrealCEFSubProcess/UnrealCEFSubProcess.Target.cs
--- a/Engine/Source/Programs/UnrealCEFSubProcess/UnrealCEFSubProcess.Target.cs
+++ b/Engine/Source/Programs/UnrealCEFSubProcess/UnrealCEFSubProcess.Target.cs
@@ -66,9 +66,6 @@ public class UnrealCEFSubProcessTarget : TargetRules
 		// CEFSubProcess is a console application, not a Windows app (sets entry point to main(), instead of WinMain())
 		OutLinkEnvironmentConfiguration.bIsBuildingConsoleApplication = false;
 
-		// Do NOT produce additional console app exe
-		OutLinkEnvironmentConfiguration.bBuildAdditionalConsoleApplication = false;
-
 		// Disable logging, as the sub processes are spawned often and logging will just slow them down
 		OutCPPEnvironmentConfiguration.Definitions.Add("ALLOW_LOG_FILE=0");
 	}
diff --git a/Engine/Source/Programs/UnrealFrontend/UnrealFrontend.Target.cs b/Engine/Source/Programs/UnrealFrontend/UnrealFrontend.Target.cs
--- a/Engine/Source/Programs/UnrealFrontend/UnrealFrontend.Target.cs
+++ b/Engine/Source/Programs/UnrealFrontend/UnrealFrontend.Target.cs
@@ -47,7 +47,6 @@ public class UnrealFrontendTarget : TargetRules
 		UEBuildConfiguration.bCompileWithStatsWithoutEngine = true;
 		UEBuildConfiguration.bCompileWithPluginSupport = true;
 
-        OutLinkEnvironmentConfiguration.bBuildAdditionalConsoleApplication = false;
 		OutLinkEnvironmentConfiguration.bHasExports = false;
 	}
     public override bool GUBP_AlwaysBuildWithBaseEditor()
