Only write out 8.3 filenames to the envvar-harvesting .bat file, to avoid Unicode paths which cmd.exe can't deal with.
Better error message when XML parsing fails.

diff --git a/Engine/Source/Programs/UnrealBuildTool/System/Utils.cs b/Engine/Source/Programs/UnrealBuildTool/System/Utils.cs
--- a/Engine/Source/Programs/UnrealBuildTool/System/Utils.cs
+++ b/Engine/Source/Programs/UnrealBuildTool/System/Utils.cs
@@ -9,6 +9,7 @@ using System.Text.RegularExpressions;
 using System.Xml;
 using System.Xml.Serialization;
 using System.Runtime.CompilerServices;
+using System.Runtime.InteropServices;
 using System.Linq;
 
 namespace UnrealBuildTool
@@ -172,6 +173,27 @@ namespace UnrealBuildTool
 			public string Value;
 		}
 
+		[DllImport("kernel32.dll", SetLastError=true)]
+		private static extern int GetShortPathName(string pathName, StringBuilder shortName, int cbShortName);
+
+		public static string GetShortPathName(string Path)
+		{
+			int BufferSize = GetShortPathName(Path, null, 0);
+			if (BufferSize == 0)
+			{
+				throw new BuildException("Unable to convert path {0} to 8.3 format", Path);
+			}
+
+			var Builder = new StringBuilder(BufferSize);
+			int ConversionResult = GetShortPathName(Path, Builder, BufferSize);
+			if (ConversionResult == 0)
+			{
+				throw new BuildException("Unable to convert path {0} to 8.3 format", Path);
+			}
+
+			return Builder.ToString();
+		}
+
 		/**
 		 * Sets the environment variables from the passed in batch file
 		 * 
@@ -194,14 +216,21 @@ namespace UnrealBuildTool
 
 					var EnvReaderBatchFileContent = new List<string>();
 
+					var EnvVarsToXMLExePath = Path.Combine( GetExecutingAssemblyDirectory(), "EnvVarsToXML.exe" );
+
+					// Convert every path to short filenames to ensure we don't accidentally write out a non-ASCII batch file
+					var ShortBatchFileName       = GetShortPathName(BatchFileName);
+					var ShortEnvOutputFileName   = GetShortPathName(EnvOutputFileName);
+					var ShortEnvVarsToXMLExePath = GetShortPathName(EnvVarsToXMLExePath);
+
 					// Run 'vcvars32.bat' (or similar x64 version) to set environment variables
-					EnvReaderBatchFileContent.Add( String.Format( "call \"{0}\"", BatchFileName ) );
+					EnvReaderBatchFileContent.Add( String.Format( "call \"{0}\"", ShortBatchFileName ) );
 
 					// Pipe all environment variables to a file where we can read them in.
 					// We use a separate executable which runs after the batch file because we want to capture
 					// the environment after it has been set, and there's no easy way of doing this, and parsing
 					// the output of the set command is problematic when the vars contain non-ASCII characters.
-					EnvReaderBatchFileContent.Add( String.Format( "\"{0}\" \"{1}\"", Path.Combine( GetExecutingAssemblyDirectory(), "EnvVarsToXML.exe" ), EnvOutputFileName ) );
+					EnvReaderBatchFileContent.Add( String.Format( "\"{0}\" \"{1}\"", ShortEnvVarsToXMLExePath, ShortEnvOutputFileName ) );
 
 					ResponseFile.Create( EnvReaderBatchFileName, EnvReaderBatchFileContent );
                 }
@@ -248,11 +277,18 @@ namespace UnrealBuildTool
 				Settings.CheckCharacters  = false;
 
 				List<EnvVar> EnvVars;
-				using (var Stream = new StreamReader(EnvOutputFileName))
-				using (var Reader = XmlReader.Create(Stream, Settings))
+				try
+				{
+					using (var Stream = new StreamReader(EnvOutputFileName))
+					using (var Reader = XmlReader.Create(Stream, Settings))
+					{
+						var Serializer = new XmlSerializer(typeof(List<EnvVar>));
+						EnvVars = (List<EnvVar>)Serializer.Deserialize(Reader);
+					}
+				}
+				catch (Exception e)
 				{
-					var Serializer = new XmlSerializer(typeof(List<EnvVar>));
-					EnvVars = (List<EnvVar>)Serializer.Deserialize(Reader);
+					throw new BuildException(e, "Failed to read environment variables from XML file: {0}", EnvOutputFileName);
 				}
 
 				foreach (var EnvVar in EnvVars)
