Missing file.

diff --git a/Engine/Source/Programs/UnrealBuildTool/Configuration/UEBuildTarget.cs b/Engine/Source/Programs/UnrealBuildTool/Configuration/UEBuildTarget.cs
--- a/Engine/Source/Programs/UnrealBuildTool/Configuration/UEBuildTarget.cs
+++ b/Engine/Source/Programs/UnrealBuildTool/Configuration/UEBuildTarget.cs
@@ -1344,7 +1344,7 @@ namespace UnrealBuildTool
 			foreach(UEBuildBinary Binary in AppBinaries)
 			{
 				BuildReceipt BinaryReceipt = Binary.MakeReceipt(ToolChain);
-				Receipt.Merge(BinaryReceipt);
+				Receipt.Append(BinaryReceipt);
 			}
 			foreach(BuildProduct BuildProduct in Receipt.BuildProducts)
 			{
