Fixed UnrealBuildTool log spew when cleaning targets

diff --git a/Engine/Source/Programs/UnrealBuildTool/System/FlatCPPIncludeDepencencyCache.cs b/Engine/Source/Programs/UnrealBuildTool/System/FlatCPPIncludeDepencencyCache.cs
--- a/Engine/Source/Programs/UnrealBuildTool/System/FlatCPPIncludeDepencencyCache.cs
+++ b/Engine/Source/Programs/UnrealBuildTool/System/FlatCPPIncludeDepencencyCache.cs
@@ -67,7 +67,8 @@ namespace UnrealBuildTool
 				}
 			}
 
-			if( !ProjectFileGenerator.bGenerateProjectFiles && !UnrealBuildTool.bNeedsFullCPPIncludeRescan )
+			bool bIsBuilding = (BuildConfiguration.bXGEExport == false) && (UEBuildConfiguration.bGenerateManifest == false) && (UEBuildConfiguration.bGenerateExternalFileList == false) && (UEBuildConfiguration.bCleanProject == false);
+			if( bIsBuilding && !UnrealBuildTool.bNeedsFullCPPIncludeRescan )
 			{ 
 				UnrealBuildTool.bNeedsFullCPPIncludeRescan = true;
 				Log.TraceInformation( "Performing full C++ include scan (no include cache file)" );
