Fixed UnrealBuildTool log spew when cleaning targets (part 2)

diff --git a/Engine/Source/Programs/UnrealBuildTool/System/FlatCPPIncludeDepencencyCache.cs b/Engine/Source/Programs/UnrealBuildTool/System/FlatCPPIncludeDepencencyCache.cs
--- a/Engine/Source/Programs/UnrealBuildTool/System/FlatCPPIncludeDepencencyCache.cs
+++ b/Engine/Source/Programs/UnrealBuildTool/System/FlatCPPIncludeDepencencyCache.cs
@@ -67,7 +67,7 @@ namespace UnrealBuildTool
 				}
 			}
 
-			bool bIsBuilding = (BuildConfiguration.bXGEExport == false) && (UEBuildConfiguration.bGenerateManifest == false) && (UEBuildConfiguration.bGenerateExternalFileList == false) && (UEBuildConfiguration.bCleanProject == false);
+			bool bIsBuilding = (ProjectFileGenerator.bGenerateProjectFiles == false ) && (BuildConfiguration.bXGEExport == false) && (UEBuildConfiguration.bGenerateManifest == false) && (UEBuildConfiguration.bGenerateExternalFileList == false) && (UEBuildConfiguration.bCleanProject == false);
 			if( bIsBuilding && !UnrealBuildTool.bNeedsFullCPPIncludeRescan )
 			{ 
 				UnrealBuildTool.bNeedsFullCPPIncludeRescan = true;
