Add more generic support for variable expansion in build receipts.

diff --git a/Engine/Source/Programs/UnrealBuildTool/Configuration/UEBuildTarget.cs b/Engine/Source/Programs/UnrealBuildTool/Configuration/UEBuildTarget.cs
--- a/Engine/Source/Programs/UnrealBuildTool/Configuration/UEBuildTarget.cs
+++ b/Engine/Source/Programs/UnrealBuildTool/Configuration/UEBuildTarget.cs
@@ -1308,7 +1308,11 @@ namespace UnrealBuildTool
 				Manifest = Utils.ReadClass<BuildManifest>(ManifestPath);
 			}
 
-			foreach(BuildProduct BuildProduct in Receipt.BuildProducts)
+			// Expand all the paths in the receipt; they'll currently use variables for the engine and project directories
+			BuildReceipt ReceiptWithFullPaths = new BuildReceipt(Receipt);
+			ReceiptWithFullPaths.ExpandPathVariables(BuildConfiguration.RelativeEnginePath, ProjectDirectory);
+
+			foreach(BuildProduct BuildProduct in ReceiptWithFullPaths.BuildProducts)
 			{
 				// If we're cleaning, don't add any precompiled binaries to the manifest. We don't want to delete them.
 				if(UEBuildConfiguration.bCleanProject && bUsePrecompiled && BuildProduct.IsPrecompiled)
@@ -1323,8 +1327,7 @@ namespace UnrealBuildTool
 				}
 
 				// Otherwise add it
-				string FullPath = BuildReceipt.ExpandPathVariables(BuildProduct.Path, BuildConfiguration.RelativeEnginePath, ProjectDirectory);
-				Manifest.AddBuildProduct(FullPath);
+				Manifest.AddBuildProduct(BuildProduct.Path);
 			}
 
 			if (UEBuildConfiguration.bCleanProject)
@@ -1346,10 +1349,7 @@ namespace UnrealBuildTool
 				BuildReceipt BinaryReceipt = Binary.MakeReceipt(ToolChain);
 				Receipt.Append(BinaryReceipt);
 			}
-			foreach(BuildProduct BuildProduct in Receipt.BuildProducts)
-			{
-				BuildProduct.Path = BuildReceipt.InsertPathVariables(BuildProduct.Path, BuildConfiguration.RelativeEnginePath, ProjectDirectory);
-			}
+			Receipt.InsertStandardPathVariables(BuildConfiguration.RelativeEnginePath, ProjectDirectory);
 		}
 
 		/** Writes the receipt for this target to disk */
@@ -1357,7 +1357,7 @@ namespace UnrealBuildTool
 		{
 			if(Receipt != null)
 			{
-				string ReceiptFileName = Path.Combine(ProjectDirectory, "Build", "Receipts", String.Format("{0}-{1}-{2}.target.xml", TargetName, Platform.ToString(), Configuration.ToString()));
+				string ReceiptFileName = BuildReceipt.GetDefaultPath(ProjectDirectory, TargetName, Platform, Configuration);
 				Directory.CreateDirectory(Path.GetDirectoryName(ReceiptFileName));
 				Receipt.Write(ReceiptFileName);
 			}
diff --git a/Engine/Source/Programs/UnrealBuildTool/System/BuildReceipt.cs b/Engine/Source/Programs/UnrealBuildTool/System/BuildReceipt.cs
--- a/Engine/Source/Programs/UnrealBuildTool/System/BuildReceipt.cs
+++ b/Engine/Source/Programs/UnrealBuildTool/System/BuildReceipt.cs
@@ -42,6 +42,13 @@ namespace UnrealBuildTool
 			Type = InType;
 		}
 
+		public BuildProduct(BuildProduct Other)
+		{
+			Path = Other.Path;
+			Type = Other.Type;
+			IsPrecompiled = Other.IsPrecompiled;
+		}
+
 		public override string ToString()
 		{
 			return Path;
@@ -90,9 +97,6 @@ namespace UnrealBuildTool
 	[Serializable]
 	public class BuildReceipt
 	{
-		const string EngineDirVariable = "$(EngineDir)";
-		const string ProjectDirVariable = "$(ProjectDir)";
-
 		[XmlArrayItem("BuildProduct")]
 		public List<BuildProduct> BuildProducts = new List<BuildProduct>();
 
@@ -106,6 +110,22 @@ namespace UnrealBuildTool
 		{
 		}
 
+		/// <summary>
+		/// Copy constructor
+		/// </summary>
+		/// <param name="InOther">Receipt to copy from</param>
+		public BuildReceipt(BuildReceipt Other)
+		{ 
+			foreach(BuildProduct OtherBuildProduct in Other.BuildProducts)
+			{
+				BuildProducts.Add(new BuildProduct(OtherBuildProduct));
+			}
+			foreach(RuntimeDependency OtherRuntimeDependency in Other.RuntimeDependencies)
+			{
+				RuntimeDependencies.Add(new RuntimeDependency(OtherRuntimeDependency));
+			}
+		}
+
 		/// <summary>
 		/// Adds a build product to the receipt. Does not check whether it already exists.
 		/// </summary>
@@ -149,54 +169,102 @@ namespace UnrealBuildTool
 		}
 
 		/// <summary>
-		/// Inserts $(EngineDir) and $(ProjectDir) variables into a path string, so it can be used on different machines.
+		/// Expand all the path variables in the manifest
 		/// </summary>
-		/// <param name="InputPath">Input path</param>
-		/// <param name="EngineDir">The engine directory. Relative paths are ok.</param>
-		/// <param name="ProjectDir">The project directory. Relative paths are ok.</param>
-		/// <returns>New string with the base directory replaced, or the original string</returns>
-		public static string InsertPathVariables(string InputPath, string EngineDir, string ProjectDir)
+		/// <param name="EngineDir">Value for the $(EngineDir) variable</param>
+		/// <param name="ProjectDir">Value for the $(ProjectDir) variable</param>
+		public void ExpandPathVariables(string EngineDir, string ProjectDir)
 		{
-			if(!InputPath.StartsWith("$("))
-			{
-				string FullInputPath = Path.GetFullPath(InputPath);
+			ExpandPathVariables(EngineDir, ProjectDir, new Dictionary<string, string>());
+		}
+
+		/// <summary>
+		/// Expand all the path variables in the manifest, including a list of supplied variable values.
+		/// </summary>
+		/// <param name="EngineDir">Value for the $(EngineDir) variable</param>
+		/// <param name="ProjectDir">Value for the $(ProjectDir) variable</param>
+		public void ExpandPathVariables(string EngineDir, string ProjectDir, IDictionary<string, string> OtherVariables)
+		{
+			// Build a dictionary containing the standard variable expansions
+			Dictionary<string, string> Variables = new Dictionary<string, string>(OtherVariables);
+			Variables["EngineDir"] = Path.GetFullPath(EngineDir).TrimEnd(Path.DirectorySeparatorChar);
+			Variables["ProjectDir"] = Path.GetFullPath(ProjectDir).TrimEnd(Path.DirectorySeparatorChar);
 
-				string FullEngineDir = Path.GetFullPath(EngineDir).TrimEnd(Path.DirectorySeparatorChar);
-				if(FullInputPath.StartsWith(FullEngineDir + Path.DirectorySeparatorChar, StringComparison.InvariantCultureIgnoreCase))
+			// Replace all the variables in the paths
+			foreach(BuildProduct BuildProduct in BuildProducts)
+			{
+				BuildProduct.Path = Utils.ExpandVariables(BuildProduct.Path, Variables);
+			}
+			foreach(RuntimeDependency RuntimeDependency in RuntimeDependencies)
+			{
+				RuntimeDependency.Path = Utils.ExpandVariables(RuntimeDependency.Path, Variables);
+				if(RuntimeDependency.StagePath != null)
 				{
-					return EngineDirVariable + FullInputPath.Substring(FullEngineDir.Length);
+					RuntimeDependency.StagePath = Utils.ExpandVariables(RuntimeDependency.StagePath, Variables);
 				}
+			}
+		}
 
-				string FullProjectDir = Path.GetFullPath(ProjectDir).TrimEnd(Path.DirectorySeparatorChar);
-				if(FullInputPath.StartsWith(FullProjectDir + Path.DirectorySeparatorChar, StringComparison.InvariantCultureIgnoreCase))
+		/// <summary>
+		/// Inserts standard $(EngineDir) and $(ProjectDir) variables into any path strings, so it can be used on different machines.
+		/// </summary>
+		/// <param name="EngineDir">The engine directory. Relative paths are ok.</param>
+		/// <param name="ProjectDir">The project directory. Relative paths are ok.</param>
+		public void InsertStandardPathVariables(string EngineDir, string ProjectDir)
+		{
+			string EnginePrefix = Path.GetFullPath(EngineDir).TrimEnd(Path.DirectorySeparatorChar) + Path.DirectorySeparatorChar;
+			string ProjectPrefix = Path.GetFullPath(ProjectDir).TrimEnd(Path.DirectorySeparatorChar) + Path.DirectorySeparatorChar;
+
+			foreach(BuildProduct BuildProduct in BuildProducts)
+			{
+				BuildProduct.Path = InsertStandardPathVariablesToString(BuildProduct.Path, EnginePrefix, ProjectPrefix);
+			}
+			foreach(RuntimeDependency RuntimeDependency in RuntimeDependencies)
+			{
+				RuntimeDependency.Path = InsertStandardPathVariablesToString(RuntimeDependency.Path, EnginePrefix, ProjectPrefix);
+				if(RuntimeDependency.StagePath != null)
 				{
-					return ProjectDirVariable + FullInputPath.Substring(FullProjectDir.Length);
+					RuntimeDependency.StagePath = InsertStandardPathVariablesToString(RuntimeDependency.StagePath, EnginePrefix, ProjectPrefix);
 				}
 			}
-			return InputPath;
 		}
 
 		/// <summary>
-		/// Expands any $(EngineDir) and $(ProjectDir) variables in the given path.
+		/// Inserts $(EngineDir) and $(ProjectDir) variables into a path string, so it can be used on different machines.
 		/// </summary>
 		/// <param name="InputPath">Input path</param>
-		/// <param name="EngineDir">The engine directory. Relative paths will be converted to absolute paths in the output string.</param>
-		/// <param name="ProjectDir">The project directory. Relative paths will be converted to absolute paths in the output string</param>
-		/// <returns>Full path with variable strings replaced, or the original string.</returns>
-		public static string ExpandPathVariables(string InputPath, string EngineDir, string ProjectDir)
+		/// <param name="EngineDir">The engine directory. Relative paths are ok.</param>
+		/// <param name="ProjectDir">The project directory. Relative paths are ok.</param>
+		/// <returns>New string with the base directory replaced, or the original string</returns>
+		static string InsertStandardPathVariablesToString(string InputPath, string EnginePrefix, string ProjectPrefix)
 		{
-			if(InputPath.StartsWith("$("))
+			string Result = InputPath;
+			if(!InputPath.StartsWith("$("))
 			{
-				if(InputPath.StartsWith(EngineDirVariable))
+				string FullInputPath = Path.GetFullPath(InputPath);
+				if(FullInputPath.StartsWith(EnginePrefix))
 				{
-					return Path.GetFullPath(EngineDir.TrimEnd('/', '\\') + InputPath.Substring(EngineDirVariable.Length));
+					Result = "$(EngineDir)" + FullInputPath.Substring(EnginePrefix.Length - 1);
 				}
-				else if(InputPath.StartsWith(ProjectDirVariable))
+				else if(FullInputPath.StartsWith(ProjectPrefix))
 				{
-					return Path.GetFullPath(ProjectDir.TrimEnd('/', '\\') + InputPath.Substring(ProjectDirVariable.Length));
+					Result = "$(ProjectDir)" + FullInputPath.Substring(ProjectPrefix.Length - 1);
 				}
 			}
-			return InputPath;
+			return Result;
+		}
+
+		/// <summary>
+		/// Returns the standard path to the build receipt for a given target
+		/// </summary>
+		/// <param name="DirectoryName">Base directory for the target being built; either the project directory or engine directory.</param>
+		/// <param name="TargetName">The target being built</param>
+		/// <param name="Configuration">The target configuration</param>
+		/// <param name="Platform">The target platform</param>
+		/// <returns>Path to the receipt for this target</returns>
+		public static string GetDefaultPath(string BaseDir, string TargetName, UnrealTargetPlatform Platform, UnrealTargetConfiguration Configuration)
+		{
+			return Path.Combine(BaseDir, "Build", "Receipts", String.Format("{0}-{1}-{2}.target.xml", TargetName, Platform.ToString(), Configuration.ToString()));
 		}
 
 		/// <summary>
diff --git a/Engine/Source/Programs/UnrealBuildTool/System/Utils.cs b/Engine/Source/Programs/UnrealBuildTool/System/Utils.cs
--- a/Engine/Source/Programs/UnrealBuildTool/System/Utils.cs
+++ b/Engine/Source/Programs/UnrealBuildTool/System/Utils.cs
@@ -76,7 +76,47 @@ namespace UnrealBuildTool
 				
 			return Result;
 		}
-				
+
+		/// <summary>
+		/// Expands variables in $(VarName) format in the given string. Variables are retrieved from the given dictionary, or through the environment of the current process.
+		/// Any unknown variables are ignored.
+		/// </summary>
+		/// <param name="InputString">String to search for variable names</param>
+		/// <param name="Variables">Lookup of variable names to values</param>
+		/// <returns>String with all variables replaced</returns>
+		public static string ExpandVariables(string InputString, Dictionary<string, string> AdditionalVariables = null)
+		{
+			string Result = InputString;
+			for(int Idx = Result.IndexOf("$("); Idx != -1; Idx = Result.IndexOf("$(", Idx))
+			{
+				// Find the end of the variable name
+				int EndIdx = Result.IndexOf(')', Idx + 2);
+				if(EndIdx == -1)
+				{
+					break;
+				}
+
+				// Extract the variable name from the string
+				string Name = Result.Substring(Idx + 2, EndIdx - (Idx + 2));
+
+				// Find the value for it, either from the dictionary or the environment block
+				string Value;
+				if(AdditionalVariables == null || !AdditionalVariables.TryGetValue(Name, out Value))
+				{
+					Value = Environment.GetEnvironmentVariable(Name);
+					if(Value == null)
+					{
+						Idx = EndIdx + 1;
+						continue;
+					}
+				}
+
+				// Replace the variable, or skip past it
+				Result = Result.Substring(0, Idx) + Value + Result.Substring(EndIdx + 1);
+			}
+			return Result;
+		}
+		
 		/**
 		 * This is a faster replacement of File.ReadAllText. Code snippet based on code 
 		 * and analysis by Sam Allen
