Alternative way of handling dylibs in Mac bundle subfolders that should work around the Mavericks dlopen() bug

diff --git a/Engine/Source/Programs/UnrealBuildTool/Mac/MacToolChain.cs b/Engine/Source/Programs/UnrealBuildTool/Mac/MacToolChain.cs
--- a/Engine/Source/Programs/UnrealBuildTool/Mac/MacToolChain.cs
+++ b/Engine/Source/Programs/UnrealBuildTool/Mac/MacToolChain.cs
@@ -55,8 +55,6 @@ namespace UnrealBuildTool
 
 		private static List<FileItem> BundleDependencies = new List<FileItem>();
 
-		private static List<string> BundleDylibPaths = new List<string>();
-
 		public List<string> BuiltBinaries = new List<string>();
 
 		public override void SetUpGlobalEnvironment()
@@ -688,12 +686,6 @@ namespace UnrealBuildTool
 
 			if (!bIsBuildingLibrary || LinkEnvironment.Config.bIncludeDependentLibrariesInLibrary)
 			{
-				// Add RPaths for other bundle modules
-				foreach (string BundleDylibPath in BundleDylibPaths)
-				{
-					LinkCommand += String.Format(" -rpath @executable_path/{0}/", BundleDylibPath);
-				}
-
 				// Add the additional libraries to the argument list.
 				foreach (string AdditionalLibrary in LinkEnvironment.Config.AdditionalLibraries)
 				{
@@ -847,7 +839,18 @@ namespace UnrealBuildTool
 			if (LinkEnvironment.Config.bIsBuildingDLL)
 			{
 				// Add the output file to the command-line.
-				LinkCommand += string.Format(" -install_name {0}/{1}", DylibsPath, Path.GetFileName(OutputFile.AbsolutePath));
+				string Filename = "";
+				int Index = OutputFile.AbsolutePath.LastIndexOf(".app/Contents/MacOS/");
+				if (Index > -1)
+				{
+					Index += ".app/Contents/MacOS/".Length;
+					Filename = OutputFile.AbsolutePath.Substring(Index);
+				}
+				else
+				{
+					Filename = Path.GetFileName(OutputFile.AbsolutePath);
+				}
+				LinkCommand += string.Format(" -install_name {0}/{1}", DylibsPath, Filename);
 			}
 
 			if (!bIsBuildingLibrary)
@@ -1295,8 +1298,6 @@ namespace UnrealBuildTool
 		{
 			base.FixBundleBinariesPaths(Target, Binaries);
 
-			BundleDylibPaths.Clear();
-
 			string BundleContentsPath = Target.OutputPath + ".app/Contents/";
 			foreach (UEBuildBinary Binary in Binaries)
 			{
@@ -1312,13 +1313,6 @@ namespace UnrealBuildTool
 						// get the subdir, which is the DylibDir - ExeDir
 						string SubDir = DylibDir.Replace(ExeDir, "");
 						Binary.Config.OutputFilePaths[0] = BundleContentsPath + "MacOS" + SubDir + "/" + BinaryFileName;
-
-						// Add the path to the list of search paths for the bundle
-						string TrimSubDir = SubDir.Trim('/', '\\');
-						if (TrimSubDir.Length > 0 && !BundleDylibPaths.Contains(TrimSubDir))
-						{
-							BundleDylibPaths.Add(TrimSubDir);
-						}
 					}
 				}
 				else if (!BinaryFileName.EndsWith(".a") && !Binary.Config.OutputFilePath.Contains(".app/Contents/MacOS/")) // Binaries can contain duplicates
diff --git a/Engine/Source/Runtime/Core/Private/Mac/MacPlatformProcess.cpp b/Engine/Source/Runtime/Core/Private/Mac/MacPlatformProcess.cpp
--- a/Engine/Source/Runtime/Core/Private/Mac/MacPlatformProcess.cpp
+++ b/Engine/Source/Runtime/Core/Private/Mac/MacPlatformProcess.cpp
@@ -19,10 +19,11 @@ void* FMacPlatformProcess::GetDllHandle( const TCHAR* Filename )
 
 	NSFileManager* FileManager = [NSFileManager defaultManager];
 	NSString* DylibPath = [NSString stringWithUTF8String:TCHAR_TO_UTF8(Filename)];
+	NSString* ExecutableFolder = [[[NSBundle mainBundle] executablePath] stringByDeletingLastPathComponent];
 	if (![FileManager fileExistsAtPath:DylibPath])
 	{
 		// If it's not a absolute or relative path, try to find the file in the app bundle
-		DylibPath = [[[[NSBundle mainBundle] executablePath] stringByDeletingLastPathComponent] stringByAppendingPathComponent:FString(Filename).GetNSString()];
+		DylibPath = [ExecutableFolder stringByAppendingPathComponent:FString(Filename).GetNSString()];
 	}
 
 	// Check if dylib is already loaded
@@ -30,7 +31,16 @@ void* FMacPlatformProcess::GetDllHandle( const TCHAR* Filename )
 	if (!Handle)
 	{
 		// Maybe it was loaded using RPATH
-		Handle = dlopen([[@"@rpath" stringByAppendingPathComponent:[DylibPath lastPathComponent]] fileSystemRepresentation], RTLD_NOLOAD | RTLD_LAZY | RTLD_LOCAL);
+		NSString* DylibName;
+		if ([DylibPath hasPrefix:ExecutableFolder])
+		{
+			DylibName = [DylibPath substringFromIndex:[ExecutableFolder length] + 1];
+		}
+		else
+		{
+			DylibName = [DylibPath lastPathComponent];
+		}
+		Handle = dlopen([[@"@rpath" stringByAppendingPathComponent:DylibName] fileSystemRepresentation], RTLD_NOLOAD | RTLD_LAZY | RTLD_LOCAL);
 	}
 	if (!Handle)
 	{
