Architecture added to manifest receipt name.
Should now match binary output name.

diff --git a/Engine/Source/Programs/UnrealBuildTool/Configuration/UEBuildTarget.cs b/Engine/Source/Programs/UnrealBuildTool/Configuration/UEBuildTarget.cs
--- a/Engine/Source/Programs/UnrealBuildTool/Configuration/UEBuildTarget.cs
+++ b/Engine/Source/Programs/UnrealBuildTool/Configuration/UEBuildTarget.cs
@@ -1330,7 +1330,8 @@ namespace UnrealBuildTool
 				Manifest.AddBuildProduct(BuildProduct.Path);
 			}
 
-			Manifest.AddBuildProduct(BuildReceipt.GetDefaultPath(ProjectDirectory, TargetName, Platform, Configuration));
+			IUEBuildPlatform BuildPlatform = UEBuildPlatform.GetBuildPlatform(Platform);
+			Manifest.AddBuildProduct(BuildReceipt.GetDefaultPath(ProjectDirectory, TargetName, Platform, Configuration, BuildPlatform.GetActiveArchitecture()));
 
 			if (UEBuildConfiguration.bCleanProject)
 			{
@@ -1359,7 +1360,8 @@ namespace UnrealBuildTool
 		{
 			if(Receipt != null)
 			{
-				string ReceiptFileName = BuildReceipt.GetDefaultPath(ProjectDirectory, TargetName, Platform, Configuration);
+				IUEBuildPlatform BuildPlatform = UEBuildPlatform.GetBuildPlatform(Platform);
+				string ReceiptFileName = BuildReceipt.GetDefaultPath(ProjectDirectory, TargetName, Platform, Configuration, BuildPlatform.GetActiveArchitecture());
 				Directory.CreateDirectory(Path.GetDirectoryName(ReceiptFileName));
 				Receipt.Write(ReceiptFileName);
 			}
diff --git a/Engine/Source/Programs/UnrealBuildTool/System/BuildReceipt.cs b/Engine/Source/Programs/UnrealBuildTool/System/BuildReceipt.cs
--- a/Engine/Source/Programs/UnrealBuildTool/System/BuildReceipt.cs
+++ b/Engine/Source/Programs/UnrealBuildTool/System/BuildReceipt.cs
@@ -262,9 +262,9 @@ namespace UnrealBuildTool
 		/// <param name="Configuration">The target configuration</param>
 		/// <param name="Platform">The target platform</param>
 		/// <returns>Path to the receipt for this target</returns>
-		public static string GetDefaultPath(string BaseDir, string TargetName, UnrealTargetPlatform Platform, UnrealTargetConfiguration Configuration)
+		public static string GetDefaultPath(string BaseDir, string TargetName, UnrealTargetPlatform Platform, UnrealTargetConfiguration Configuration, string BuildArchitecture)
 		{
-			return Path.Combine(BaseDir, "Build", "Receipts", String.Format("{0}-{1}-{2}.target.xml", TargetName, Platform.ToString(), Configuration.ToString()));
+			return Path.Combine(BaseDir, "Build", "Receipts", String.Format("{0}-{1}-{2}{3}.target.xml", TargetName, Platform.ToString(), Configuration.ToString(), BuildArchitecture));
 		}
 
 		/// <summary>
