Linux: gcc-like compiler output when not cross-compiling.

- Cross-compiling is unaffected (warnings/etc are clickable in MSVC), but this makes errors/warnings clickable in IDEs that expect gcc-like format.

diff --git a/Engine/Source/Programs/UnrealBuildTool/Linux/LinuxToolChain.cs b/Engine/Source/Programs/UnrealBuildTool/Linux/LinuxToolChain.cs
--- a/Engine/Source/Programs/UnrealBuildTool/Linux/LinuxToolChain.cs
+++ b/Engine/Source/Programs/UnrealBuildTool/Linux/LinuxToolChain.cs
@@ -309,7 +309,10 @@ namespace UnrealBuildTool
             else
             {
                 // Clang only options
-                Result += " -fdiagnostics-format=msvc";     // make diagnostics compatible with MSVC
+				if (CrossCompiling())
+				{
+					Result += " -fdiagnostics-format=msvc";     // make diagnostics compatible with MSVC when cross-compiling
+				}
                 Result += " -Wno-unused-private-field";     // MultichannelTcpSocket.h triggers this, possibly more
                 // this hides the "warning : comparison of unsigned expression < 0 is always false" type warnings due to constant comparisons, which are possible with template arguments
                 Result += " -Wno-tautological-compare";
@@ -493,37 +496,48 @@ namespace UnrealBuildTool
                 return;
             }
 
-            // Need to match following for clickable links
-            string RegexFilePath = @"^[A-Z]\:([\\\/][A-Za-z0-9_\-\.]*)+\.(cpp|c|mm|m|hpp|h)";
-            string RegexLineNumber = @"\:\d+\:\d+\:";
-            string RegexDescription = @"(\serror:\s|\swarning:\s|\snote:\s).*";
-
-            // Get Matches
-            string MatchFilePath = Regex.Match(Output, RegexFilePath).Value.Replace("Engine\\Source\\..\\..\\", "");
-            string MatchLineNumber = Regex.Match(Output, RegexLineNumber).Value;
-            string MatchDescription = Regex.Match(Output, RegexDescription).Value;
+			if (CrossCompiling())
+			{
+				// format the string so the output errors are clickable in Visual Studio
+
+				// Need to match following for clickable links
+				string RegexFilePath = @"^[A-Z]\:([\\\/][A-Za-z0-9_\-\.]*)+\.(cpp|c|mm|m|hpp|h)";
+				string RegexLineNumber = @"\:\d+\:\d+\:";
+				string RegexDescription = @"(\serror:\s|\swarning:\s|\snote:\s).*";
+
+				// Get Matches
+				string MatchFilePath = Regex.Match(Output, RegexFilePath).Value.Replace("Engine\\Source\\..\\..\\", "");
+				string MatchLineNumber = Regex.Match(Output, RegexLineNumber).Value;
+				string MatchDescription = Regex.Match(Output, RegexDescription).Value;
+
+				// If any of the above matches failed, do nothing
+				if (MatchFilePath.Length == 0 ||
+					MatchLineNumber.Length == 0 ||
+					MatchDescription.Length == 0)
+				{
+					Console.WriteLine(Output);
+					return;
+				}
 
-            // If any of the above matches failed, do nothing
-            if (MatchFilePath.Length == 0 ||
-                MatchLineNumber.Length == 0 ||
-                MatchDescription.Length == 0)
-            {
-                Console.WriteLine(Output);
-                return;
-            }
+				// Convert Path
+				string RegexStrippedPath = @"\\Engine\\.*"; //@"(Engine\/|[A-Za-z0-9_\-\.]*\/).*";
+				string ConvertedFilePath = Regex.Match(MatchFilePath, RegexStrippedPath).Value;
+				ConvertedFilePath = Path.GetFullPath("..\\.." + ConvertedFilePath);
 
-            // Convert Path
-            string RegexStrippedPath = @"\\Engine\\.*"; //@"(Engine\/|[A-Za-z0-9_\-\.]*\/).*";
-            string ConvertedFilePath = Regex.Match(MatchFilePath, RegexStrippedPath).Value;
-            ConvertedFilePath = Path.GetFullPath("..\\.." + ConvertedFilePath);
+				// Extract Line + Column Number
+				string ConvertedLineNumber = Regex.Match(MatchLineNumber, @"\d+").Value;
+				string ConvertedColumnNumber = Regex.Match(MatchLineNumber, @"(?<=:\d+:)\d+").Value;
 
-            // Extract Line + Column Number
-            string ConvertedLineNumber = Regex.Match(MatchLineNumber, @"\d+").Value;
-            string ConvertedColumnNumber = Regex.Match(MatchLineNumber, @"(?<=:\d+:)\d+").Value;
+				// Write output
+				string ConvertedExpression = "  " + ConvertedFilePath + "(" + ConvertedLineNumber + "," + ConvertedColumnNumber + "):" + MatchDescription;
+				Console.WriteLine(ConvertedExpression); // To create clickable vs link
+			}
+			else
+			{
+				// native platform tools expect this in stderror
 
-            // Write output
-            string ConvertedExpression = "  " + ConvertedFilePath + "(" + ConvertedLineNumber + "," + ConvertedColumnNumber + "):" + MatchDescription;
-            Console.WriteLine(ConvertedExpression); // To create clickable vs link
+				Console.Error.WriteLine(Output);
+			}
         }
 
         // cache the location of NDK tools
