Added NotForLicensees to another section of code that builds a path for UnrealCodeAnalyzer

diff --git a/Engine/Source/Programs/UnrealBuildTool/Windows/VCToolChain.cs b/Engine/Source/Programs/UnrealBuildTool/Windows/VCToolChain.cs
--- a/Engine/Source/Programs/UnrealBuildTool/Windows/VCToolChain.cs
+++ b/Engine/Source/Programs/UnrealBuildTool/Windows/VCToolChain.cs
@@ -1243,7 +1243,7 @@ namespace UnrealBuildTool
 				CompileAction.WorkingDirectory = Path.GetFullPath(".");
 				if (BuildConfiguration.bRunUnrealCodeAnalyzer)
 				{
-					CompileAction.CommandPath = System.IO.Path.Combine(CompileAction.WorkingDirectory, @"..", @"Binaries", @"Win32", @"UnrealCodeAnalyzer.exe");
+					CompileAction.CommandPath = System.IO.Path.Combine(CompileAction.WorkingDirectory, @"..", @"Binaries", @"Win32", @"NotForLicensees", @"UnrealCodeAnalyzer.exe");
 				}
 				else
 				{
