Allow Engine projects to be built when we are generating rocket projects so that they can include intellisense data for engine source.

diff --git a/Engine/Source/Programs/UnrealBuildTool/Configuration/UEBuildTarget.cs b/Engine/Source/Programs/UnrealBuildTool/Configuration/UEBuildTarget.cs
--- a/Engine/Source/Programs/UnrealBuildTool/Configuration/UEBuildTarget.cs
+++ b/Engine/Source/Programs/UnrealBuildTool/Configuration/UEBuildTarget.cs
@@ -1763,7 +1763,7 @@ namespace UnrealBuildTool
 			}
 
 			[REDACTED_EMAIL]: Will users be able to rebuild UnrealHeaderTool? NO
-			if (UnrealBuildTool.RunningRocket() && AppName != "UnrealHeaderTool")
+			if (!ProjectFileGenerator.bGenerateProjectFiles && UnrealBuildTool.RunningRocket() && AppName != "UnrealHeaderTool")
 			{
 				var FilteredBinaries = new List<UEBuildBinary>();
 
