Fix PS4 Morpheus link errrors.

diff --git a/Engine/Source/Programs/UnrealBuildTool/Configuration/UEBuildTarget.cs b/Engine/Source/Programs/UnrealBuildTool/Configuration/UEBuildTarget.cs
--- a/Engine/Source/Programs/UnrealBuildTool/Configuration/UEBuildTarget.cs
+++ b/Engine/Source/Programs/UnrealBuildTool/Configuration/UEBuildTarget.cs
@@ -2580,23 +2580,22 @@ namespace UnrealBuildTool
 			ExtraModuleNames.AddRange(PlatformExtraModules);			
 		}
 
-		/** Sets up the plugins for this target */
-		protected virtual void SetupPlugins()
+		public void FindValidPlugins(out List<PluginInfo> ValidPlugins, out List<string> EnabledPluginNames)
 		{
 			// Filter the plugins list by the current project
-			List<PluginInfo> ValidPlugins = new List<PluginInfo>();
-			foreach(PluginInfo Plugin in Plugins.AllPlugins)
+			ValidPlugins = new List<PluginInfo>();
+			foreach (PluginInfo Plugin in Plugins.AllPlugins)
 			{
-				if(Plugin.LoadedFrom != PluginInfo.LoadedFromType.GameProject || Plugin.Directory.StartsWith(ProjectDirectory, StringComparison.InvariantCultureIgnoreCase))
+				if (Plugin.LoadedFrom != PluginInfo.LoadedFromType.GameProject || Plugin.Directory.StartsWith(ProjectDirectory, StringComparison.InvariantCultureIgnoreCase))
 				{
 					ValidPlugins.Add(Plugin);
 				}
 			}
 
 			// Remove any plugins for platforms we don't have
-			foreach(UnrealTargetPlatform TargetPlatform in Enum.GetValues(typeof(UnrealTargetPlatform)))
+			foreach (UnrealTargetPlatform TargetPlatform in Enum.GetValues(typeof(UnrealTargetPlatform)))
 			{
-				if(UEBuildPlatform.GetBuildPlatform(TargetPlatform, true) == null)
+				if (UEBuildPlatform.GetBuildPlatform(TargetPlatform, true) == null)
 				{
 					string DirectoryFragment = String.Format("/{0}/", TargetPlatform.ToString());
 					ValidPlugins.RemoveAll(x => x.Directory.Replace('\\', '/').Contains(DirectoryFragment));
@@ -2604,7 +2603,7 @@ namespace UnrealBuildTool
 			}
 
 			// Build a list of enabled plugins
-			List<string> EnabledPluginNames = new List<string>(Rules.AdditionalPlugins);
+			EnabledPluginNames = new List<string>(Rules.AdditionalPlugins);
 
 			// Add the list of plugins enabled by default
 			if (UEBuildConfiguration.bCompileAgainstEngine)
@@ -2613,7 +2612,7 @@ namespace UnrealBuildTool
 			}
 
 			// Update the plugin list for game targets
-			if(TargetType != TargetRules.TargetType.Program && UnrealBuildTool.HasUProjectFile())
+			if (TargetType != TargetRules.TargetType.Program && UnrealBuildTool.HasUProjectFile())
 			{
 				// Enable all the game specific plugins by default
 				EnabledPluginNames.AddRange(ValidPlugins.Where(x => x.LoadedFrom == PluginInfo.LoadedFromType.GameProject).Select(x => x.Name));
@@ -2621,6 +2620,15 @@ namespace UnrealBuildTool
 				// Use the project settings to update the plugin list for this target
 				EnabledPluginNames = UProjectInfo.GetEnabledPlugins(UnrealBuildTool.GetUProjectFile(), EnabledPluginNames, Platform);
 			}
+		}
+		
+
+		/** Sets up the plugins for this target */
+		protected virtual void SetupPlugins()
+		{
+			List<PluginInfo> ValidPlugins;
+			List<string> EnabledPluginNames;
+			FindValidPlugins(out ValidPlugins, out EnabledPluginNames);
 
 			// Set the list of plugins we're dependent on
 			PluginInfo[] EnabledPlugins = ValidPlugins.Where(x => EnabledPluginNames.Contains(x.Name)).Distinct().ToArray();
