Fix UBT compilation when PS4 code is not present.

diff --git a/Engine/Source/Programs/UnrealBuildTool/UnrealBuildTool.csproj b/Engine/Source/Programs/UnrealBuildTool/UnrealBuildTool.csproj
--- a/Engine/Source/Programs/UnrealBuildTool/UnrealBuildTool.csproj
+++ b/Engine/Source/Programs/UnrealBuildTool/UnrealBuildTool.csproj
@@ -110,6 +110,7 @@
         <Compile Include="PS4\PS4ProjectGenerator.cs" />
         <Compile Include="PS4\PS4ToolChain.cs" />
         <Compile Include="PS4\UEBuildPS4.cs" />
+        <Compile Include="PS4\UEDeployPS4.cs" />
       </ItemGroup>
     </When>
   </Choose>
@@ -196,8 +197,7 @@
     <Compile Include="fastJSON\SafeDictionary.cs" />
     <Compile Include="Properties\AssemblyInfo.cs" />
     <Compile Include="Configuration\UEBuildConfiguration.cs" />
-    <Compile Include="Configuration\UEBuildTarget.cs" />
-    <Compile Include="PS4\UEDeployPS4.cs" />
+    <Compile Include="Configuration\UEBuildTarget.cs" />    
     <Compile Include="System\ActionGraph.cs" />
     <Compile Include="System\ActionHistory.cs" />
     <Compile Include="System\BuildException.cs" />
