Remove the output path override for modules. I can't see any code path that it would still trigger it, as FixupOutputPath() seems to require a "Binaries" substring, but the only possible values for OutputDirectory are (1) the compiler intermediate directory, or (2) a substring of the module rules path.

diff --git a/Engine/Source/Programs/UnrealBuildTool/Configuration/UEBuildBinary.cs b/Engine/Source/Programs/UnrealBuildTool/Configuration/UEBuildBinary.cs
--- a/Engine/Source/Programs/UnrealBuildTool/Configuration/UEBuildBinary.cs
+++ b/Engine/Source/Programs/UnrealBuildTool/Configuration/UEBuildBinary.cs
@@ -490,15 +490,6 @@ namespace UnrealBuildTool
 				// to the plugins.
 				Config.TargetConfiguration = Target.Configuration;
 				Config.TargetName = Target.GetAppName();
-
-				if (Module != null && (Target.Rules == null || Target.Rules.bOutputToEngineBinaries == false))
-				{
-					// Fix up the binary path if this is module specifies an alternate output directory
-					for (int Index = 0; Index < Config.OutputFilePaths.Length; Index++ )
-					{
-						Config.OutputFilePaths[Index] = Module.FixupOutputPath(Config.OutputFilePaths[Index]);
-					}
-				}
 			}
 		}
 
diff --git a/Engine/Source/Programs/UnrealBuildTool/Configuration/UEBuildModule.cs b/Engine/Source/Programs/UnrealBuildTool/Configuration/UEBuildModule.cs
--- a/Engine/Source/Programs/UnrealBuildTool/Configuration/UEBuildModule.cs
+++ b/Engine/Source/Programs/UnrealBuildTool/Configuration/UEBuildModule.cs
@@ -193,9 +193,6 @@ namespace UnrealBuildTool
 		/** Path to the module directory */
 		public readonly string ModuleDirectory;
 
-		/** An optional output path for the module */
-		public readonly string OutputDirectory;
-
 		/** Is this module allowed to be redistributed. */
 		private readonly bool? IsRedistributableOverride;
 
@@ -265,7 +262,6 @@ namespace UnrealBuildTool
 			string InName,
 			UEBuildModuleType InType,
 			string InModuleDirectory,
-			string InOutputDirectory,
 			bool? InIsRedistributableOverride,
 			IEnumerable<string> InPublicDefinitions = null,
 			IEnumerable<string> InPublicIncludePaths = null,
@@ -293,7 +289,6 @@ namespace UnrealBuildTool
 			Name = InName;
 			Type = InType;
 			ModuleDirectory = InModuleDirectory;
-			OutputDirectory = InOutputDirectory;
 			PublicDefinitions = HashSetFromOptionalEnumerableStringParameter(InPublicDefinitions);
 			PublicIncludePaths = HashSetFromOptionalEnumerableStringParameter(InPublicIncludePaths);
 			PublicSystemIncludePaths = HashSetFromOptionalEnumerableStringParameter(InPublicSystemIncludePaths);
@@ -581,29 +576,6 @@ namespace UnrealBuildTool
 			PublicDefinitions.Remove(InDefintion);
 		}
 
-		/** Return the (optional) output directory */
-		public string GetOutputDirectory()
-		{
-			return OutputDirectory;
-		}
-
-		/** Fix up the output path */
-		public virtual string FixupOutputPath(string InOutputPath)
-		{
-			string ModuleOutputPath = InOutputPath;
-			if (OutputDirectory.Length != 0)
-			{
-				// Use from 'Binaries/' on as it is setup w/ the correct platform already
-				int BinariesIndex = ModuleOutputPath.LastIndexOf("Binaries");
-				if (BinariesIndex != -1)
-				{
-					ModuleOutputPath = OutputDirectory + ModuleOutputPath.Substring(BinariesIndex);
-				}
-			}
-
-			return ModuleOutputPath;
-		}
-
 		/** Sets up the environment for compiling any module that includes the public interface of this module. */
 		protected virtual void SetupPublicCompileEnvironment(
 			UEBuildBinary SourceBinary,
@@ -934,7 +906,6 @@ namespace UnrealBuildTool
 			UEBuildModuleType InType,
 			string InName,
 			string InModuleDirectory,
-			string InOutputDirectory,
 			bool? InIsRedistributableOverride,
 			IEnumerable<string> InPublicDefinitions = null,
 			IEnumerable<string> InPublicIncludePaths = null,
@@ -955,7 +926,6 @@ namespace UnrealBuildTool
 			InType:							InType,
 			InName:							InName,
 			InModuleDirectory:				InModuleDirectory,
-			InOutputDirectory:				InOutputDirectory,
 			InIsRedistributableOverride:	InIsRedistributableOverride,
 			InPublicDefinitions:			InPublicDefinitions,
 			InPublicIncludePaths:			InPublicIncludePaths,
@@ -1138,7 +1108,6 @@ namespace UnrealBuildTool
 			string InName,
 			UEBuildModuleType InType,
 			string InModuleDirectory,
-			string InOutputDirectory,
 			bool? InIsRedistributableOverride,
 			IntelliSenseGatherer InIntelliSenseGatherer,
 			IEnumerable<FileItem> InSourceFiles,
@@ -1176,7 +1145,6 @@ namespace UnrealBuildTool
 					InName, 
 					InType,
 					InModuleDirectory,
-					InOutputDirectory,
 					InIsRedistributableOverride,
 					InDefinitions, 
 					InPublicIncludePaths,
@@ -2105,12 +2073,6 @@ namespace UnrealBuildTool
 																									InTargetConfiguration: Target.Configuration,
 																									bInCompileMonolithic: Target.ShouldCompileMonolithic() );
 
-								// Fix up the binary path if this is module specifies an alternate output directory
-								for (int Index = 0; Index < Config.OutputFilePaths.Length; Index++)
-								{
-									Config.OutputFilePaths[Index] = DependencyModule.FixupOutputPath(Config.OutputFilePaths[Index]);
-								}
-	
 								BinaryToBindTo = new UEBuildBinaryCPP( Target, Config );
 							}
 
@@ -2222,7 +2184,6 @@ namespace UnrealBuildTool
 			string InName,
 			UEBuildModuleType InType,
 			string InModuleDirectory,
-			string InOutputDirectory,
 			bool? InIsRedistributableOverride,
 			IntelliSenseGatherer InIntelliSenseGatherer,
 			IEnumerable<FileItem> InSourceFiles,
@@ -2256,7 +2217,7 @@ namespace UnrealBuildTool
 			bool InEnableExceptions,
 			bool bInBuildSourceFiles
 			)
-			: base(InTarget,InName,InType,InModuleDirectory,InOutputDirectory,InIsRedistributableOverride,InIntelliSenseGatherer,
+			: base(InTarget,InName,InType,InModuleDirectory,InIsRedistributableOverride,InIntelliSenseGatherer,
 			InSourceFiles,InPublicIncludePaths,InPublicSystemIncludePaths,null,InDefinitions,
 			InPublicIncludePathModuleNames,InPublicDependencyModuleNames,InPublicDelayLoadDLLs,InPublicAdditionalLibraries,InPublicFrameworks,InPublicWeakFrameworks,InPublicAdditionalFrameworks,InPublicAdditionalShadowFiles,InPublicAdditionalBundleResources,
 			InPrivateIncludePaths,InPrivateIncludePathModuleNames,InPrivateDependencyModuleNames,
diff --git a/Engine/Source/Programs/UnrealBuildTool/Configuration/UEBuildTarget.cs b/Engine/Source/Programs/UnrealBuildTool/Configuration/UEBuildTarget.cs
--- a/Engine/Source/Programs/UnrealBuildTool/Configuration/UEBuildTarget.cs
+++ b/Engine/Source/Programs/UnrealBuildTool/Configuration/UEBuildTarget.cs
@@ -2053,7 +2053,6 @@ namespace UnrealBuildTool
 				InName: Name,
 				InType: UEBuildModuleType.Game,
 				InModuleDirectory: Directory,
-				InOutputDirectory: GlobalCompileEnvironment.Config.OutputDirectory,
 				InIsRedistributableOverride: null,
 				InIntelliSenseGatherer: null,
 				InSourceFiles: SourceFiles.ToList(),
@@ -2999,7 +2998,6 @@ namespace UnrealBuildTool
 				// is NOT from the engine and therefore must be an application-specific module.
 				// Get the type of module we're creating
 				var ModuleType = UEBuildModuleType.Unknown;
-				var ApplicationOutputPath = "";
 				var ModuleFileRelativeToEngineDirectory = Utils.MakePathRelativeTo(ModuleFileName, ProjectFileGenerator.EngineRelativePath);
 
 				// see if it's external
@@ -3016,18 +3014,7 @@ namespace UnrealBuildTool
 						// not a plugin, see if it is a game module 
 						if (ModuleFileRelativeToEngineDirectory.StartsWith("..") || Path.IsPathRooted(ModuleFileRelativeToEngineDirectory))
 						{
-							ApplicationOutputPath = ModuleFileName;
-							int SourceIndex = ApplicationOutputPath.LastIndexOf(Path.DirectorySeparatorChar + "Source" + Path.DirectorySeparatorChar);
-							if (SourceIndex != -1)
-							{
-								ApplicationOutputPath = ApplicationOutputPath.Substring(0, SourceIndex + 1);
-								ModuleType = UEBuildModuleType.Game;
-							}
-							else
-							{
-								throw new BuildException("Unable to parse application directory for module '{0}' ({1}). Is it in '../../<APP>/Source'?",
-									ModuleName, ModuleFileName);
-							}
+							ModuleType = UEBuildModuleType.Game;
 						}
 						else
 						{
@@ -3103,7 +3090,7 @@ namespace UnrealBuildTool
 				}
 
 				// Now, go ahead and create the module builder instance
-				Module = InstantiateModule(RulesObject, ModuleName, ModuleType, ModuleDirectory, ApplicationOutputPath, IntelliSenseGatherer, FoundSourceFiles, bBuildFiles);
+				Module = InstantiateModule(RulesObject, ModuleName, ModuleType, ModuleDirectory, IntelliSenseGatherer, FoundSourceFiles, bBuildFiles);
 				if(Module == null)
 				{
 					throw new BuildException("Unrecognized module type specified by 'Rules' object {0}", RulesObject.ToString());
@@ -3129,7 +3116,6 @@ namespace UnrealBuildTool
 			string               ModuleName,
 			UEBuildModuleType    ModuleType,
 			string               ModuleDirectory,
-			string               ApplicationOutputPath,
 			IntelliSenseGatherer IntelliSenseGatherer,
 			List<FileItem>       ModuleSourceFiles,
 			bool                 bBuildSourceFiles)
@@ -3142,7 +3128,6 @@ namespace UnrealBuildTool
 							InName: ModuleName,
 							InType: ModuleType,
 							InModuleDirectory: ModuleDirectory,
-							InOutputDirectory: ApplicationOutputPath,
 							InIsRedistributableOverride: RulesObject.IsRedistributableOverride,
 							InIntelliSenseGatherer: IntelliSenseGatherer,
 							InSourceFiles: ModuleSourceFiles,
@@ -3183,7 +3168,6 @@ namespace UnrealBuildTool
 							InName: ModuleName,
 							InType: ModuleType,
 							InModuleDirectory: ModuleDirectory,
-							InOutputDirectory: ApplicationOutputPath,
 							InIsRedistributableOverride: RulesObject.IsRedistributableOverride,
 							InIntelliSenseGatherer: IntelliSenseGatherer,
 							InSourceFiles: ModuleSourceFiles,
@@ -3224,7 +3208,6 @@ namespace UnrealBuildTool
 							InName: ModuleName,
 							InType: ModuleType,
 							InModuleDirectory: ModuleDirectory,
-							InOutputDirectory: ApplicationOutputPath,
 							InIsRedistributableOverride: RulesObject.IsRedistributableOverride,
 							InPublicDefinitions: RulesObject.Definitions,
 							InPublicSystemIncludePaths: RulesObject.PublicSystemIncludePaths,
