Remove unnecessary variables from UEBuildBinaryConfiguration.

diff --git a/Engine/Source/Programs/UnrealBuildTool/Configuration/UEBuildBinary.cs b/Engine/Source/Programs/UnrealBuildTool/Configuration/UEBuildBinary.cs
--- a/Engine/Source/Programs/UnrealBuildTool/Configuration/UEBuildBinary.cs
+++ b/Engine/Source/Programs/UnrealBuildTool/Configuration/UEBuildBinary.cs
@@ -121,16 +121,6 @@ namespace UnrealBuildTool
 		/// </summary>
 		public bool bBuildAdditionalConsoleApp = false;
 
-		/// <summary>
-		/// The build target configuration being compiled
-		/// </summary>
-		public UnrealTargetConfiguration TargetConfiguration = UnrealTargetConfiguration.Development;
-
-		/// <summary>
-		/// The name of the target being compiled
-		/// </summary>
-		public string TargetName = "";
-
 		/// <summary>
 		/// The projectfile path
 		/// </summary>
@@ -162,8 +152,6 @@ namespace UnrealBuildTool
 				bool bInHasModuleRules = true,
                 bool bInIsCrossTarget = false,
                 bool bInCompileMonolithic = false,
-				UnrealTargetConfiguration InTargetConfiguration = UnrealTargetConfiguration.Development,
-				string InTargetName = "",
 				string InProjectFilePath = "",
 				List<string> InModuleNames = null
 			)
@@ -178,8 +166,6 @@ namespace UnrealBuildTool
 			bHasModuleRules = bInHasModuleRules;
             bIsCrossTarget = bInIsCrossTarget;
             bCompileMonolithic = bInCompileMonolithic;
-			TargetConfiguration = InTargetConfiguration;
-			TargetName = InTargetName;
 			ProjectFilePath = InProjectFilePath;
 			ModuleNames = InModuleNames;
 		}
@@ -484,12 +470,6 @@ namespace UnrealBuildTool
 				// We set whether the binary is being compiled monolithic here to know later - specifically
 				// when we are determining whether to use SharedPCHs or not for static lib builds of plugins.
 				Config.bCompileMonolithic = Target.ShouldCompileMonolithic();
-
-				// We also need to know what the actual build target configuration is later in the process
-				// where we do not have access to the Target itself... This is for generating the paths
-				// to the plugins.
-				Config.TargetConfiguration = Target.Configuration;
-				Config.TargetName = Target.GetAppName();
 			}
 		}
 
diff --git a/Engine/Source/Programs/UnrealBuildTool/Configuration/UEBuildModule.cs b/Engine/Source/Programs/UnrealBuildTool/Configuration/UEBuildModule.cs
--- a/Engine/Source/Programs/UnrealBuildTool/Configuration/UEBuildModule.cs
+++ b/Engine/Source/Programs/UnrealBuildTool/Configuration/UEBuildModule.cs
@@ -2068,9 +2068,7 @@ namespace UnrealBuildTool
 																									InIntermediateDirectory: IntermediateDirectory,
 																									bInAllowExports: true,
 																									InModuleNames: new List<string> { DependencyModule.Name },
-																									InTargetName: Target.GetAppName(),
 																									bInIsCrossTarget: PlatformSpecificDynamicallyLoadedModuleNames.Contains(DependencyName) && !DynamicallyLoadedModuleNames.Contains(DependencyName),
-																									InTargetConfiguration: Target.Configuration,
 																									bInCompileMonolithic: Target.ShouldCompileMonolithic() );
 
 								BinaryToBindTo = new UEBuildBinaryCPP( Target, Config );
diff --git a/Engine/Source/Programs/UnrealBuildTool/Mac/MacToolChain.cs b/Engine/Source/Programs/UnrealBuildTool/Mac/MacToolChain.cs
--- a/Engine/Source/Programs/UnrealBuildTool/Mac/MacToolChain.cs
+++ b/Engine/Source/Programs/UnrealBuildTool/Mac/MacToolChain.cs
@@ -1381,7 +1381,7 @@ namespace UnrealBuildTool
 				Receipt.AddBuildProduct(BundleContentsDirectory + "PkgInfo", BuildProductType.RequiredResource);
 				Receipt.AddBuildProduct(BundleContentsDirectory + "Resources/UE4.icns", BuildProductType.RequiredResource);
 
-				if (Binary.Config.TargetName.StartsWith("UE4Editor"))
+				if (Binary.Target.AppName.StartsWith("UE4Editor"))
 				{
 					Receipt.AddBuildProduct(BundleContentsDirectory + "Resources/UProject.icns", BuildProductType.RequiredResource);
 				}
