Remove monolithic setting from binary configuration, so we can make "monolithic-ness" more of a pick-and-mix affair. Only used by shared PCH handling in two places:

* Shared PCHs were disabled for static libraries in monolithic builds -> now disabled for static libraries in all builds
* Modules could not use their own shared PCH in non-monolithic builds due to exports -> now disabled if their binary has exports

diff --git a/Engine/Source/Programs/UnrealBuildTool/Configuration/UEBuildBinary.cs b/Engine/Source/Programs/UnrealBuildTool/Configuration/UEBuildBinary.cs
--- a/Engine/Source/Programs/UnrealBuildTool/Configuration/UEBuildBinary.cs
+++ b/Engine/Source/Programs/UnrealBuildTool/Configuration/UEBuildBinary.cs
@@ -110,11 +110,6 @@ namespace UnrealBuildTool
         /// </summary>
         public bool bIsCrossTarget = false;
 
-        /// <summary>
-		/// If true, the binary is being compiled as a monolithic build
-		/// </summary>
-		public bool bCompileMonolithic = false;
-
         /// <summary>
 		/// If true, creates an additional console application. Hack for Windows, where it's not possible to conditionally inherit a parent's console Window depending on how
 		/// the application is invoked; you have to link the same executable with a different subsystem setting.
@@ -165,7 +160,6 @@ namespace UnrealBuildTool
 			bAllowCompilation = bInAllowCompilation;
 			bHasModuleRules = bInHasModuleRules;
             bIsCrossTarget = bInIsCrossTarget;
-            bCompileMonolithic = bInCompileMonolithic;
 			ProjectFilePath = InProjectFilePath;
 			ModuleNames = InModuleNames;
 		}
@@ -466,10 +460,6 @@ namespace UnrealBuildTool
 						throw new BuildException("Module \"{0}\" linked into both {1} and {2}, which creates ambiguous linkage for dependents.", ModuleName, Module.Binary.Config.OutputFilePath, Config.OutputFilePath);
 					}
 				}
-
-				// We set whether the binary is being compiled monolithic here to know later - specifically
-				// when we are determining whether to use SharedPCHs or not for static lib builds of plugins.
-				Config.bCompileMonolithic = Target.ShouldCompileMonolithic();
 			}
 		}
 
@@ -746,7 +736,7 @@ namespace UnrealBuildTool
 				{
 					// Compile each module.
 					Log.TraceVerbose("Compile module: " + ModuleName);
-					LinkInputFiles = Module.Compile(CompileEnvironment, BinaryCompileEnvironment, Config.bCompileMonolithic);
+					LinkInputFiles = Module.Compile(CompileEnvironment, BinaryCompileEnvironment);
 
 					// NOTE: Because of 'Shared PCHs', in monolithic builds the same PCH file may appear as a link input
 					// multiple times for a single binary.  We'll check for that here, and only add it once.  This avoids
diff --git a/Engine/Source/Programs/UnrealBuildTool/Configuration/UEBuildModule.cs b/Engine/Source/Programs/UnrealBuildTool/Configuration/UEBuildModule.cs
--- a/Engine/Source/Programs/UnrealBuildTool/Configuration/UEBuildModule.cs
+++ b/Engine/Source/Programs/UnrealBuildTool/Configuration/UEBuildModule.cs
@@ -843,7 +843,7 @@ namespace UnrealBuildTool
 		}
 
 		/** Compiles the module, and returns a list of files output by the compiler. */
-		public abstract List<FileItem> Compile( CPPEnvironment GlobalCompileEnvironment, CPPEnvironment CompileEnvironment, bool bCompileMonolithic );
+		public abstract List<FileItem> Compile( CPPEnvironment GlobalCompileEnvironment, CPPEnvironment CompileEnvironment );
 		
 		// Object interface.
 		public override string ToString()
@@ -946,7 +946,7 @@ namespace UnrealBuildTool
 		}
 
 		// UEBuildModule interface.
-		public override List<FileItem> Compile(CPPEnvironment GlobalCompileEnvironment, CPPEnvironment CompileEnvironment, bool bCompileMonolithic)
+		public override List<FileItem> Compile(CPPEnvironment GlobalCompileEnvironment, CPPEnvironment CompileEnvironment)
 		{
 			return new List<FileItem>();
 		}
@@ -1191,7 +1191,7 @@ namespace UnrealBuildTool
 		}
 
 		// UEBuildModule interface.
-		public override List<FileItem> Compile(CPPEnvironment GlobalCompileEnvironment, CPPEnvironment CompileEnvironment, bool bCompileMonolithic)
+		public override List<FileItem> Compile(CPPEnvironment GlobalCompileEnvironment, CPPEnvironment CompileEnvironment)
 		{
 			var BuildPlatform = UEBuildPlatform.GetBuildPlatformForCPPTargetPlatform(CompileEnvironment.Config.Target.Platform);
 
@@ -1316,13 +1316,12 @@ namespace UnrealBuildTool
 				// In the case of a shared PCH, we also need to keep track of which module that PCH's header file is a member of
 				string SharedPCHModuleName = String.Empty;
 
-				bool bDisableSharedPCHFiles = (Binary.Config.bCompileMonolithic && CompileEnvironment.Config.bIsBuildingLibrary);
-				if( BuildConfiguration.bUseSharedPCHs && bDisableSharedPCHFiles )
+				if( BuildConfiguration.bUseSharedPCHs && CompileEnvironment.Config.bIsBuildingLibrary )
 				{
-					Log.TraceVerbose("Module '{0}' was not allowed to use Shared PCHs, because we're compiling to a library in monolithic mode", this.Name );
+					Log.TraceVerbose("Module '{0}' was not allowed to use Shared PCHs, because we're compiling to a library", this.Name );
 				}
 
-				bool bUseSharedPCHFiles  = BuildConfiguration.bUseSharedPCHs && (bDisableSharedPCHFiles == false) && GlobalCompileEnvironment.SharedPCHHeaderFiles.Count > 0;
+				bool bUseSharedPCHFiles  = BuildConfiguration.bUseSharedPCHs && !CompileEnvironment.Config.bIsBuildingLibrary && GlobalCompileEnvironment.SharedPCHHeaderFiles.Count > 0;
 
 				if( bUseSharedPCHFiles )
 				{
@@ -1333,9 +1332,9 @@ namespace UnrealBuildTool
 						SharingPCHHeaderFilePath = Path.GetFullPath( Path.Combine( ProjectFileGenerator.RootRelativePath, "Engine", "Source", this.SharedPCHHeaderFile ) );
 					}
 
-					// When compiling in modular mode, we can't use a shared PCH file when compiling a module
+					// We can't use a shared PCH file when compiling a module
 					// with exports, because the shared PCH can only have imports in it to work correctly.
-					bool bCanModuleUseOwnSharedPCH = bAllowSharedPCH && bIsASharedPCHModule && bCompileMonolithic && ProcessedDependencies.UniquePCHHeaderFile.AbsolutePath.Equals( SharingPCHHeaderFilePath, StringComparison.InvariantCultureIgnoreCase );
+					bool bCanModuleUseOwnSharedPCH = bAllowSharedPCH && bIsASharedPCHModule && !Binary.Config.bAllowExports && ProcessedDependencies.UniquePCHHeaderFile.AbsolutePath.Equals( SharingPCHHeaderFilePath, StringComparison.InvariantCultureIgnoreCase );
 					if( bAllowSharedPCH && ( !bIsASharedPCHModule || bCanModuleUseOwnSharedPCH ) )
 					{
 						// Figure out which shared PCH tier we're in
@@ -2227,7 +2226,7 @@ namespace UnrealBuildTool
 		}
 
 		// UEBuildModule interface.
-		public override List<FileItem> Compile( CPPEnvironment GlobalCompileEnvironment, CPPEnvironment CompileEnvironment, bool bCompileMonolithic )
+		public override List<FileItem> Compile( CPPEnvironment GlobalCompileEnvironment, CPPEnvironment CompileEnvironment )
 		{
 			var ModuleCLREnvironment = CompileEnvironment.DeepCopy();
 
@@ -2241,7 +2240,7 @@ namespace UnrealBuildTool
 			}
 
 			// Pass the CLR compilation environment to the standard C++ module compilation code.
-			return base.Compile(GlobalCompileEnvironment, ModuleCLREnvironment, bCompileMonolithic );
+			return base.Compile(GlobalCompileEnvironment, ModuleCLREnvironment );
 		}
 
 		public override void SetupPrivateLinkEnvironment(
