Remove unused parameter to UEBuildBinaryConfiguration constructor for bCompileMonolithic.

diff --git a/Engine/Source/Programs/UnrealBuildTool/Configuration/UEBuildBinary.cs b/Engine/Source/Programs/UnrealBuildTool/Configuration/UEBuildBinary.cs
--- a/Engine/Source/Programs/UnrealBuildTool/Configuration/UEBuildBinary.cs
+++ b/Engine/Source/Programs/UnrealBuildTool/Configuration/UEBuildBinary.cs
@@ -146,7 +146,6 @@ namespace UnrealBuildTool
 				bool bInAllowCompilation = true,
 				bool bInHasModuleRules = true,
                 bool bInIsCrossTarget = false,
-                bool bInCompileMonolithic = false,
 				string InProjectFilePath = "",
 				List<string> InModuleNames = null
 			)
diff --git a/Engine/Source/Programs/UnrealBuildTool/Configuration/UEBuildModule.cs b/Engine/Source/Programs/UnrealBuildTool/Configuration/UEBuildModule.cs
--- a/Engine/Source/Programs/UnrealBuildTool/Configuration/UEBuildModule.cs
+++ b/Engine/Source/Programs/UnrealBuildTool/Configuration/UEBuildModule.cs
@@ -2067,8 +2067,7 @@ namespace UnrealBuildTool
 																									InIntermediateDirectory: IntermediateDirectory,
 																									bInAllowExports: true,
 																									InModuleNames: new List<string> { DependencyModule.Name },
-																									bInIsCrossTarget: PlatformSpecificDynamicallyLoadedModuleNames.Contains(DependencyName) && !DynamicallyLoadedModuleNames.Contains(DependencyName),
-																									bInCompileMonolithic: Target.ShouldCompileMonolithic() );
+																									bInIsCrossTarget: PlatformSpecificDynamicallyLoadedModuleNames.Contains(DependencyName) && !DynamicallyLoadedModuleNames.Contains(DependencyName) );
 
 								BinaryToBindTo = new UEBuildBinaryCPP( Target, Config );
 							}
diff --git a/Engine/Source/Programs/UnrealBuildTool/Configuration/UEBuildTarget.cs b/Engine/Source/Programs/UnrealBuildTool/Configuration/UEBuildTarget.cs
--- a/Engine/Source/Programs/UnrealBuildTool/Configuration/UEBuildTarget.cs
+++ b/Engine/Source/Programs/UnrealBuildTool/Configuration/UEBuildTarget.cs
@@ -2410,7 +2410,7 @@ namespace UnrealBuildTool
 			}
 
 			// Get the binary configuration
-			UEBuildBinaryConfiguration BinaryConfig = new UEBuildBinaryConfiguration(BinaryType, OutputFilePaths, IntermediateDirectory, true, bInAllowCompilation: bPrecompile, bInCompileMonolithic: bCompileMonolithic, InModuleNames: new List<string>{ ModuleName });
+			UEBuildBinaryConfiguration BinaryConfig = new UEBuildBinaryConfiguration(BinaryType, OutputFilePaths, IntermediateDirectory, true, bInAllowCompilation: bPrecompile, InModuleNames: new List<string>{ ModuleName });
 
 			// Create the binary
 			UEBuildBinaryCPP Binary = new UEBuildBinaryCPP(this, BinaryConfig);
