Fix UBT output path for game binaries.

diff --git a/Engine/Source/Programs/UnrealBuildTool/Configuration/UEBuildTarget.cs b/Engine/Source/Programs/UnrealBuildTool/Configuration/UEBuildTarget.cs
--- a/Engine/Source/Programs/UnrealBuildTool/Configuration/UEBuildTarget.cs
+++ b/Engine/Source/Programs/UnrealBuildTool/Configuration/UEBuildTarget.cs
@@ -2500,7 +2500,15 @@ namespace UnrealBuildTool
 			}
 			else
 			{
-				BinariesDirName = Path.Combine( "..", "Binaries" );
+				string ModuleFileName = RulesCompiler.GetModuleFilename(ModuleName);
+				if(!String.IsNullOrEmpty(ModuleFileName) && UnrealBuildTool.HasUProjectFile() && Utils.IsFileUnderDirectory(ModuleFileName, UnrealBuildTool.GetUProjectPath()))
+				{
+					BinariesDirName = Path.Combine( UnrealBuildTool.GetUProjectPath(), "Binaries" );
+				}
+				else
+				{
+					BinariesDirName = Path.Combine( "..", "Binaries" );
+				}
 			}
 			var BaseDirectory = Path.Combine( BinariesDirName, Platform.ToString());
 			if (ModuleBinariesSubDir.Length > 0)
