Change a couple of XmlSerializers to cache an object created with XmlSerializer.FromTypes() rather than just constructing it. Constructor throws and catches a FileNotFoundException internally, which is really annoying when you have break on exceptions turned on.

diff --git a/Engine/Source/Programs/UnrealBuildTool/System/BuildReceipt.cs b/Engine/Source/Programs/UnrealBuildTool/System/BuildReceipt.cs
--- a/Engine/Source/Programs/UnrealBuildTool/System/BuildReceipt.cs
+++ b/Engine/Source/Programs/UnrealBuildTool/System/BuildReceipt.cs
@@ -329,13 +329,14 @@ namespace UnrealBuildTool
 			return Path.Combine(BaseDir, "Build", "Receipts", String.Format("{0}-{1}-{2}{3}.target.xml", TargetName, Platform.ToString(), Configuration.ToString(), BuildArchitecture));
 		}
 
+		static XmlSerializer Serializer = XmlSerializer.FromTypes(new Type[]{ typeof(BuildReceipt) })[0];
+
 		/// <summary>
 		/// Read a receipt from disk.
 		/// </summary>
 		/// <param name="FileName">Filename to read from</param>
 		public static BuildReceipt Read(string FileName)
 		{
-			XmlSerializer Serializer = new XmlSerializer(typeof(BuildReceipt));
 			using(StreamReader Reader = new StreamReader(FileName))
 			{
 				return (BuildReceipt)Serializer.Deserialize(Reader);
@@ -348,7 +349,6 @@ namespace UnrealBuildTool
 		/// <param name="FileName">Output filename</param>
 		public void Write(string FileName)
 		{
-			XmlSerializer Serializer = new XmlSerializer(typeof(BuildReceipt));
 			using(StreamWriter Writer = new StreamWriter(FileName))
 			{
 				Serializer.Serialize(Writer, this);
diff --git a/Engine/Source/Programs/UnrealBuildTool/System/Utils.cs b/Engine/Source/Programs/UnrealBuildTool/System/Utils.cs
--- a/Engine/Source/Programs/UnrealBuildTool/System/Utils.cs
+++ b/Engine/Source/Programs/UnrealBuildTool/System/Utils.cs
@@ -204,6 +204,7 @@ namespace UnrealBuildTool
 		}
 
 		[Serializable]
+		[DebuggerDisplay("\\{{Key}={Value}\\}")]
 		public struct EnvVar
 		{
 			[XmlAttribute("Key")]
@@ -213,6 +214,8 @@ namespace UnrealBuildTool
 			public string Value;
 		}
 
+		private static XmlSerializer EnvVarListSerializer = XmlSerializer.FromTypes(new Type[]{ typeof(List<EnvVar>) })[0];
+
 		[DllImport("kernel32.dll", SetLastError=true)]
 		private static extern int GetShortPathName(string pathName, StringBuilder shortName, int cbShortName);
 
@@ -322,8 +325,7 @@ namespace UnrealBuildTool
 					using (var Stream = new StreamReader(EnvOutputFileName))
 					using (var Reader = XmlReader.Create(Stream, Settings))
 					{
-						var Serializer = new XmlSerializer(typeof(List<EnvVar>));
-						EnvVars = (List<EnvVar>)Serializer.Deserialize(Reader);
+						EnvVars = (List<EnvVar>)EnvVarListSerializer.Deserialize(Reader);
 					}
 				}
 				catch (Exception e)
