Reduce the number of codepaths used for creating new binary modules.

diff --git a/Engine/Source/Programs/UnrealBuildTool/Configuration/UEBuildBinary.cs b/Engine/Source/Programs/UnrealBuildTool/Configuration/UEBuildBinary.cs
--- a/Engine/Source/Programs/UnrealBuildTool/Configuration/UEBuildBinary.cs
+++ b/Engine/Source/Programs/UnrealBuildTool/Configuration/UEBuildBinary.cs
@@ -160,7 +160,10 @@ namespace UnrealBuildTool
 			bHasModuleRules = bInHasModuleRules;
             bIsCrossTarget = bInIsCrossTarget;
 			ProjectFilePath = InProjectFilePath;
-			ModuleNames = InModuleNames;
+			if(InModuleNames != null)
+			{
+				ModuleNames.AddRange(InModuleNames);
+			}
 		}
 	}
 
@@ -249,11 +252,8 @@ namespace UnrealBuildTool
 		/// Process all modules that aren't yet bound, creating binaries for modules that don't yet have one (if needed),
 		/// and updating modules for circular dependencies.
 		/// </summary>
-		/// <param name="ExecutableBinary">The executable binary, which links against all unbound modules when building monolithically</param>
-		/// <returns>List of newly-created binaries (may be empty)</returns>
-		public virtual List<UEBuildBinary> ProcessUnboundModules(UEBuildBinary ExecutableBinary)
+		public virtual void ProcessUnboundModules()
 		{
-			return null;
 		}
 
 		/// <summary>
@@ -402,11 +402,7 @@ namespace UnrealBuildTool
 	/// </summary>
 	public class UEBuildBinaryCPP : UEBuildBinary
 	{
-		public HashSet<string> ModuleNames
-		{
-			get;
-			private set;
-		}
+		public readonly List<string> ModuleNames;
 		private bool bCreateImportLibrarySeparately;
 		private bool bIncludeDependentLibrariesInLibrary;
 
@@ -417,7 +413,7 @@ namespace UnrealBuildTool
 		public UEBuildBinaryCPP( UEBuildTarget InTarget, UEBuildBinaryConfiguration InConfig )
 			: base( InTarget, InConfig )
 		{
-			ModuleNames = new HashSet<string>(InConfig.ModuleNames);
+			ModuleNames = new List<string>(InConfig.ModuleNames);
 			bCreateImportLibrarySeparately = InConfig.bCreateImportLibrarySeparately;
 			bIncludeDependentLibrariesInLibrary = InConfig.bIncludeDependentLibrariesInLibrary;
 		}
@@ -493,43 +489,18 @@ namespace UnrealBuildTool
 		/// Process all modules that aren't yet bound, creating binaries for modules that don't yet have one (if needed),
 		/// and updating modules for circular dependencies.
 		/// </summary>
-		/// <param name="ExecutableBinary">The executable binary, which links against all unbound modules when building monolithically</param>
 		/// <returns>List of newly-created binaries (may be empty)</returns>
-		public override List<UEBuildBinary> ProcessUnboundModules(UEBuildBinary ExecutableBinary)
+		public override void ProcessUnboundModules()
 		{
-			var Binaries = new Dictionary<string, UEBuildBinary>( StringComparer.InvariantCultureIgnoreCase );
 			if (Config.bHasModuleRules)
 			{
-				foreach (var ModuleName in ModuleNames)
-				{
-					var Module = Target.FindOrCreateModuleByName(ModuleName);
-					Module.RecursivelyProcessUnboundModules(Target, Binaries, ExecutableBinary);
-				}
-			}
-			else
-			{
-				// There's only one module in this case, so just bind it to this binary
-				foreach (var ModuleName in ModuleNames)
+				// Modules may be added to this binary during this process, so don't foreach over ModuleNames
+				for(int Idx = 0; Idx < ModuleNames.Count; Idx++)
 				{
-					Binaries.Add(ModuleName, this);
-				}
-			}
-
-			// Now build a final list of newly-created binaries that were bound to.  The hash may contain duplicates, so
-			// we filter those out here.
-			var BinaryList = new List<UEBuildBinary>();
-			foreach( var CurBinary in Binaries.Values )
-			{
-				// Never include ourselves in the new binary list (monolithic case)
-				if( CurBinary != this )
-				{
-					if( !BinaryList.Contains( CurBinary ) )
-					{
-						BinaryList.Add( CurBinary );
-					}
+					UEBuildModule Module = Target.FindOrCreateModuleByName(ModuleNames[Idx]);
+					Module.RecursivelyProcessUnboundModules();
 				}
 			}
-			return BinaryList;
 		}
 
 		/// <summary>
diff --git a/Engine/Source/Programs/UnrealBuildTool/Configuration/UEBuildModule.cs b/Engine/Source/Programs/UnrealBuildTool/Configuration/UEBuildModule.cs
--- a/Engine/Source/Programs/UnrealBuildTool/Configuration/UEBuildModule.cs
+++ b/Engine/Source/Programs/UnrealBuildTool/Configuration/UEBuildModule.cs
@@ -876,14 +876,8 @@ namespace UnrealBuildTool
 
 		/**
 		 * Gathers and binds binaries for this module and all of it's dependent modules
-		 *
-		 * @param	Target				The target we are currently building
-		 * @param	Binaries			Dictionary of all binaries we've gathered so far
-		 * @param	ExecutableBinary	The binary for the target's main executable (used only when linking monolithically)
-		 * @param	bBuildOnlyModules	True to build only specific modules, false for all
-		 * @param	ModulesToBuild		The specific modules to build
 		 */
-		public virtual void RecursivelyProcessUnboundModules(UEBuildTarget Target, Dictionary<string, UEBuildBinary> Binaries, UEBuildBinary ExecutableBinary)
+		public virtual void RecursivelyProcessUnboundModules()
 		{
 		}
 
@@ -2016,7 +2010,7 @@ namespace UnrealBuildTool
 			}
 		}
 
-		public override void RecursivelyProcessUnboundModules(UEBuildTarget Target, Dictionary<string, UEBuildBinary> Binaries, UEBuildBinary ExecutableBinary)
+		public override void RecursivelyProcessUnboundModules()
 		{
 			try
 			{
@@ -2039,48 +2033,18 @@ namespace UnrealBuildTool
 					// Skip modules that are included with the target (externals)
 					if( !DependencyModule.bIncludedInTarget )
 					{
-						if( !Binaries.ContainsKey( DependencyModule.Name ) )
-						{
-							UEBuildBinary BinaryToBindTo;
-							if (Target.ShouldCompileMonolithic())
-							{
-								// When linking monolithically, any unbound modules will be linked into the main executable
-								BinaryToBindTo = ExecutableBinary;
-							}
-							else
-							{
-								// Is this a plugin module?
-								var PluginInfo = Plugins.GetPluginInfoForModule( DependencyName );
+						bool bIsCrossTarget = PlatformSpecificDynamicallyLoadedModuleNames.Contains(DependencyName) && !DynamicallyLoadedModuleNames.Contains(DependencyName);
 
-								string[] OutputFilePaths = Target.MakeBinaryPaths(DependencyModule.Name, Target.GetAppName() + "-" + DependencyModule.Name, UEBuildBinaryType.DynamicLinkLibrary, Target.TargetType, PluginInfo, "");
+						// Get the binary that this module should be bound to
+						UEBuildBinary BinaryToBindTo = Target.FindOrAddBinaryForModule(DependencyName, bIsCrossTarget);
+						BinaryToBindTo.AddModule(DependencyName);
 
-								// If it's an engine module, output intermediates to the engine intermediates directory. 
-								string IntermediateDirectory = Binary.Config.IntermediateDirectory;
-								if (IntermediateDirectory != Target.EngineIntermediateDirectory && Path.GetFullPath(DependencyModule.ModuleDirectory).StartsWith(Path.GetFullPath(BuildConfiguration.RelativeEnginePath)))
-								{
-									IntermediateDirectory = Target.EngineIntermediateDirectory;
-								}
+						// Bind this module
+						DependencyModule.Binary = BinaryToBindTo;
+						DependencyModule.bIncludedInTarget = true;
 
-								// When using modular linkage, unbound modules will be linked into their own DLL files
-								UEBuildBinaryConfiguration Config = new UEBuildBinaryConfiguration( InType: UEBuildBinaryType.DynamicLinkLibrary,
-																									InOutputFilePaths: OutputFilePaths,
-																									InIntermediateDirectory: IntermediateDirectory,
-																									bInAllowExports: true,
-																									InModuleNames: new List<string> { DependencyModule.Name },
-																									bInIsCrossTarget: PlatformSpecificDynamicallyLoadedModuleNames.Contains(DependencyName) && !DynamicallyLoadedModuleNames.Contains(DependencyName) );
-
-								BinaryToBindTo = new UEBuildBinaryCPP( Target, Config );
-							}
-
-							Binaries[ DependencyModule.Name ] = BinaryToBindTo;
-
-							// Bind this module
-							DependencyModule.Binary = BinaryToBindTo;
-							DependencyModule.bIncludedInTarget = true;
-
-							// Also add binaries for this module's dependencies
-							DependencyModule.RecursivelyProcessUnboundModules( Target, Binaries, ExecutableBinary );
-						}
+						// Also add binaries for this module's dependencies
+						DependencyModule.RecursivelyProcessUnboundModules();
 					}
 
 					if (Target.ShouldCompileMonolithic() == false)
diff --git a/Engine/Source/Programs/UnrealBuildTool/Configuration/UEBuildTarget.cs b/Engine/Source/Programs/UnrealBuildTool/Configuration/UEBuildTarget.cs
--- a/Engine/Source/Programs/UnrealBuildTool/Configuration/UEBuildTarget.cs
+++ b/Engine/Source/Programs/UnrealBuildTool/Configuration/UEBuildTarget.cs
@@ -1681,23 +1681,11 @@ namespace UnrealBuildTool
 				}
 			}
 
-			// Process all referenced modules and create new binaries for DLL dependencies if needed
-			var NewBinaries = new List<UEBuildBinary>();
-			foreach (var Binary in AppBinaries)
+			// Process all referenced modules and create new binaries for DLL dependencies if needed. The AppBinaries 
+			// list may have entries added to it as modules are bound, so make sure we handle those too.
+			for(int Idx = 0; Idx < AppBinaries.Count; Idx++)
 			{
-				// Add binaries for all of our dependent modules
-				var FoundBinaries = Binary.ProcessUnboundModules(AppBinaries[0]);
-				if (FoundBinaries != null)
-				{
-					NewBinaries.AddRange(FoundBinaries);
-				}
-			}
-			foreach(var NewBinary in NewBinaries)
-			{
-				if(!AppBinaries.Contains(NewBinary))
-				{
-					AppBinaries.Add(NewBinary);
-				}
+				AppBinaries[Idx].ProcessUnboundModules();
 			}
 
 			// On Mac AppBinaries paths for non-console targets need to be adjusted to be inside the app bundle
@@ -2195,67 +2183,12 @@ namespace UnrealBuildTool
 			{
 				if (ShouldIncludePluginModule(Plugin, Module))
 				{
-                    AddPluginModule(Plugin, Module);
+					UEBuildBinaryType BinaryType = ShouldCompileMonolithic() ? UEBuildBinaryType.StaticLibrary : UEBuildBinaryType.DynamicLinkLibrary;
+					AddBinaryForModule(Module.Name, BinaryType, true, false);
 				}
 			}
 		}
 
-		/// <summary>
-		/// Include the given plugin module in the target. Will be built in the appropriate subfolder under the plugin directory.
-		/// </summary>
-        public void AddPluginModule(PluginInfo Plugin, PluginInfo.PluginModuleInfo Module)
-		{
-			bool bCompileMonolithic = ShouldCompileMonolithic();
-
-			// Get the binary type to build
-			UEBuildBinaryType BinaryType = bCompileMonolithic ? UEBuildBinaryType.StaticLibrary : UEBuildBinaryType.DynamicLinkLibrary;
-
-			// Get the output path. Don't prefix the app name for Rocket
-			string[] OutputFilePaths;
-			if (UnrealBuildTool.RunningRocket() && bCompileMonolithic)
-			{
-				OutputFilePaths = MakeBinaryPaths(Module.Name, "UE4Game-" + Module.Name, BinaryType, TargetType, Plugin, AppName);
-			}
-			else
-			{
-				OutputFilePaths = MakeBinaryPaths(Module.Name, GetAppName() + "-" + Module.Name, BinaryType, TargetType, Plugin, AppName);
-			}
-
-			// Try to determine if we have the rules file
-			var ModuleFilename = RulesCompiler.GetModuleFilename(Module.Name);
-			var bHasModuleRules = String.IsNullOrEmpty(ModuleFilename) == false;
-
-			// Figure out whether we should build it from source
-			var ModuleSourceFolder = bHasModuleRules ? Path.GetDirectoryName(RulesCompiler.GetModuleFilename(Module.Name)) : ModuleFilename;
-			bool bShouldBeBuiltFromSource = bHasModuleRules && Directory.GetFiles(ModuleSourceFolder, "*.cpp", SearchOption.AllDirectories).Length > 0;
-
-			string PluginIntermediateBuildPath;
-			{ 
-				if (Plugin.LoadedFrom == PluginInfo.LoadedFromType.Engine)
-				{
-					// Plugin folder is in the engine directory
-					var PluginConfiguration = Configuration == UnrealTargetConfiguration.DebugGame ? UnrealTargetConfiguration.Development : Configuration;
-					PluginIntermediateBuildPath = Path.GetFullPath(Path.Combine(BuildConfiguration.RelativeEnginePath, BuildConfiguration.PlatformIntermediateFolder, AppName, PluginConfiguration.ToString()));
-				}
-				else
-				{
-					// Plugin folder is in the project directory
-					PluginIntermediateBuildPath = Path.GetFullPath(Path.Combine(ProjectDirectory, BuildConfiguration.PlatformIntermediateFolder, GetTargetName(), Configuration.ToString()));
-				}
-				PluginIntermediateBuildPath = Path.Combine(PluginIntermediateBuildPath, "Plugins", ShouldCompileMonolithic() ? "Static" : "Dynamic");
-			}
-
-			// Create the binary
-			UEBuildBinaryConfiguration Config = new UEBuildBinaryConfiguration( InType:                  BinaryType,
-																				InOutputFilePaths:       OutputFilePaths,
-																				InIntermediateDirectory: PluginIntermediateBuildPath,
-																				bInAllowExports:         true,
-																				bInAllowCompilation:     bShouldBeBuiltFromSource,
-																				bInHasModuleRules:       bHasModuleRules,
-																				InModuleNames:           new List<string> { Module.Name } );
-			AppBinaries.Add(new UEBuildBinaryCPP(this, Config));
-		}
-
 		/// When building a target, this is called to add any additional modules that should be compiled along
 		/// with the main target.  If you override this in a derived class, remember to call the base implementation!
 		protected virtual void AddExtraModules()
@@ -2263,34 +2196,8 @@ namespace UnrealBuildTool
 			// Add extra modules that will either link into the main binary (monolithic), or be linked into separate DLL files (modular)
 			foreach (var ModuleName in ExtraModuleNames)
 			{
-				AddExtraModule(ModuleName);
-			}
-		}
-
-		/// <summary>
-		/// Adds extra module to the target.
-		/// </summary>
-		/// <param name="ModuleName">Name of the module.</param>
-		protected void AddExtraModule(string ModuleName)
-		{
-			if (ShouldCompileMonolithic())
-			{
-				// Add this module to the executable's list of included modules
-				var ExecutableBinary = AppBinaries[0];
-				ExecutableBinary.AddModule(ModuleName);
-			}
-			else
-			{
-				// Create a DLL binary for this module
-				string[] OutputFilePaths = MakeBinaryPaths(ModuleName, GetAppName() + "-" + ModuleName, UEBuildBinaryType.DynamicLinkLibrary, TargetType, null, AppName);
-				UEBuildBinaryConfiguration Config = new UEBuildBinaryConfiguration(InType: UEBuildBinaryType.DynamicLinkLibrary,
-																					InOutputFilePaths: OutputFilePaths,
-																					InIntermediateDirectory: RulesCompiler.IsGameModule(ModuleName) ? ProjectIntermediateDirectory : EngineIntermediateDirectory,
-																					bInAllowExports: true,
-																					InModuleNames: new List<string> { ModuleName });
-
-				// Tell the target about this new binary
-				AppBinaries.Add(new UEBuildBinaryCPP(this, Config));
+				UEBuildBinaryCPP Binary = FindOrAddBinaryForModule(ModuleName, false);
+				Binary.AddModule(ModuleName);
 			}
 		}
 
@@ -2367,7 +2274,8 @@ namespace UnrealBuildTool
 				{
 					if(PrecompiledModule is UEBuildModuleCPP && !BoundModuleNames.Contains(PrecompiledModule.Name))
 					{
-						UEBuildBinary Binary = AddBinaryForModule(PrecompiledModule.Name, bCompileMonolithic? UEBuildBinaryType.StaticLibrary : UEBuildBinaryType.DynamicLinkLibrary, bUsePrecompiled);
+						UEBuildBinaryType BinaryType = bCompileMonolithic? UEBuildBinaryType.StaticLibrary : UEBuildBinaryType.DynamicLinkLibrary;
+						UEBuildBinary Binary = AddBinaryForModule(PrecompiledModule.Name, BinaryType, bAllowCompilation: !bUsePrecompiled, bIsCrossTarget: false);
 						PrecompiledBinaries.Add(Binary);
 						BoundModuleNames.Add(PrecompiledModule.Name);
 					}
@@ -2375,7 +2283,23 @@ namespace UnrealBuildTool
 			}
 		}
 
-		protected UEBuildBinaryCPP AddBinaryForModule(string ModuleName, UEBuildBinaryType BinaryType, bool bAllowCompilation)
+		public UEBuildBinaryCPP FindOrAddBinaryForModule(string ModuleName, bool bIsCrossTarget)
+		{
+			UEBuildBinaryCPP Binary;
+			if (ShouldCompileMonolithic())
+			{
+				// When linking monolithically, any unbound modules will be linked into the main executable
+				Binary = (UEBuildBinaryCPP)AppBinaries[0];
+			}
+			else
+			{
+				// Otherwise create a new module for it
+				Binary = AddBinaryForModule(ModuleName, UEBuildBinaryType.DynamicLinkLibrary, bAllowCompilation: true, bIsCrossTarget: bIsCrossTarget);
+			}
+			return Binary;
+		}
+
+		protected UEBuildBinaryCPP AddBinaryForModule(string ModuleName, UEBuildBinaryType BinaryType, bool bAllowCompilation, bool bIsCrossTarget)
 		{
 			// Get the plugin info for the module
 			PluginInfo Plugin = Plugins.GetPluginInfoForModule(ModuleName);
@@ -2393,15 +2317,38 @@ namespace UnrealBuildTool
 			}
 
 			// Get the intermediate path
-			string IntermediateDirectory = RulesCompiler.IsGameModule(ModuleName)? ProjectIntermediateDirectory : EngineIntermediateDirectory;
+			string IntermediateDirectory;
+			if(Plugin == null)
+			{
+				if(RulesCompiler.IsGameModule(ModuleName))
+				{
+					IntermediateDirectory = ProjectIntermediateDirectory;
+				}
+				else
+				{
+					IntermediateDirectory = EngineIntermediateDirectory;
+				}
+			}
+			else
+			{
+				if(Plugin.LoadedFrom == PluginInfo.LoadedFromType.GameProject)
+				{
+					IntermediateDirectory = Path.Combine(ProjectIntermediateDirectory, "Plugins", ShouldCompileMonolithic() ? "Static" : "Dynamic");
+				}
+				else
+				{
+					IntermediateDirectory = Path.Combine(EngineIntermediateDirectory, "Plugins", ShouldCompileMonolithic() ? "Static" : "Dynamic");
+				}
+			}
 
 			// If it's a plugin, check if it actually has source.
+			bool bHasModuleRules = true;
 			if(Plugin != null)
 			{
 				string ModuleFileName = RulesCompiler.GetModuleFilename(ModuleName);
 				if(String.IsNullOrEmpty(ModuleFileName))
 				{
-					bAllowCompilation = false;
+					bAllowCompilation = bHasModuleRules = false;
 				}
 				else if(!Directory.EnumerateFiles(Path.GetDirectoryName(ModuleFileName), "*.cpp", SearchOption.AllDirectories).Any())
 				{
@@ -2409,11 +2356,18 @@ namespace UnrealBuildTool
 				}
 			}
 
-			// Get the binary configuration
-			UEBuildBinaryConfiguration BinaryConfig = new UEBuildBinaryConfiguration(BinaryType, OutputFilePaths, IntermediateDirectory, true, bInAllowCompilation: bPrecompile, InModuleNames: new List<string>{ ModuleName });
+			// Prepare the configuration object
+			UEBuildBinaryConfiguration Config = new UEBuildBinaryConfiguration(BinaryType);
+			Config.OutputFilePaths = OutputFilePaths;
+			Config.IntermediateDirectory = IntermediateDirectory;
+			Config.bHasModuleRules = bHasModuleRules;
+			Config.bAllowExports = (BinaryType == UEBuildBinaryType.DynamicLinkLibrary);
+			Config.bAllowCompilation = bAllowCompilation;
+			Config.bIsCrossTarget = bIsCrossTarget;
+			Config.ModuleNames.Add(ModuleName);
 
-			// Create the binary
-			UEBuildBinaryCPP Binary = new UEBuildBinaryCPP(this, BinaryConfig);
+			// Create the new binary
+			UEBuildBinaryCPP Binary = new UEBuildBinaryCPP(this, Config);
 			AppBinaries.Add(Binary);
 			return Binary;
 		}
