Remove -Static suffix from plugin static libraries; they won't overlap with import libraries in the same configuration, and all other precompiled libraries don't use it.

diff --git a/Engine/Source/Programs/UnrealBuildTool/Configuration/UEBuildTarget.cs b/Engine/Source/Programs/UnrealBuildTool/Configuration/UEBuildTarget.cs
--- a/Engine/Source/Programs/UnrealBuildTool/Configuration/UEBuildTarget.cs
+++ b/Engine/Source/Programs/UnrealBuildTool/Configuration/UEBuildTarget.cs
@@ -2470,14 +2470,8 @@ namespace UnrealBuildTool
 				BaseDirectory = Path.Combine(BaseDirectory, ModuleBinariesSubDir);
 			}
 
-			string BinarySuffix = "";
-			if ((PluginInfo != null) && (BinaryType != UEBuildBinaryType.DynamicLinkLibrary))
-			{
-				BinarySuffix = "-Static";
-			}
-
 			// append the architecture to the end of the binary name
-            BinarySuffix = BuildPlatform.ApplyArchitectureName(BinarySuffix);
+            string BinarySuffix = BuildPlatform.ApplyArchitectureName("");
 
 			string OutBinaryPath = "";
 			// Append binary file name
