Separate out the binary naming code from the code which figures out which directory it goes in.

diff --git a/Engine/Source/Programs/UnrealBuildTool/Configuration/UEBuildTarget.cs b/Engine/Source/Programs/UnrealBuildTool/Configuration/UEBuildTarget.cs
--- a/Engine/Source/Programs/UnrealBuildTool/Configuration/UEBuildTarget.cs
+++ b/Engine/Source/Programs/UnrealBuildTool/Configuration/UEBuildTarget.cs
@@ -2397,23 +2397,54 @@ namespace UnrealBuildTool
 			return false;
 		}
 
-		/** Given a UBT-built binary name (e.g. "Core"), returns a relative path to the binary for the current build configuration (e.g. "../Binaries/Win64/Core-Win64-Debug.lib") */
-		public static string[] MakeBinaryPaths(string ModuleName, string BinaryName, UnrealTargetPlatform Platform, 
-			UnrealTargetConfiguration Configuration, UEBuildBinaryType BinaryType, TargetRules.TargetType? TargetType, PluginInfo PluginInfo, string AppName, bool bForceNameAsForDevelopment = false, string ExeBinariesSubFolder = null)
+		/// <summary>
+		/// Makes a filename (without path) for a compiled binary (e.g. "Core-Win64-Debug.lib") */
+		/// </summary>
+		/// <param name="BinaryName">The name of this binary</param>
+		/// <param name="Platform">The platform being built for</param>
+		/// <param name="Configuration">The configuration being built</param>
+		/// <param name="BinaryType">Type of binary</param>
+		/// <param name="bAppendConfigurationName">If true, appends the platform and configuration name. Typically false for development configurations, true otherwise.</param>
+		/// <returns>Name of the binary</returns>
+		public static string MakeBinaryFileName(string BinaryName, UnrealTargetPlatform Platform, UnrealTargetConfiguration Configuration, UEBuildBinaryType BinaryType, bool bIncludeConfiguration)
 		{
-			// Determine the binary extension for the platform and binary type.
-			var BuildPlatform = UEBuildPlatform.GetBuildPlatform(Platform);
-			string BinaryExtension;
+			StringBuilder Result = new StringBuilder();
 
-			if (!BuildConfiguration.bRunUnrealCodeAnalyzer)
+			if (Platform == UnrealTargetPlatform.Linux && (BinaryType == UEBuildBinaryType.DynamicLinkLibrary || BinaryType == UEBuildBinaryType.StaticLibrary))
 			{
-				BinaryExtension = BuildPlatform.GetBinaryExtension(BinaryType);
+				Result.Append("lib");
+			}
+
+			Result.Append(BinaryName);
+
+			if(bIncludeConfiguration)
+			{
+				Result.AppendFormat("-{0}-{1}", Platform.ToString(), Configuration.ToString());
+			}
+
+			IUEBuildPlatform BuildPlatform = UEBuildPlatform.GetBuildPlatform(Platform);
+			Result.Append(BuildPlatform.ApplyArchitectureName(""));
+
+			if (BuildConfiguration.bRunUnrealCodeAnalyzer)
+			{
+				Result.AppendFormat("-{0}.analysis", BuildConfiguration.UCAModuleToAnalyze);
 			}
 			else
 			{
-				BinaryExtension = @"-" + BuildConfiguration.UCAModuleToAnalyze + @".analysis";
+				Result.Append(BuildPlatform.GetBinaryExtension(BinaryType));
 			}
 
+			return Result.ToString();
+		}
+
+		/** Given a UBT-built binary name (e.g. "Core"), returns a relative path to the binary for the current build configuration (e.g. "../Binaries/Win64/Core-Win64-Debug.lib") */
+		public static string[] MakeBinaryPaths(string ModuleName, string BinaryName, UnrealTargetPlatform Platform, 
+			UnrealTargetConfiguration Configuration, UEBuildBinaryType BinaryType, TargetRules.TargetType? TargetType, PluginInfo PluginInfo, string AppName, bool bForceNameAsForDevelopment = false, string ExeBinariesSubFolder = null)
+		{
+			// Determine the binary extension for the platform and binary type.
+			var BuildPlatform = UEBuildPlatform.GetBuildPlatform(Platform);
+
+
 			UnrealTargetConfiguration LocalConfig = Configuration;
 			if(Configuration == UnrealTargetConfiguration.DebugGame && !String.IsNullOrEmpty(ModuleName) && !RulesCompiler.IsGameModule(ModuleName))
 			{
@@ -2470,27 +2501,9 @@ namespace UnrealBuildTool
 				BaseDirectory = Path.Combine(BaseDirectory, ModuleBinariesSubDir);
 			}
 
-			// append the architecture to the end of the binary name
-            string BinarySuffix = BuildPlatform.ApplyArchitectureName("");
-
-			string OutBinaryPath = "";
-			// Append binary file name
-			string Prefix = "";
-			if (Platform == UnrealTargetPlatform.Linux && (BinaryType == UEBuildBinaryType.DynamicLinkLibrary || BinaryType == UEBuildBinaryType.StaticLibrary))
-			{
-				Prefix = "lib";
-			}
-
-			if (LocalConfig == UnrealTargetConfiguration.Development || bForceNameAsForDevelopment)
-			{
-				OutBinaryPath = Path.Combine(BaseDirectory, String.Format("{3}{0}{1}{2}", BinaryName, BinarySuffix, BinaryExtension, Prefix));
-			}
-			else
-			{
-				OutBinaryPath = Path.Combine(BaseDirectory, String.Format("{5}{0}-{1}-{2}{3}{4}", 
-					BinaryName, Platform.ToString(), LocalConfig.ToString(), BinarySuffix, BinaryExtension, Prefix));
-			}
-
+			// append the binary name
+			bool bIncludeConfiguration = (LocalConfig != UnrealTargetConfiguration.Development && !bForceNameAsForDevelopment);
+			string OutBinaryPath = Path.Combine(BaseDirectory, MakeBinaryFileName(BinaryName, Platform, Configuration, BinaryType, bIncludeConfiguration));
 			return BuildPlatform.FinalizeBinaryPaths(OutBinaryPath);
 		}
 
