Fix binary names for DebugGame configurations.

diff --git a/Engine/Source/Programs/UnrealBuildTool/Configuration/UEBuildTarget.cs b/Engine/Source/Programs/UnrealBuildTool/Configuration/UEBuildTarget.cs
--- a/Engine/Source/Programs/UnrealBuildTool/Configuration/UEBuildTarget.cs
+++ b/Engine/Source/Programs/UnrealBuildTool/Configuration/UEBuildTarget.cs
@@ -2503,7 +2503,7 @@ namespace UnrealBuildTool
 
 			// append the binary name
 			bool bIncludeConfiguration = (LocalConfig != UnrealTargetConfiguration.Development && !bForceNameAsForDevelopment);
-			string OutBinaryPath = Path.Combine(BaseDirectory, MakeBinaryFileName(BinaryName, Platform, Configuration, BinaryType, bIncludeConfiguration));
+			string OutBinaryPath = Path.Combine(BaseDirectory, MakeBinaryFileName(BinaryName, Platform, LocalConfig, BinaryType, bIncludeConfiguration));
 			return BuildPlatform.FinalizeBinaryPaths(OutBinaryPath);
 		}
 
