Add a command line option to disable UBT makefiles.

diff --git a/Engine/Source/Programs/UnrealBuildTool/System/UnrealBuildTool.cs b/Engine/Source/Programs/UnrealBuildTool/System/UnrealBuildTool.cs
--- a/Engine/Source/Programs/UnrealBuildTool/System/UnrealBuildTool.cs
+++ b/Engine/Source/Programs/UnrealBuildTool/System/UnrealBuildTool.cs
@@ -917,6 +917,10 @@ namespace UnrealBuildTool
                             BuildConfiguration.bXGEExport = true;
                             BuildConfiguration.bAllowXGE = true;
                         }
+						else if (LowercaseArg == "-noubtmakefiles")
+						{
+							BuildConfiguration.bUseUBTMakefiles = false;
+						}
 						else if (LowercaseArg == "-invalidatemakefilesonly")
 						{
 							UnrealBuildTool.bIsInvalidatingMakefilesOnly = true;
