Remove the unused AppName parameter from MakeBinaryPaths.

diff --git a/Engine/Source/Programs/UnrealBuildTool/Configuration/UEBuildTarget.cs b/Engine/Source/Programs/UnrealBuildTool/Configuration/UEBuildTarget.cs
--- a/Engine/Source/Programs/UnrealBuildTool/Configuration/UEBuildTarget.cs
+++ b/Engine/Source/Programs/UnrealBuildTool/Configuration/UEBuildTarget.cs
@@ -824,7 +824,7 @@ namespace UnrealBuildTool
 			TargetTypeOrNull = (Rules != null) ? Rules.Type : (TargetRules.TargetType?)null;
 
 			// Construct the output path based on configuration, platform, game if not specified.
-            OutputPaths = MakeBinaryPaths("", AppName, UEBuildBinaryType.Executable, TargetType, null, AppName, Configuration == UnrealTargetConfiguration.Shipping ? Rules.ForceNameAsForDevelopment() : false, Rules.ExeBinariesSubFolder);
+            OutputPaths = MakeBinaryPaths("", AppName, UEBuildBinaryType.Executable, TargetType, null, Configuration == UnrealTargetConfiguration.Shipping ? Rules.ForceNameAsForDevelopment() : false, Rules.ExeBinariesSubFolder);
 			for (int Index = 0; Index < OutputPaths.Length; Index++)
 			{
 				OutputPaths[Index] = Path.GetFullPath(OutputPaths[Index]);
@@ -2309,11 +2309,11 @@ namespace UnrealBuildTool
 			string[] OutputFilePaths;
 			if(TargetType == TargetRules.TargetType.Game && UnrealBuildTool.RunningRocket())
 			{
-				OutputFilePaths = MakeBinaryPaths(ModuleName, "UE4Game-" + ModuleName, BinaryType, TargetType, Plugin, AppName);
+				OutputFilePaths = MakeBinaryPaths(ModuleName, "UE4Game-" + ModuleName, BinaryType, TargetType, Plugin);
 			}
 			else
 			{
-				OutputFilePaths = MakeBinaryPaths(ModuleName, AppName + "-" + ModuleName, BinaryType, TargetType, Plugin, AppName);
+				OutputFilePaths = MakeBinaryPaths(ModuleName, AppName + "-" + ModuleName, BinaryType, TargetType, Plugin);
 			}
 
 			// Get the intermediate path
@@ -2439,7 +2439,7 @@ namespace UnrealBuildTool
 
 		/** Given a UBT-built binary name (e.g. "Core"), returns a relative path to the binary for the current build configuration (e.g. "../Binaries/Win64/Core-Win64-Debug.lib") */
 		public static string[] MakeBinaryPaths(string ModuleName, string BinaryName, UnrealTargetPlatform Platform, 
-			UnrealTargetConfiguration Configuration, UEBuildBinaryType BinaryType, TargetRules.TargetType? TargetType, PluginInfo PluginInfo, string AppName, bool bForceNameAsForDevelopment = false, string ExeBinariesSubFolder = null)
+			UnrealTargetConfiguration Configuration, UEBuildBinaryType BinaryType, TargetRules.TargetType? TargetType, PluginInfo PluginInfo, bool bForceNameAsForDevelopment = false, string ExeBinariesSubFolder = null)
 		{
 			// Determine the binary extension for the platform and binary type.
 			var BuildPlatform = UEBuildPlatform.GetBuildPlatform(Platform);
@@ -2508,15 +2508,15 @@ namespace UnrealBuildTool
 		}
 
 		/** Given a UBT-built binary name (e.g. "Core"), returns a relative path to the binary for the current build configuration (e.g. "../../Binaries/Win64/Core-Win64-Debug.lib") */
-		public string[] MakeBinaryPaths(string ModuleName, string BinaryName, UEBuildBinaryType BinaryType, TargetRules.TargetType? TargetType, PluginInfo PluginInfo, string AppName, bool bForceNameAsForDevelopment = false, string ExeBinariesSubFolder = null)
+		public string[] MakeBinaryPaths(string ModuleName, string BinaryName, UEBuildBinaryType BinaryType, TargetRules.TargetType? TargetType, PluginInfo PluginInfo, bool bForceNameAsForDevelopment = false, string ExeBinariesSubFolder = null)
 		{
 			if (String.IsNullOrEmpty(ModuleName) && Configuration == UnrealTargetConfiguration.DebugGame && !bCompileMonolithic)
 			{
-				return MakeBinaryPaths(ModuleName, BinaryName, Platform, UnrealTargetConfiguration.Development, BinaryType, TargetType, PluginInfo, AppName, bForceNameAsForDevelopment);
+				return MakeBinaryPaths(ModuleName, BinaryName, Platform, UnrealTargetConfiguration.Development, BinaryType, TargetType, PluginInfo, bForceNameAsForDevelopment);
 			}
 			else
 			{
-				return MakeBinaryPaths(ModuleName, BinaryName, Platform, Configuration, BinaryType, TargetType, PluginInfo, AppName, bForceNameAsForDevelopment, ExeBinariesSubFolder);
+				return MakeBinaryPaths(ModuleName, BinaryName, Platform, Configuration, BinaryType, TargetType, PluginInfo, bForceNameAsForDevelopment, ExeBinariesSubFolder);
 			}
 		}
 
diff --git a/Engine/Source/Programs/UnrealBuildTool/System/UnrealBuildTool.cs b/Engine/Source/Programs/UnrealBuildTool/System/UnrealBuildTool.cs
--- a/Engine/Source/Programs/UnrealBuildTool/System/UnrealBuildTool.cs
+++ b/Engine/Source/Programs/UnrealBuildTool/System/UnrealBuildTool.cs
@@ -1922,7 +1922,7 @@ namespace UnrealBuildTool
 
 			if (!ProjectFileGenerator.bGenerateProjectFiles && !UEBuildConfiguration.bGenerateManifest && bIsEditorTarget )
 			{
-				var EditorProcessFilenames = UEBuildTarget.MakeBinaryPaths("Launcher", "UE4Editor", TargetDesc.Platform, TargetDesc.Configuration, UEBuildBinaryType.Executable, null, null, "UE4Editor");
+				var EditorProcessFilenames = UEBuildTarget.MakeBinaryPaths("Launcher", "UE4Editor", TargetDesc.Platform, TargetDesc.Configuration, UEBuildBinaryType.Executable, null, null);
 				if (EditorProcessFilenames.Length != 1)
 				{
 					throw new BuildException("ShouldDoHotReload cannot handle multiple binaries returning from UEBuildTarget.MakeBinaryPaths");
