Don't force adding a module to a binary, even if we create it. Some targets have dependencies on include paths from other modules, but those modules never actually get built. (UnrealFileServer has an indirect dependency on DDCUtils this way, but it's not actually part of the target even though the module exists a bIncludeInTarget gets set).

diff --git a/Engine/Source/Programs/UnrealBuildTool/Configuration/UEBuildModule.cs b/Engine/Source/Programs/UnrealBuildTool/Configuration/UEBuildModule.cs
--- a/Engine/Source/Programs/UnrealBuildTool/Configuration/UEBuildModule.cs
+++ b/Engine/Source/Programs/UnrealBuildTool/Configuration/UEBuildModule.cs
@@ -2037,7 +2037,6 @@ namespace UnrealBuildTool
 
 						// Get the binary that this module should be bound to
 						UEBuildBinary BinaryToBindTo = Target.FindOrAddBinaryForModule(DependencyName, bIsCrossTarget);
-						BinaryToBindTo.AddModule(DependencyName);
 
 						// Bind this module
 						DependencyModule.Binary = BinaryToBindTo;
