Allow target rules to specify which build configuration doesn't have the -Platform-Configuration suffix appended to all binaries.

diff --git a/Engine/Source/Programs/CrashReporter/CrashReportClient/CrashReportClient.Target.cs b/Engine/Source/Programs/CrashReporter/CrashReportClient/CrashReportClient.Target.cs
--- a/Engine/Source/Programs/CrashReporter/CrashReportClient/CrashReportClient.Target.cs
+++ b/Engine/Source/Programs/CrashReporter/CrashReportClient/CrashReportClient.Target.cs
@@ -8,6 +8,7 @@ public class CrashReportClientTarget : TargetRules
 	public CrashReportClientTarget(TargetInfo Target)
 	{
 		Type = TargetType.Program;
+		UndecoratedConfiguration = UnrealTargetConfiguration.Shipping;
 	}
 
 	//
@@ -53,11 +54,6 @@ public class CrashReportClientTarget : TargetRules
 		}
 	}
 
-	public override bool ForceNameAsForDevelopment()
-	{
-		return true;
-	}
-
 	public override bool ShouldCompileMonolithic(UnrealTargetPlatform InPlatform, UnrealTargetConfiguration InConfiguration)
 	{
 		return true;
diff --git a/Engine/Source/Programs/UnrealBuildTool/Configuration/UEBuildTarget.cs b/Engine/Source/Programs/UnrealBuildTool/Configuration/UEBuildTarget.cs
--- a/Engine/Source/Programs/UnrealBuildTool/Configuration/UEBuildTarget.cs
+++ b/Engine/Source/Programs/UnrealBuildTool/Configuration/UEBuildTarget.cs
@@ -824,7 +824,7 @@ namespace UnrealBuildTool
 			TargetTypeOrNull = (Rules != null) ? Rules.Type : (TargetRules.TargetType?)null;
 
 			// Construct the output path based on configuration, platform, game if not specified.
-            OutputPaths = MakeBinaryPaths("", AppName, UEBuildBinaryType.Executable, TargetType, null, Configuration == UnrealTargetConfiguration.Shipping ? Rules.ForceNameAsForDevelopment() : false, Rules.ExeBinariesSubFolder);
+            OutputPaths = MakeBinaryPaths("", AppName, UEBuildBinaryType.Executable, TargetType, null, Rules.ExeBinariesSubFolder);
 			for (int Index = 0; Index < OutputPaths.Length; Index++)
 			{
 				OutputPaths[Index] = Path.GetFullPath(OutputPaths[Index]);
@@ -2403,10 +2403,10 @@ namespace UnrealBuildTool
 		/// <param name="BinaryName">The name of this binary</param>
 		/// <param name="Platform">The platform being built for</param>
 		/// <param name="Configuration">The configuration being built</param>
+ 		/// <param name="UndecoratedConfiguration">The target configuration which doesn't require a platform and configuration suffix. Development by default.</param>
 		/// <param name="BinaryType">Type of binary</param>
-		/// <param name="bAppendConfigurationName">If true, appends the platform and configuration name. Typically false for development configurations, true otherwise.</param>
 		/// <returns>Name of the binary</returns>
-		public static string MakeBinaryFileName(string BinaryName, UnrealTargetPlatform Platform, UnrealTargetConfiguration Configuration, UEBuildBinaryType BinaryType, bool bIncludeConfiguration)
+		public static string MakeBinaryFileName(string BinaryName, UnrealTargetPlatform Platform, UnrealTargetConfiguration Configuration, UnrealTargetConfiguration UndecoratedConfiguration, UEBuildBinaryType BinaryType)
 		{
 			StringBuilder Result = new StringBuilder();
 
@@ -2417,7 +2417,7 @@ namespace UnrealBuildTool
 
 			Result.Append(BinaryName);
 
-			if(bIncludeConfiguration)
+			if(Configuration != UndecoratedConfiguration)
 			{
 				Result.AppendFormat("-{0}-{1}", Platform.ToString(), Configuration.ToString());
 			}
@@ -2439,7 +2439,7 @@ namespace UnrealBuildTool
 
 		/** Given a UBT-built binary name (e.g. "Core"), returns a relative path to the binary for the current build configuration (e.g. "../Binaries/Win64/Core-Win64-Debug.lib") */
 		public static string[] MakeBinaryPaths(string ModuleName, string BinaryName, UnrealTargetPlatform Platform, 
-			UnrealTargetConfiguration Configuration, UEBuildBinaryType BinaryType, TargetRules.TargetType? TargetType, PluginInfo PluginInfo, bool bForceNameAsForDevelopment = false, string ExeBinariesSubFolder = null)
+			UnrealTargetConfiguration Configuration, UEBuildBinaryType BinaryType, TargetRules.TargetType? TargetType, PluginInfo PluginInfo, UnrealTargetConfiguration UndecoratedConfiguration, string ExeBinariesSubFolder = null)
 		{
 			// Determine the binary extension for the platform and binary type.
 			var BuildPlatform = UEBuildPlatform.GetBuildPlatform(Platform);
@@ -2502,21 +2502,20 @@ namespace UnrealBuildTool
 			}
 
 			// append the binary name
-			bool bIncludeConfiguration = (LocalConfig != UnrealTargetConfiguration.Development && !bForceNameAsForDevelopment);
-			string OutBinaryPath = Path.Combine(BaseDirectory, MakeBinaryFileName(BinaryName, Platform, LocalConfig, BinaryType, bIncludeConfiguration));
+			string OutBinaryPath = Path.Combine(BaseDirectory, MakeBinaryFileName(BinaryName, Platform, LocalConfig, UndecoratedConfiguration, BinaryType));
 			return BuildPlatform.FinalizeBinaryPaths(OutBinaryPath);
 		}
 
 		/** Given a UBT-built binary name (e.g. "Core"), returns a relative path to the binary for the current build configuration (e.g. "../../Binaries/Win64/Core-Win64-Debug.lib") */
-		public string[] MakeBinaryPaths(string ModuleName, string BinaryName, UEBuildBinaryType BinaryType, TargetRules.TargetType? TargetType, PluginInfo PluginInfo, bool bForceNameAsForDevelopment = false, string ExeBinariesSubFolder = null)
+		public string[] MakeBinaryPaths(string ModuleName, string BinaryName, UEBuildBinaryType BinaryType, TargetRules.TargetType? TargetType, PluginInfo PluginInfo, string ExeBinariesSubFolder = null)
 		{
 			if (String.IsNullOrEmpty(ModuleName) && Configuration == UnrealTargetConfiguration.DebugGame && !bCompileMonolithic)
 			{
-				return MakeBinaryPaths(ModuleName, BinaryName, Platform, UnrealTargetConfiguration.Development, BinaryType, TargetType, PluginInfo, bForceNameAsForDevelopment);
+				return MakeBinaryPaths(ModuleName, BinaryName, Platform, UnrealTargetConfiguration.Development, BinaryType, TargetType, PluginInfo, Rules.UndecoratedConfiguration);
 			}
 			else
 			{
-				return MakeBinaryPaths(ModuleName, BinaryName, Platform, Configuration, BinaryType, TargetType, PluginInfo, bForceNameAsForDevelopment, ExeBinariesSubFolder);
+				return MakeBinaryPaths(ModuleName, BinaryName, Platform, Configuration, BinaryType, TargetType, PluginInfo, Rules.UndecoratedConfiguration, ExeBinariesSubFolder);
 			}
 		}
 
diff --git a/Engine/Source/Programs/UnrealBuildTool/System/RulesCompiler.cs b/Engine/Source/Programs/UnrealBuildTool/System/RulesCompiler.cs
--- a/Engine/Source/Programs/UnrealBuildTool/System/RulesCompiler.cs
+++ b/Engine/Source/Programs/UnrealBuildTool/System/RulesCompiler.cs
@@ -441,6 +441,11 @@ namespace UnrealBuildTool
 		/// </summary>
 		public bool bOutputPubliclyDistributable = false;
 
+		/// <summary>
+		/// Specifies the configuration whose binaries do not require a "-Platform-Configuration" suffix.
+		/// </summary>
+		public UnrealTargetConfiguration UndecoratedConfiguration = UnrealTargetConfiguration.Development;
+
 		/// <summary>
 		/// A list of additional plugins which need to be built for this target. Game and editor targets can use the EnabledPlugins 
 		/// setting in their config files to control this.
@@ -644,16 +649,6 @@ namespace UnrealBuildTool
 		}
 
 
-		/// <summary>
-		/// Returns true if this target's output path needs to be the same as for the development configuration.
-		/// Currently only used by the CrashReportClient.
-		/// </summary>
-		/// <returns>true if this target's output path needs to be the same as for the development configuration.</returns>
-		public virtual bool ForceNameAsForDevelopment()
-		{
-			return false;
-		}
-
 		/// <summary>
 		/// Setup the global environment for building this target
 		/// IMPORTANT: Game targets will *not* have this function called unless they are built as monolithic targets.
