Missing file.

diff --git a/Engine/Source/Programs/UnrealBuildTool/System/UnrealBuildTool.cs b/Engine/Source/Programs/UnrealBuildTool/System/UnrealBuildTool.cs
--- a/Engine/Source/Programs/UnrealBuildTool/System/UnrealBuildTool.cs
+++ b/Engine/Source/Programs/UnrealBuildTool/System/UnrealBuildTool.cs
@@ -1922,7 +1922,7 @@ namespace UnrealBuildTool
 
 			if (!ProjectFileGenerator.bGenerateProjectFiles && !UEBuildConfiguration.bGenerateManifest && bIsEditorTarget )
 			{
-				var EditorProcessFilenames = UEBuildTarget.MakeBinaryPaths("Launcher", "UE4Editor", TargetDesc.Platform, TargetDesc.Configuration, UEBuildBinaryType.Executable, null, null);
+				var EditorProcessFilenames = UEBuildTarget.MakeBinaryPaths("Launcher", "UE4Editor", TargetDesc.Platform, TargetDesc.Configuration, UEBuildBinaryType.Executable, null, null, UnrealTargetConfiguration.Development);
 				if (EditorProcessFilenames.Length != 1)
 				{
 					throw new BuildException("ShouldDoHotReload cannot handle multiple binaries returning from UEBuildTarget.MakeBinaryPaths");
