Remove handling of null TargetRules in UEBuildTarget constructor. There are already plenty of places that *don't* check for it being null, so I assume that this code path is long abandoned.

diff --git a/Engine/Source/Programs/UnrealBuildTool/Configuration/UEBuildTarget.cs b/Engine/Source/Programs/UnrealBuildTool/Configuration/UEBuildTarget.cs
--- a/Engine/Source/Programs/UnrealBuildTool/Configuration/UEBuildTarget.cs
+++ b/Engine/Source/Programs/UnrealBuildTool/Configuration/UEBuildTarget.cs
@@ -642,15 +642,8 @@ namespace UnrealBuildTool
 		[NonSerialized]
 		public TargetRules Rules = null;
 
-		/** Type of target, or null if undetermined (such as in the case of a synthetic target with no TargetRules) */
-		public TargetRules.TargetType TargetType
-		{
-			get
-			{
-				return TargetTypeOrNull.Value;
-			}
-		}
-		private readonly TargetRules.TargetType? TargetTypeOrNull;
+		/** Type of target */
+		public TargetRules.TargetType TargetType;
 
 		/** The name of the application the target is part of. For non-monolithic targets, this is typically the name of the base application, eg. UE4Editor for any game editor. */
 		public string AppName;
@@ -760,19 +753,19 @@ namespace UnrealBuildTool
 		 * @param InAdditionalDefinitions - Additional definitions provided on the UBT command-line for the target.
 		 * @param InRemoteRoot - The remote path that the build output is synced to.
 		 */
-		public UEBuildTarget(string InAppName, TargetDescriptor InDesc, TargetRules InRulesObject)
+		public UEBuildTarget(string InAppName, TargetDescriptor InDesc, TargetRules InRules)
 		{
 			AppName = InAppName;
 			TargetName = InDesc.TargetName;
 			Platform = InDesc.Platform;
 			Configuration = InDesc.Configuration;
-			Rules = InRulesObject;
+			Rules = InRules;
 			bEditorRecompile = InDesc.bIsEditorRecompile;
 			bPrecompile = InDesc.bPrecompile;
 			bUsePrecompiled = InDesc.bUsePrecompiled;
 
 			{
-				bCompileMonolithic = (Rules != null) ? Rules.ShouldCompileMonolithic(InDesc.Platform, InDesc.Configuration) : false;
+				bCompileMonolithic = Rules.ShouldCompileMonolithic(InDesc.Platform, InDesc.Configuration);
 
 				// Platforms may *require* monolithic compilation...
 				bCompileMonolithic |= UEBuildPlatform.PlatformRequiresMonolithicBuilds(InDesc.Platform, InDesc.Configuration);
@@ -821,7 +814,7 @@ namespace UnrealBuildTool
 
 			OnlyModules = InDesc.OnlyModules;
 
-			TargetTypeOrNull = (Rules != null) ? Rules.Type : (TargetRules.TargetType?)null;
+			TargetType = Rules.Type;
 
 			// Construct the output path based on configuration, platform, game if not specified.
             OutputPaths = MakeBinaryPaths("", AppName, UEBuildBinaryType.Executable, TargetType, null, Rules.ExeBinariesSubFolder);
@@ -830,10 +823,10 @@ namespace UnrealBuildTool
 				OutputPaths[Index] = Path.GetFullPath(OutputPaths[Index]);
 			}
 
-			if (bCompileMonolithic && TargetRules.IsGameType(InRulesObject.Type))
+			if (bCompileMonolithic && TargetRules.IsGameType(Rules.Type))
 			{
 				// For Rocket, UE4Game.exe and UE4Editor.exe still go into Engine/Binaries/<PLATFORM>
-				if (!InRulesObject.bOutputToEngineBinaries)
+				if (!Rules.bOutputToEngineBinaries)
 				{
 					// We are compiling a monolithic game...
 					// We want the output to go into the <GAME>\Binaries folder
@@ -1885,7 +1878,7 @@ namespace UnrealBuildTool
 			UEBuildBinaryCPP BinaryCPP = ExecutableBinary as UEBuildBinaryCPP;
 			if (BinaryCPP != null)
 			{
-				TargetInfo CurrentTarget = new TargetInfo(Platform, Configuration, (Rules != null) ? TargetType : (TargetRules.TargetType?)null);
+				TargetInfo CurrentTarget = new TargetInfo(Platform, Configuration, TargetType);
 				foreach (var TargetModuleName in BinaryCPP.ModuleNames)
 				{
 					string UnusedFilename;
@@ -2210,7 +2203,7 @@ namespace UnrealBuildTool
 			{
 				// Find all the precompiled module names.
 				List<string> PrecompiledModuleNames = new List<string>();
-				Rules.GetModulesToPrecompile(new TargetInfo(Platform, Configuration, TargetTypeOrNull), PrecompiledModuleNames);
+				Rules.GetModulesToPrecompile(new TargetInfo(Platform, Configuration, TargetType), PrecompiledModuleNames);
 
 				// Add all the enabled plugins to the precompiled module list. Plugins are always precompiled, even if bPrecompile is not set, so we should precompile their dependencies.
 				foreach(PluginInfo Plugin in BuildPlugins)
@@ -2439,7 +2432,7 @@ namespace UnrealBuildTool
 
 		/** Given a UBT-built binary name (e.g. "Core"), returns a relative path to the binary for the current build configuration (e.g. "../Binaries/Win64/Core-Win64-Debug.lib") */
 		public static string[] MakeBinaryPaths(string ModuleName, string BinaryName, UnrealTargetPlatform Platform, 
-			UnrealTargetConfiguration Configuration, UEBuildBinaryType BinaryType, TargetRules.TargetType? TargetType, PluginInfo PluginInfo, UnrealTargetConfiguration UndecoratedConfiguration, string ExeBinariesSubFolder = null)
+			UnrealTargetConfiguration Configuration, UEBuildBinaryType BinaryType, TargetRules.TargetType TargetType, PluginInfo PluginInfo, UnrealTargetConfiguration UndecoratedConfiguration, string ExeBinariesSubFolder = null)
 		{
 			// Determine the binary extension for the platform and binary type.
 			var BuildPlatform = UEBuildPlatform.GetBuildPlatform(Platform);
@@ -2475,7 +2468,7 @@ namespace UnrealBuildTool
 
 			// Build base directory string ("../Binaries/<Platform>/")
 			string BinariesDirName;
-			if(TargetType.HasValue && TargetType.Value == TargetRules.TargetType.Program && UnrealBuildTool.HasUProjectFile() && !ProjectFileGenerator.bGenerateProjectFiles)
+			if(TargetType == TargetRules.TargetType.Program && UnrealBuildTool.HasUProjectFile() && !ProjectFileGenerator.bGenerateProjectFiles)
 			{
 				BinariesDirName = Path.Combine( UnrealBuildTool.GetUProjectPath(), "Binaries" );
 			}
@@ -2507,7 +2500,7 @@ namespace UnrealBuildTool
 		}
 
 		/** Given a UBT-built binary name (e.g. "Core"), returns a relative path to the binary for the current build configuration (e.g. "../../Binaries/Win64/Core-Win64-Debug.lib") */
-		public string[] MakeBinaryPaths(string ModuleName, string BinaryName, UEBuildBinaryType BinaryType, TargetRules.TargetType? TargetType, PluginInfo PluginInfo, string ExeBinariesSubFolder = null)
+		public string[] MakeBinaryPaths(string ModuleName, string BinaryName, UEBuildBinaryType BinaryType, TargetRules.TargetType TargetType, PluginInfo PluginInfo, string ExeBinariesSubFolder = null)
 		{
 			if (String.IsNullOrEmpty(ModuleName) && Configuration == UnrealTargetConfiguration.DebugGame && !bCompileMonolithic)
 			{
@@ -2929,14 +2922,7 @@ namespace UnrealBuildTool
 		/** Constructs a TargetInfo object for this target. */
 		public TargetInfo GetTargetInfo()
 		{
-			if(Rules == null)
-			{
-				return new TargetInfo(Platform, Configuration);
-			}
-			else
-			{
-				return new TargetInfo(Platform, Configuration, TargetType);
-			}
+			return new TargetInfo(Platform, Configuration, TargetType);
 		}
 
 		/** Registers a module with this target. */
diff --git a/Engine/Source/Programs/UnrealBuildTool/System/UnrealBuildTool.cs b/Engine/Source/Programs/UnrealBuildTool/System/UnrealBuildTool.cs
--- a/Engine/Source/Programs/UnrealBuildTool/System/UnrealBuildTool.cs
+++ b/Engine/Source/Programs/UnrealBuildTool/System/UnrealBuildTool.cs
@@ -1922,7 +1922,7 @@ namespace UnrealBuildTool
 
 			if (!ProjectFileGenerator.bGenerateProjectFiles && !UEBuildConfiguration.bGenerateManifest && bIsEditorTarget )
 			{
-				var EditorProcessFilenames = UEBuildTarget.MakeBinaryPaths("Launcher", "UE4Editor", TargetDesc.Platform, TargetDesc.Configuration, UEBuildBinaryType.Executable, null, null, UnrealTargetConfiguration.Development);
+				var EditorProcessFilenames = UEBuildTarget.MakeBinaryPaths("Launcher", "UE4Editor", TargetDesc.Platform, TargetDesc.Configuration, UEBuildBinaryType.Executable, TargetRules.TargetType.Editor, null, UnrealTargetConfiguration.Development);
 				if (EditorProcessFilenames.Length != 1)
 				{
 					throw new BuildException("ShouldDoHotReload cannot handle multiple binaries returning from UEBuildTarget.MakeBinaryPaths");
