Separate code for generating executable paths into its own function; it has different logic to everything else in MakeBinaryPaths().

diff --git a/Engine/Source/Programs/UnrealBuildTool/Configuration/UEBuildTarget.cs b/Engine/Source/Programs/UnrealBuildTool/Configuration/UEBuildTarget.cs
--- a/Engine/Source/Programs/UnrealBuildTool/Configuration/UEBuildTarget.cs
+++ b/Engine/Source/Programs/UnrealBuildTool/Configuration/UEBuildTarget.cs
@@ -816,45 +816,17 @@ namespace UnrealBuildTool
 
 			TargetType = Rules.Type;
 
-			// Construct the output path based on configuration, platform, game if not specified.
-            OutputPaths = MakeBinaryPaths("", AppName, UEBuildBinaryType.Executable, TargetType, null, Rules.ExeBinariesSubFolder);
-			for (int Index = 0; Index < OutputPaths.Length; Index++)
+			// Construct the output paths for this target's executable
+			string OutputDirectory;
+			if((bCompileMonolithic || TargetType == TargetRules.TargetType.Program) && !Rules.bOutputToEngineBinaries)
 			{
-				OutputPaths[Index] = Path.GetFullPath(OutputPaths[Index]);
+				OutputDirectory = ProjectDirectory;
 			}
-
-			if (bCompileMonolithic && TargetRules.IsGameType(Rules.Type))
+			else
 			{
-				// For Rocket, UE4Game.exe and UE4Editor.exe still go into Engine/Binaries/<PLATFORM>
-				if (!Rules.bOutputToEngineBinaries)
-				{
-					// We are compiling a monolithic game...
-					// We want the output to go into the <GAME>\Binaries folder
-					if (UnrealBuildTool.HasUProjectFile() == false)
-					{
-						for (int Index = 0; Index < OutputPaths.Length; Index++)
-						{
-							OutputPaths[Index] = OutputPaths[Index].Replace(Path.Combine("Engine", "Binaries"), Path.Combine(InDesc.TargetName, "Binaries"));
-						}
-					}
-					else
-					{
-						string EnginePath = Path.GetFullPath(Path.Combine(ProjectFileGenerator.EngineRelativePath, "Binaries"));
-						string UProjectPath = UnrealBuildTool.GetUProjectPath();
-						if (Path.IsPathRooted(UProjectPath) == false)
-						{
-							string FilePath = UProjectInfo.GetProjectForTarget(InDesc.TargetName);
-							string FullPath = Path.GetFullPath(FilePath);
-							UProjectPath = Path.GetDirectoryName(FullPath);
-						}
-						string ProjectPath = Path.GetFullPath(Path.Combine(UProjectPath, "Binaries"));
-						for (int Index = 0; Index < OutputPaths.Length; Index++)
-						{
-							OutputPaths[Index] = OutputPaths[Index].Replace(EnginePath, ProjectPath);
-						}
-					}
-				}
+				OutputDirectory = Path.GetFullPath(BuildConfiguration.RelativeEnginePath);
 			}
+			OutputPaths = MakeExecutablePaths(OutputDirectory, AppName, Platform, Configuration, Rules.UndecoratedConfiguration, bCompileMonolithic && UnrealBuildTool.HasUProjectFile(), Rules.ExeBinariesSubFolder);
 
 			// handle some special case defines (so build system can pass -DEFINE as normal instead of needing
 			// to know about special parameters)
@@ -2302,11 +2274,11 @@ namespace UnrealBuildTool
 			string[] OutputFilePaths;
 			if(TargetType == TargetRules.TargetType.Game && UnrealBuildTool.RunningRocket())
 			{
-				OutputFilePaths = MakeBinaryPaths(ModuleName, "UE4Game-" + ModuleName, BinaryType, TargetType, Plugin);
+				OutputFilePaths = MakeBinaryPaths(ModuleName, "UE4Game-" + ModuleName, Platform, Configuration, BinaryType, TargetType, Plugin, Rules.UndecoratedConfiguration);
 			}
 			else
 			{
-				OutputFilePaths = MakeBinaryPaths(ModuleName, AppName + "-" + ModuleName, BinaryType, TargetType, Plugin);
+				OutputFilePaths = MakeBinaryPaths(ModuleName, AppName + "-" + ModuleName, Platform, Configuration, BinaryType, TargetType, Plugin, Rules.UndecoratedConfiguration);
 			}
 
 			// Get the intermediate path
@@ -2432,20 +2404,20 @@ namespace UnrealBuildTool
 
 		/** Given a UBT-built binary name (e.g. "Core"), returns a relative path to the binary for the current build configuration (e.g. "../Binaries/Win64/Core-Win64-Debug.lib") */
 		public static string[] MakeBinaryPaths(string ModuleName, string BinaryName, UnrealTargetPlatform Platform, 
-			UnrealTargetConfiguration Configuration, UEBuildBinaryType BinaryType, TargetRules.TargetType TargetType, PluginInfo PluginInfo, UnrealTargetConfiguration UndecoratedConfiguration, string ExeBinariesSubFolder = null)
+			UnrealTargetConfiguration Configuration, UEBuildBinaryType BinaryType, TargetRules.TargetType TargetType, PluginInfo PluginInfo, UnrealTargetConfiguration UndecoratedConfiguration)
 		{
 			// Determine the binary extension for the platform and binary type.
 			var BuildPlatform = UEBuildPlatform.GetBuildPlatform(Platform);
 
 
 			UnrealTargetConfiguration LocalConfig = Configuration;
-			if(Configuration == UnrealTargetConfiguration.DebugGame && !String.IsNullOrEmpty(ModuleName) && !RulesCompiler.IsGameModule(ModuleName))
+			if(Configuration == UnrealTargetConfiguration.DebugGame && !RulesCompiler.IsGameModule(ModuleName))
 			{
 				LocalConfig = UnrealTargetConfiguration.Development;
 			}
 			
 			string ModuleBinariesSubDir = "";
-			if ((BinaryType == UEBuildBinaryType.DynamicLinkLibrary || BinaryType == UEBuildBinaryType.StaticLibrary) && (string.IsNullOrEmpty(ModuleName) == false))
+			if (BinaryType == UEBuildBinaryType.DynamicLinkLibrary || BinaryType == UEBuildBinaryType.StaticLibrary)
 			{
 				// Allow for modules to specify sub-folders in the Binaries folder
 				var RulesFilename = RulesCompiler.GetModuleFilename(ModuleName);
@@ -2459,10 +2431,6 @@ namespace UnrealBuildTool
 					}
 				}
 			}
-            else if ( BinaryType == UEBuildBinaryType.Executable && string.IsNullOrEmpty(ExeBinariesSubFolder) == false )
-            {
-                ModuleBinariesSubDir = ExeBinariesSubFolder;
-            }
 
 			[REDACTED_EMAIL]: This just happens to work since exp and lib files go into intermediate...
 
@@ -2499,17 +2467,36 @@ namespace UnrealBuildTool
 			return BuildPlatform.FinalizeBinaryPaths(OutBinaryPath);
 		}
 
-		/** Given a UBT-built binary name (e.g. "Core"), returns a relative path to the binary for the current build configuration (e.g. "../../Binaries/Win64/Core-Win64-Debug.lib") */
-		public string[] MakeBinaryPaths(string ModuleName, string BinaryName, UEBuildBinaryType BinaryType, TargetRules.TargetType TargetType, PluginInfo PluginInfo, string ExeBinariesSubFolder = null)
+		/// <summary>
+		/// Determine the output path for a target's executable
+		/// </summary>
+		/// <param name="BaseDirectory">The base directory for the executable; typically either the engine directory or project directory.</param>
+		/// <param name="BinaryName">Name of the binary</param>
+		/// <param name="Platform">Target platform to build for</param>
+		/// <param name="Configuration">Target configuration being built</param>
+		/// <param name="UndecoratedConfiguration">The configuration which doesn't have a "-{Platform}-{Configuration}" suffix added to the binary</param>
+		/// <param name="bIncludesGameModules">Whether this executable contains game modules</param>
+		/// <param name="ExeSubFolder">Subfolder for executables. May be null.</param>
+		/// <returns>List of executable paths for this target</returns>
+		public static string[] MakeExecutablePaths(string BaseDirectory, string BinaryName, UnrealTargetPlatform Platform, UnrealTargetConfiguration Configuration, UnrealTargetConfiguration UndecoratedConfiguration, bool bIncludesGameModules, string ExeSubFolder)
 		{
-			if (String.IsNullOrEmpty(ModuleName) && Configuration == UnrealTargetConfiguration.DebugGame && !bCompileMonolithic)
+			// Get the configuration for the executable. If we're building DebugGame, and this executable only contains engine modules, use the same name as development.
+			UnrealTargetConfiguration ExeConfiguration = Configuration;
+			if(Configuration == UnrealTargetConfiguration.DebugGame && !bIncludesGameModules)
 			{
-				return MakeBinaryPaths(ModuleName, BinaryName, Platform, UnrealTargetConfiguration.Development, BinaryType, TargetType, PluginInfo, Rules.UndecoratedConfiguration);
+				ExeConfiguration = UnrealTargetConfiguration.Development;
 			}
-			else
+
+			// Build the binary path
+			string BinaryPath = Path.Combine(BaseDirectory, "Binaries", Platform.ToString());
+			if (!String.IsNullOrEmpty(ExeSubFolder))
 			{
-				return MakeBinaryPaths(ModuleName, BinaryName, Platform, Configuration, BinaryType, TargetType, PluginInfo, Rules.UndecoratedConfiguration, ExeBinariesSubFolder);
+				BinaryPath = Path.Combine(BinaryPath, ExeSubFolder);
 			}
+			BinaryPath = Path.Combine(BinaryPath, MakeBinaryFileName(BinaryName, Platform, ExeConfiguration, UndecoratedConfiguration, UEBuildBinaryType.Executable));
+
+			// Allow the platform to customize the output path (and output several executables at once if necessary)
+			return UEBuildPlatform.GetBuildPlatform(Platform).FinalizeBinaryPaths(BinaryPath);
 		}
 
 		/// <summary>
diff --git a/Engine/Source/Programs/UnrealBuildTool/System/UnrealBuildTool.cs b/Engine/Source/Programs/UnrealBuildTool/System/UnrealBuildTool.cs
--- a/Engine/Source/Programs/UnrealBuildTool/System/UnrealBuildTool.cs
+++ b/Engine/Source/Programs/UnrealBuildTool/System/UnrealBuildTool.cs
@@ -1922,10 +1922,10 @@ namespace UnrealBuildTool
 
 			if (!ProjectFileGenerator.bGenerateProjectFiles && !UEBuildConfiguration.bGenerateManifest && bIsEditorTarget )
 			{
-				var EditorProcessFilenames = UEBuildTarget.MakeBinaryPaths("Launcher", "UE4Editor", TargetDesc.Platform, TargetDesc.Configuration, UEBuildBinaryType.Executable, TargetRules.TargetType.Editor, null, UnrealTargetConfiguration.Development);
+				var EditorProcessFilenames = UEBuildTarget.MakeExecutablePaths("..", "UE4Editor", TargetDesc.Platform, TargetDesc.Configuration, UnrealTargetConfiguration.Development, false, null);
 				if (EditorProcessFilenames.Length != 1)
 				{
-					throw new BuildException("ShouldDoHotReload cannot handle multiple binaries returning from UEBuildTarget.MakeBinaryPaths");
+					throw new BuildException("ShouldDoHotReload cannot handle multiple binaries returning from UEBuildTarget.MakeExecutablePaths");
 				}
 
 				var EditorProcessFilename = EditorProcessFilenames[0];
