Add an accurate "IsMonolithic" setting to TargetInfo, rather than just guessing from the target type.

diff --git a/Engine/Source/Programs/UnrealBuildTool/Configuration/UEBuildTarget.cs b/Engine/Source/Programs/UnrealBuildTool/Configuration/UEBuildTarget.cs
--- a/Engine/Source/Programs/UnrealBuildTool/Configuration/UEBuildTarget.cs
+++ b/Engine/Source/Programs/UnrealBuildTool/Configuration/UEBuildTarget.cs
@@ -657,6 +657,9 @@ namespace UnrealBuildTool
 		/** Target as defined by the VCProject and passed via the command line. Not necessarily the same as internal name. */
 		public UnrealTargetConfiguration Configuration;
 
+		/** TargetInfo object which can be passed to RulesCompiler */
+		public TargetInfo TargetInfo;
+
 		/** Root directory for the active project. Typically contains the .uproject file, or the engine root. */
 		public string ProjectDirectory;
 
@@ -816,6 +819,8 @@ namespace UnrealBuildTool
 
 			TargetType = Rules.Type;
 
+			TargetInfo = new TargetInfo(Platform, Configuration, Rules.Type, bCompileMonolithic);
+
 			// Construct the output paths for this target's executable
 			string OutputDirectory;
 			if((bCompileMonolithic || TargetType == TargetRules.TargetType.Program) && !Rules.bOutputToEngineBinaries)
@@ -1167,7 +1172,7 @@ namespace UnrealBuildTool
 			{
 				// Create the module rules
 				string ModuleRulesFileName;
-				ModuleRules Rules = RulesCompiler.CreateModuleRules(ModuleName, GetTargetInfo(), out ModuleRulesFileName);
+				ModuleRules Rules = RulesCompiler.CreateModuleRules(ModuleName, TargetInfo, out ModuleRulesFileName);
 
 				// Add the rules file itself
 				FileNames.Add(ModuleRulesFileName);
@@ -1850,11 +1855,10 @@ namespace UnrealBuildTool
 			UEBuildBinaryCPP BinaryCPP = ExecutableBinary as UEBuildBinaryCPP;
 			if (BinaryCPP != null)
 			{
-				TargetInfo CurrentTarget = new TargetInfo(Platform, Configuration, TargetType);
 				foreach (var TargetModuleName in BinaryCPP.ModuleNames)
 				{
 					string UnusedFilename;
-					ModuleRules CheckRules = RulesCompiler.CreateModuleRules(TargetModuleName, CurrentTarget, out UnusedFilename);
+					ModuleRules CheckRules = RulesCompiler.CreateModuleRules(TargetModuleName, TargetInfo, out UnusedFilename);
 					if ((CheckRules != null) && (CheckRules.Type != ModuleRules.ModuleType.External))
 					{
 						PrivateDependencyModuleNames.Add(TargetModuleName);
@@ -2175,7 +2179,7 @@ namespace UnrealBuildTool
 			{
 				// Find all the precompiled module names.
 				List<string> PrecompiledModuleNames = new List<string>();
-				Rules.GetModulesToPrecompile(new TargetInfo(Platform, Configuration, TargetType), PrecompiledModuleNames);
+				Rules.GetModulesToPrecompile(TargetInfo, PrecompiledModuleNames);
 
 				// Add all the enabled plugins to the precompiled module list. Plugins are always precompiled, even if bPrecompile is not set, so we should precompile their dependencies.
 				foreach(PluginInfo Plugin in BuildPlugins)
@@ -2274,11 +2278,11 @@ namespace UnrealBuildTool
 			string[] OutputFilePaths;
 			if(TargetType == TargetRules.TargetType.Game && UnrealBuildTool.RunningRocket())
 			{
-				OutputFilePaths = MakeBinaryPaths(ModuleName, "UE4Game-" + ModuleName, Platform, Configuration, BinaryType, TargetType, Plugin, Rules.UndecoratedConfiguration);
+				OutputFilePaths = MakeBinaryPaths(ModuleName, "UE4Game-" + ModuleName, BinaryType, Plugin);
 			}
 			else
 			{
-				OutputFilePaths = MakeBinaryPaths(ModuleName, AppName + "-" + ModuleName, Platform, Configuration, BinaryType, TargetType, Plugin, Rules.UndecoratedConfiguration);
+				OutputFilePaths = MakeBinaryPaths(ModuleName, AppName + "-" + ModuleName, BinaryType, Plugin);
 			}
 
 			// Get the intermediate path
@@ -2403,8 +2407,7 @@ namespace UnrealBuildTool
 		}
 
 		/** Given a UBT-built binary name (e.g. "Core"), returns a relative path to the binary for the current build configuration (e.g. "../Binaries/Win64/Core-Win64-Debug.lib") */
-		public static string[] MakeBinaryPaths(string ModuleName, string BinaryName, UnrealTargetPlatform Platform, 
-			UnrealTargetConfiguration Configuration, UEBuildBinaryType BinaryType, TargetRules.TargetType TargetType, PluginInfo PluginInfo, UnrealTargetConfiguration UndecoratedConfiguration)
+		private string[] MakeBinaryPaths(string ModuleName, string BinaryName, UEBuildBinaryType BinaryType, PluginInfo PluginInfo)
 		{
 			// Determine the binary extension for the platform and binary type.
 			var BuildPlatform = UEBuildPlatform.GetBuildPlatform(Platform);
@@ -2424,7 +2427,7 @@ namespace UnrealBuildTool
 				// Plugins can be binary-only and can have no rules object
 				if (PluginInfo == null || !String.IsNullOrEmpty(RulesFilename))
 				{
-					ModuleRules ModuleRulesObj = RulesCompiler.CreateModuleRules(ModuleName, new TargetInfo(Platform, Configuration, TargetType), out RulesFilename);
+					ModuleRules ModuleRulesObj = RulesCompiler.CreateModuleRules(ModuleName, TargetInfo, out RulesFilename);
 					if (ModuleRulesObj != null)
 					{
 						ModuleBinariesSubDir = ModuleRulesObj.BinariesSubFolder;
@@ -2463,7 +2466,7 @@ namespace UnrealBuildTool
 			}
 
 			// append the binary name
-			string OutBinaryPath = Path.Combine(BaseDirectory, MakeBinaryFileName(BinaryName, Platform, LocalConfig, UndecoratedConfiguration, BinaryType));
+			string OutBinaryPath = Path.Combine(BaseDirectory, MakeBinaryFileName(BinaryName, Platform, LocalConfig, Rules.UndecoratedConfiguration, BinaryType));
 			return BuildPlatform.FinalizeBinaryPaths(OutBinaryPath);
 		}
 
@@ -2533,7 +2536,7 @@ namespace UnrealBuildTool
 		{
 			var BuildPlatform = UEBuildPlatform.GetBuildPlatform(Platform);
 			List<string> PlatformExtraModules = new List<string>();
-			BuildPlatform.GetExtraModules(new TargetInfo(Platform, Configuration, TargetType), this, ref PlatformExtraModules);
+			BuildPlatform.GetExtraModules(TargetInfo, this, ref PlatformExtraModules);
 			ExtraModuleNames.AddRange(PlatformExtraModules);			
 		}
 
@@ -2614,7 +2617,7 @@ namespace UnrealBuildTool
 				List<UEBuildBinaryConfiguration> RulesBuildBinaryConfigurations = new List<UEBuildBinaryConfiguration>();
 				List<string> RulesExtraModuleNames = new List<string>();
 				Rules.SetupBinaries(
-					new TargetInfo(Platform, Configuration, TargetType),
+					TargetInfo,
 					ref RulesBuildBinaryConfigurations,
 					ref RulesExtraModuleNames
 					);
@@ -2660,7 +2663,7 @@ namespace UnrealBuildTool
 				LinkEnvironmentConfiguration RulesLinkEnvConfig = GlobalLinkEnvironment.Config;
 				CPPEnvironmentConfiguration RulesCPPEnvConfig = GlobalCompileEnvironment.Config;
 				SetupDefaultGlobalEnvironment(
-						new TargetInfo(Platform, Configuration, TargetType),
+						TargetInfo,
 						ref RulesLinkEnvConfig,
 						ref RulesCPPEnvConfig
 						);
@@ -2687,7 +2690,7 @@ namespace UnrealBuildTool
 					Rules.ShouldCompileMonolithic(Platform, Configuration))
 				{
 					Rules.SetupGlobalEnvironment(
-						new TargetInfo(Platform, Configuration, TargetType),
+						TargetInfo,
 						ref RulesLinkEnvConfig,
 						ref RulesCPPEnvConfig
 						);
@@ -2906,12 +2909,6 @@ namespace UnrealBuildTool
 			UEBuildPlatform.GetBuildPlatform(Platform).SetUpProjectEnvironment(Platform);
 		}
 
-		/** Constructs a TargetInfo object for this target. */
-		public TargetInfo GetTargetInfo()
-		{
-			return new TargetInfo(Platform, Configuration, TargetType);
-		}
-
 		/** Registers a module with this target. */
 		public void RegisterModule(UEBuildModule Module)
 		{
@@ -2925,14 +2922,12 @@ namespace UnrealBuildTool
 			UEBuildModule Module;
 			if (!Modules.TryGetValue(ModuleName, out Module))
 			{
-				TargetInfo TargetInfo = GetTargetInfo();
-
 				// Create the module!  (It will be added to our hash table in its constructor)
 
 				// @todo projectfiles: Cross-platform modules can appear here during project generation, but they may have already
 				//   been filtered out by the project generator.  This causes the projects to not be added to directories properly.
 				string ModuleFileName;
-				var RulesObject = RulesCompiler.CreateModuleRules(ModuleName, GetTargetInfo(), out ModuleFileName);
+				var RulesObject = RulesCompiler.CreateModuleRules(ModuleName, TargetInfo, out ModuleFileName);
 				var ModuleDirectory = Path.GetDirectoryName(ModuleFileName);
 
 				// Making an assumption here that any project file path that contains '../../'
diff --git a/Engine/Source/Programs/UnrealBuildTool/System/ActionGraph.cs b/Engine/Source/Programs/UnrealBuildTool/System/ActionGraph.cs
--- a/Engine/Source/Programs/UnrealBuildTool/System/ActionGraph.cs
+++ b/Engine/Source/Programs/UnrealBuildTool/System/ActionGraph.cs
@@ -813,7 +813,7 @@ namespace UnrealBuildTool
 					// Scan this file for included headers that may be out of date.  Note that it's OK if we break out early because we found
 					// the action to be outdated.  For outdated actions, we kick off a separate include scan in a background thread later on to
 					// catch all of the other includes and form an exhaustive set.
-					var BuildPlatform = UEBuildPlatform.GetBuildPlatform( Target.GetTargetInfo().Platform );
+					var BuildPlatform = UEBuildPlatform.GetBuildPlatform( Target.Platform );
 					foreach (FileItem PrerequisiteItem in RootAction.PrerequisiteItems)
 					{
 						// @todo ubtmake: Make sure we are catching RC files here too.  Anything that the toolchain would have tried it on.  Logic should match the CACHING stuff below
diff --git a/Engine/Source/Programs/UnrealBuildTool/System/RulesCompiler.cs b/Engine/Source/Programs/UnrealBuildTool/System/RulesCompiler.cs
--- a/Engine/Source/Programs/UnrealBuildTool/System/RulesCompiler.cs
+++ b/Engine/Source/Programs/UnrealBuildTool/System/RulesCompiler.cs
@@ -26,22 +26,38 @@ namespace UnrealBuildTool
 		/// Target type (if known)
 		public readonly TargetRules.TargetType? Type;
 
+		/// Whether the target is monolithic (if known)
+		public readonly bool? bIsMonolithic;
+
 		/// <summary>
 		/// Constructs a TargetInfo
 		/// </summary>
 		/// <param name="InitPlatform">Target platform</param>
 		/// <param name="InitConfiguration">Target build configuration</param>
-		public TargetInfo( UnrealTargetPlatform InitPlatform, UnrealTargetConfiguration InitConfiguration, TargetRules.TargetType? InitType = null )
+		public TargetInfo( UnrealTargetPlatform InitPlatform, UnrealTargetConfiguration InitConfiguration )
 		{
 			Platform = InitPlatform;
 			Configuration = InitConfiguration;
-			Type = InitType;
 
 			// get the platform's architecture
 			var BuildPlatform = UEBuildPlatform.GetBuildPlatform(Platform);
 			Architecture = BuildPlatform.GetActiveArchitecture();
 		}
 
+		/// <summary>
+		/// Constructs a TargetInfo
+		/// </summary>
+		/// <param name="InitPlatform">Target platform</param>
+		/// <param name="InitConfiguration">Target build configuration</param>
+		/// <param name="InitType">Target type</param>
+		/// <param name="bInitIsMonolithic">Whether the target is monolithic</param>
+		public TargetInfo( UnrealTargetPlatform InitPlatform, UnrealTargetConfiguration InitConfiguration, TargetRules.TargetType InitType, bool bInitIsMonolithic )
+			: this(InitPlatform, InitConfiguration)
+		{
+			Type = InitType;
+			bIsMonolithic = bInitIsMonolithic;
+		}
+
 		/// <summary>
 		/// True if the target type is a cooked game.
 		/// </summary>
@@ -66,13 +82,11 @@ namespace UnrealBuildTool
         {
             get
             {
-                if (!Type.HasValue)
+                if (!bIsMonolithic.HasValue)
                 {
-                    throw new BuildException("Trying to access TargetInfo.IsMonolithic when TargetInfo.Type is not set. Make sure IsMonolithic is used only in ModuleRules.");
+                    throw new BuildException("Trying to access TargetInfo.IsMonolithic when bIsMonolithic is not set. Make sure IsMonolithic is used only in ModuleRules.");
                 }
-                return Type == TargetRules.TargetType.Client ||
-                     Type == TargetRules.TargetType.Game ||
-                     Type == TargetRules.TargetType.Server;
+                return bIsMonolithic.Value;
             }
         }
     }
@@ -1431,7 +1445,7 @@ namespace UnrealBuildTool
 
 			UnrealTargetPlatform LocalPlatform = Target.Platform;
 			UnrealTargetConfiguration LocalConfiguration = Target.Configuration;
-			TargetInfo LocalTarget = new TargetInfo(LocalPlatform, LocalConfiguration, Target.Type);
+			TargetInfo LocalTarget = new TargetInfo(LocalPlatform, LocalConfiguration, Target.Type.Value, Target.bIsMonolithic.Value);
 
 			// The build module must define a type named 'Rules' that derives from our 'ModuleRules' type.  
 			var RulesObjectType = RulesAssembly.GetType( ModuleName );
diff --git a/Engine/Source/Programs/UnrealBuildTool/System/UnrealBuildTool.cs b/Engine/Source/Programs/UnrealBuildTool/System/UnrealBuildTool.cs
--- a/Engine/Source/Programs/UnrealBuildTool/System/UnrealBuildTool.cs
+++ b/Engine/Source/Programs/UnrealBuildTool/System/UnrealBuildTool.cs
@@ -1725,7 +1725,7 @@ namespace UnrealBuildTool
                                         ECompilationResult UHTResult = ECompilationResult.OtherCompilationError;
                                         if (!ExternalExecution.ExecuteHeaderToolIfNecessary(Target, GlobalCompileEnvironment:null, UObjectModules:TargetUObjectModules, ModuleInfoFileName:ModuleInfoFileName, UHTResult:ref UHTResult))
                                         {
-                                            Log.TraceInformation("UnrealHeaderTool failed for target '" + Target.GetTargetName() + "' (platform: " + Target.GetTargetInfo().Platform.ToString() + ", module info: " + ModuleInfoFileName + ").");
+                                            Log.TraceInformation("UnrealHeaderTool failed for target '" + Target.GetTargetName() + "' (platform: " + Target.Platform.ToString() + ", module info: " + ModuleInfoFileName + ").");
                                             BuildResult = UHTResult;
                                             break;
                                         }
@@ -2128,7 +2128,7 @@ namespace UnrealBuildTool
                         {
                             // Invoke our deep include scanner to figure out whether any of the files included by this source file have
                             // changed since the build product was built
-                            var FileBuildPlatform = UEBuildPlatform.GetBuildPlatform( Target.GetTargetInfo().Platform );
+                            var FileBuildPlatform = UEBuildPlatform.GetBuildPlatform( Target.Platform );
                             CPPEnvironment.FindAndCacheAllIncludedFiles( Target, PrerequisiteItem, FileBuildPlatform, PrerequisiteItem.CachedCPPIncludeInfo, bOnlyCachedDependencies:false );
                         }
                     }
