Allowed rocket builds to check for existence of .uprojectdirs files

diff --git a/Engine/Source/Programs/UnrealBuildTool/System/UnrealBuildTool.cs b/Engine/Source/Programs/UnrealBuildTool/System/UnrealBuildTool.cs
--- a/Engine/Source/Programs/UnrealBuildTool/System/UnrealBuildTool.cs
+++ b/Engine/Source/Programs/UnrealBuildTool/System/UnrealBuildTool.cs
@@ -671,7 +671,7 @@ namespace UnrealBuildTool
 				string IntermediateSourceFolder = Path.Combine(Path.GetDirectoryName(GetUProjectFile()), "Intermediate", "Source");
                 UProjectInfo.AddProject(GetUProjectFile(), Directory.Exists(SourceFolder) || Directory.Exists(IntermediateSourceFolder));
             }
-            else if(!RunningRocket())
+            else
             {
                 UProjectInfo.FillProjectInfo();
             }
