Add Serializable keyword to TargetInfo for UBT makefiles.

diff --git a/Engine/Source/Programs/UnrealBuildTool/System/RulesCompiler.cs b/Engine/Source/Programs/UnrealBuildTool/System/RulesCompiler.cs
--- a/Engine/Source/Programs/UnrealBuildTool/System/RulesCompiler.cs
+++ b/Engine/Source/Programs/UnrealBuildTool/System/RulesCompiler.cs
@@ -12,6 +12,7 @@ namespace UnrealBuildTool
 	/// <summary>
 	/// Information about a target, passed along when creating a module descriptor
 	/// </summary>
+	[Serializable]
 	public class TargetInfo
 	{
 		/// Target platform
