- Updated GearVR manifest generation for 0.5.0
- Mostly removing the platform activity and the camera permission.

diff --git a/Engine/Source/Programs/UnrealBuildTool/Android/UEDeployAndroid.cs b/Engine/Source/Programs/UnrealBuildTool/Android/UEDeployAndroid.cs
--- a/Engine/Source/Programs/UnrealBuildTool/Android/UEDeployAndroid.cs
+++ b/Engine/Source/Programs/UnrealBuildTool/Android/UEDeployAndroid.cs
@@ -712,11 +712,6 @@ namespace UnrealBuildTool.Android
 			{
 				Text.AppendLine("\t\t<meta-data android:name=\"com.samsung.android.vr.application.mode\"");
 				Text.AppendLine("\t\t           android:value=\"vr_only\" />");
-				Text.AppendLine("\t\t<activity android:name=\"com.oculusvr.vrlib.PlatformActivity\"");
-				Text.AppendLine("\t\t          android:theme=\"@android:style/Theme.Black.NoTitleBar.Fullscreen\"");
-				Text.AppendLine("\t\t          android:launchMode=\"singleTask\"");
-				Text.AppendLine("\t\t          android:screenOrientation=\"landscape\"");
-				Text.AppendLine("\t\t          android:configChanges=\"screenSize|orientation|keyboardHidden|keyboard\"/>");
 			}
 			if (!string.IsNullOrEmpty(ExtraApplicationSettings))
 			{
@@ -751,8 +746,6 @@ namespace UnrealBuildTool.Android
 			}
 			if (bPackageForGearVR)
 			{
-				Text.AppendLine("\t<uses-permission android:name=\"android.permission.CAMERA\"/>");
-				Text.AppendLine("\t<uses-feature android:name=\"android.hardware.camera\"/>");
 				Text.AppendLine("\t<uses-feature android:name=\"android.hardware.usb.host\"/>");
 			}
 			if (ExtraPermissions != null)
