Allow targets to opt-in to using a shared build environment with other targets. Doing so allows intermediates and precompiled libraries (including plugins) to be used in the current build. Also move plugin output files and intermediates into the plugin directory.

diff --git a/Engine/Source/Programs/UnrealBuildTool/Configuration/UEBuildClient.cs b/Engine/Source/Programs/UnrealBuildTool/Configuration/UEBuildClient.cs
--- a/Engine/Source/Programs/UnrealBuildTool/Configuration/UEBuildClient.cs
+++ b/Engine/Source/Programs/UnrealBuildTool/Configuration/UEBuildClient.cs
@@ -11,7 +11,7 @@ namespace UnrealBuildTool
     public class UEBuildClient : UEBuildTarget
 	{
 		public UEBuildClient(TargetDescriptor InDesc, TargetRules InRulesObject)
-            : base(GetDefaultAppName(InDesc, InRulesObject, "Client"), InDesc, InRulesObject)
+            : base(InDesc, InRulesObject, "UE4Client")
         {
             if (ShouldCompileMonolithic())
             {
diff --git a/Engine/Source/Programs/UnrealBuildTool/Configuration/UEBuildEditor.cs b/Engine/Source/Programs/UnrealBuildTool/Configuration/UEBuildEditor.cs
--- a/Engine/Source/Programs/UnrealBuildTool/Configuration/UEBuildEditor.cs
+++ b/Engine/Source/Programs/UnrealBuildTool/Configuration/UEBuildEditor.cs
@@ -11,7 +11,7 @@ namespace UnrealBuildTool
 	public class UEBuildEditor : UEBuildTarget
 	{
 		public UEBuildEditor(TargetDescriptor InDesc, TargetRules InRulesObject)
-			: base(GetDefaultAppName(InDesc, InRulesObject, "Editor"), InDesc, InRulesObject)
+			: base(InDesc, InRulesObject, "UE4Editor")
 		{
 		}
 
diff --git a/Engine/Source/Programs/UnrealBuildTool/Configuration/UEBuildGame.cs b/Engine/Source/Programs/UnrealBuildTool/Configuration/UEBuildGame.cs
--- a/Engine/Source/Programs/UnrealBuildTool/Configuration/UEBuildGame.cs
+++ b/Engine/Source/Programs/UnrealBuildTool/Configuration/UEBuildGame.cs
@@ -11,7 +11,7 @@ namespace UnrealBuildTool
 	public class UEBuildGame : UEBuildTarget
 	{
 		public UEBuildGame(TargetDescriptor InDesc, TargetRules InRulesObject)
-			: base(GetDefaultAppName(InDesc, InRulesObject, ""), InDesc, InRulesObject)
+			: base(InDesc, InRulesObject, "UE4")
 		{
 			if (ShouldCompileMonolithic())
 			{
diff --git a/Engine/Source/Programs/UnrealBuildTool/Configuration/UEBuildServer.cs b/Engine/Source/Programs/UnrealBuildTool/Configuration/UEBuildServer.cs
--- a/Engine/Source/Programs/UnrealBuildTool/Configuration/UEBuildServer.cs
+++ b/Engine/Source/Programs/UnrealBuildTool/Configuration/UEBuildServer.cs
@@ -11,7 +11,7 @@ namespace UnrealBuildTool
 	public class UEBuildServer : UEBuildTarget
 	{
 		public UEBuildServer(TargetDescriptor InDesc, TargetRules InRulesObject)
-			: base(GetDefaultAppName(InDesc, InRulesObject, "Server"), InDesc, InRulesObject)
+			: base(InDesc, InRulesObject, "UE4Server")
 		{
 		}
 
diff --git a/Engine/Source/Programs/UnrealBuildTool/Configuration/UEBuildTarget.cs b/Engine/Source/Programs/UnrealBuildTool/Configuration/UEBuildTarget.cs
--- a/Engine/Source/Programs/UnrealBuildTool/Configuration/UEBuildTarget.cs
+++ b/Engine/Source/Programs/UnrealBuildTool/Configuration/UEBuildTarget.cs
@@ -608,36 +608,6 @@ namespace UnrealBuildTool
 			return AllGameFolders;
 		}
 
-		/// <summary>
-		/// Finds the base name for all binaries produced for the given target
-		/// </summary>
-		/// <param name="InDesc">Target descriptor</param>
-		/// <param name="InRulesObject">Target rules object created by RulesCompiler</param>
-		/// <param name="InNameAppend">Suffix to append to the name</param>
-		/// <returns>Base name for the application</returns>
-		public static string GetDefaultAppName(TargetDescriptor InDesc, TargetRules InRulesObject, string InNameAppend)
-		{
-			//@todo. Allow failure to get build platform here?
-			bool bPlatformRequiresMonolithic = false;
-			var BuildPlatform = UEBuildPlatform.GetBuildPlatform(InDesc.Platform, true);
-			if (BuildPlatform != null)
-			{
-				bPlatformRequiresMonolithic = BuildPlatform.ShouldCompileMonolithicBinary(InDesc.Platform);
-			}
-
-			// If we're building a monolithic binary, then the game and engine code are linked together into one
-			// program executable, so we want the application name to be the game name.  In the case of a modular
-			// binary, we use 'UnrealEngine' for our application name
-			if (InRulesObject.ShouldCompileMonolithic(InDesc.Platform, InDesc.Configuration) || bPlatformRequiresMonolithic)
-			{
-				return InDesc.TargetName;
-			}
-			else
-			{
-				return "UE4" + InNameAppend;
-			}
-		}
-
 		/** The target rules */
 		[NonSerialized]
 		public TargetRules Rules = null;
@@ -645,12 +615,15 @@ namespace UnrealBuildTool
 		/** Type of target */
 		public TargetRules.TargetType TargetType;
 
-		/** The name of the application the target is part of. For non-monolithic targets, this is typically the name of the base application, eg. UE4Editor for any game editor. */
+		/** The name of the application the target is part of. For targets with bUseSharedBuildEnvironment = true, this is typically the name of the base application, eg. UE4Editor for any game editor. */
 		public string AppName;
 
 		/** The name of the target */
 		public string TargetName;
 
+		/** Whether the target uses the shared build environment. If false, AppName==TargetName and all binaries should be written to the project directory. */
+		public bool bUseSharedBuildEnvironment;
+
 		/** Platform as defined by the VCProject and passed via the command line. Not the same as internal config names. */
 		public UnrealTargetPlatform Platform;
 
@@ -748,21 +721,20 @@ namespace UnrealBuildTool
 			return bCompileMonolithic;	// @todo ubtmake: We need to make sure this function and similar things aren't called in assembler mode
 		}
 
-		/** 
-		 * @param InAppName - The name of the application being built, which is used to scope all intermediate and output file names.
-		 * @param InGameName - The name of the game being build - can be empty
-		 * @param InPlatform - The platform the target is being built for.
-		 * @param InConfiguration - The configuration the target is being built for.
-		 * @param InAdditionalDefinitions - Additional definitions provided on the UBT command-line for the target.
-		 * @param InRemoteRoot - The remote path that the build output is synced to.
-		 */
-		public UEBuildTarget(string InAppName, TargetDescriptor InDesc, TargetRules InRules)
+		/// <summary>
+		/// Constructor.
+		/// </summary>
+		/// <param name="InDesc">Target descriptor</param>
+		/// <param name="InRules">The target rules, as created by RulesCompiler.</param>
+		/// <param name="InPossibleAppName">The AppName for shared binaries of this target type, if used (null if there is none).</param>
+		public UEBuildTarget(TargetDescriptor InDesc, TargetRules InRules, string InPossibleAppName)
 		{
-			AppName = InAppName;
+			AppName = InDesc.TargetName;
 			TargetName = InDesc.TargetName;
 			Platform = InDesc.Platform;
 			Configuration = InDesc.Configuration;
 			Rules = InRules;
+			TargetType = Rules.Type;
 			bEditorRecompile = InDesc.bIsEditorRecompile;
 			bPrecompile = InDesc.bPrecompile;
 			bUsePrecompiled = InDesc.bUsePrecompiled;
@@ -785,6 +757,14 @@ namespace UnrealBuildTool
 				}
 			}
 
+			TargetInfo = new TargetInfo(Platform, Configuration, Rules.Type, bCompileMonolithic);
+
+			if(InPossibleAppName != null && InRules.ShouldUseSharedBuildEnvironment(TargetInfo))
+			{
+				AppName = InPossibleAppName;
+				bUseSharedBuildEnvironment = true;
+			}
+
 			// Figure out what the project directory is. If we have a uproject file, use that. Otherwise use the engine directory.
 			if (UnrealBuildTool.HasUProjectFile())
 			{
@@ -800,7 +780,7 @@ namespace UnrealBuildTool
 
 			// Build the engine intermediate directory. If we're building agnostic engine binaries, we can use the engine intermediates folder. Otherwise we need to use the project intermediates directory.
 			// If we are building a content project with UE4Game as the base executable, then we don't want to use the Project as the intermediate, it should use UE4's intermedate
-			if (ShouldCompileMonolithic() && AppName != "UE4Game")
+			if (!bUseSharedBuildEnvironment && AppName != "UE4Game")
 			{
 				EngineIntermediateDirectory = ProjectIntermediateDirectory;
 			}
@@ -817,10 +797,6 @@ namespace UnrealBuildTool
 
 			OnlyModules = InDesc.OnlyModules;
 
-			TargetType = Rules.Type;
-
-			TargetInfo = new TargetInfo(Platform, Configuration, Rules.Type, bCompileMonolithic);
-
 			// Construct the output paths for this target's executable
 			string OutputDirectory;
 			if((bCompileMonolithic || TargetType == TargetRules.TargetType.Program) && !Rules.bOutputToEngineBinaries)
@@ -831,7 +807,7 @@ namespace UnrealBuildTool
 			{
 				OutputDirectory = Path.GetFullPath(BuildConfiguration.RelativeEnginePath);
 			}
-			OutputPaths = MakeExecutablePaths(OutputDirectory, AppName, Platform, Configuration, Rules.UndecoratedConfiguration, bCompileMonolithic && UnrealBuildTool.HasUProjectFile(), Rules.ExeBinariesSubFolder);
+			OutputPaths = MakeExecutablePaths(OutputDirectory, bCompileMonolithic? TargetName : AppName, Platform, Configuration, Rules.UndecoratedConfiguration, bCompileMonolithic && UnrealBuildTool.HasUProjectFile(), Rules.ExeBinariesSubFolder);
 
 			// handle some special case defines (so build system can pass -DEFINE as normal instead of needing
 			// to know about special parameters)
@@ -988,6 +964,15 @@ namespace UnrealBuildTool
 					}
 				}
 
+				// Delete the intermediate folder for each binary. This will catch all plugin intermediate folders, as well as any project and engine folders.
+				foreach(UEBuildBinary Binary in AppBinaries)
+				{
+					if(Directory.Exists(Binary.Config.IntermediateDirectory))
+					{
+						CleanDirectory(Binary.Config.IntermediateDirectory);
+					}
+				}
+
 				// Generate a list of all the modules of each AppBinaries entry
 				var ModuleList = new List<string>();
 				var bTargetUsesUObjectModule = false;
@@ -2148,12 +2133,14 @@ namespace UnrealBuildTool
 		/// </summary>
         public void AddPlugin(PluginInfo Plugin)
 		{
+			UEBuildBinaryType BinaryType = ShouldCompileMonolithic() ? UEBuildBinaryType.StaticLibrary : UEBuildBinaryType.DynamicLinkLibrary;
 			foreach(PluginInfo.PluginModuleInfo Module in Plugin.Modules)
 			{
 				if (ShouldIncludePluginModule(Plugin, Module))
 				{
-					UEBuildBinaryType BinaryType = ShouldCompileMonolithic() ? UEBuildBinaryType.StaticLibrary : UEBuildBinaryType.DynamicLinkLibrary;
-					AddBinaryForModule(Module.Name, BinaryType, true, false);
+					string ModuleFileName = RulesCompiler.GetModuleFilename(Module.Name);
+					bool bHasSource = (ModuleFileName != null && Directory.EnumerateFiles(Path.GetDirectoryName(ModuleFileName), "*.cpp", SearchOption.AllDirectories).Any());
+					AddBinaryForModule(Module.Name, BinaryType, bHasSource, false);
 				}
 			}
 		}
@@ -2268,68 +2255,62 @@ namespace UnrealBuildTool
 			return Binary;
 		}
 
-		protected UEBuildBinaryCPP AddBinaryForModule(string ModuleName, UEBuildBinaryType BinaryType, bool bAllowCompilation, bool bIsCrossTarget)
+		/// <summary>
+		/// Adds a binary for the given module. Does not check whether a binary already exists, or whether a binary should be created for this build configuration.
+		/// </summary>
+		/// <param name="ModuleName">Name of the binary</param>
+		/// <param name="BinaryType">Type of binary to be created</param>
+		/// <param name="bAllowCompilation">Whether this binary can be compiled. The function will check whether plugin binaries can be compiled.</param>
+		/// <param name="bIsCrossTarget">True if module is for supporting a different target-platform</param>
+		/// <returns>The new binary</returns>
+		private UEBuildBinaryCPP AddBinaryForModule(string ModuleName, UEBuildBinaryType BinaryType, bool bAllowCompilation, bool bIsCrossTarget)
 		{
 			// Get the plugin info for the module
 			PluginInfo Plugin = Plugins.GetPluginInfoForModule(ModuleName);
 
-			// Get the output paths for it. 
-			// HACK: Rocket does not allow overriding the compile environment, so always use the UE4Game prefix. This should be fixed by allowing targets choose whether to use a custom build environment.
-			string[] OutputFilePaths;
-			if(TargetType == TargetRules.TargetType.Game && UnrealBuildTool.RunningRocket())
+			// Get the root output directory and base name (target name/app name) for this binary
+			string BaseOutputDirectory;
+			if(Plugin != null)
 			{
-				OutputFilePaths = MakeBinaryPaths(ModuleName, "UE4Game-" + ModuleName, BinaryType, Plugin);
+				BaseOutputDirectory = Path.GetFullPath(Plugin.Directory);
 			}
-			else
+			else if(RulesCompiler.IsGameModule(ModuleName) || !bUseSharedBuildEnvironment)
 			{
-				OutputFilePaths = MakeBinaryPaths(ModuleName, AppName + "-" + ModuleName, BinaryType, Plugin);
+				BaseOutputDirectory = Path.GetFullPath(ProjectDirectory);
 			}
-
-			// Get the intermediate path
-			string IntermediateDirectory;
-			if(Plugin == null)
+			else
 			{
-				if(RulesCompiler.IsGameModule(ModuleName))
-				{
-					IntermediateDirectory = ProjectIntermediateDirectory;
-				}
-				else
-				{
-					IntermediateDirectory = EngineIntermediateDirectory;
-				}
+				BaseOutputDirectory = Path.GetFullPath(BuildConfiguration.RelativeEnginePath);
 			}
-			else
+
+			// Get the configuration that this module will be built in. Engine modules compiled in DebugGame will use Development.
+			UnrealTargetConfiguration ModuleConfiguration = Configuration;
+			if(Configuration == UnrealTargetConfiguration.DebugGame && !RulesCompiler.IsGameModule(ModuleName))
 			{
-				if(Plugin.LoadedFrom == PluginInfo.LoadedFromType.GameProject)
-				{
-					IntermediateDirectory = Path.Combine(ProjectIntermediateDirectory, "Plugins", ShouldCompileMonolithic() ? "Static" : "Dynamic");
-				}
-				else
-				{
-					IntermediateDirectory = Path.Combine(EngineIntermediateDirectory, "Plugins", ShouldCompileMonolithic() ? "Static" : "Dynamic");
-				}
+				ModuleConfiguration = UnrealTargetConfiguration.Development;
 			}
 
-			// If it's a plugin, check if it actually has source.
-			bool bHasModuleRules = true;
-			if(Plugin != null)
+			// Get the output and intermediate directories for this module
+			string OutputDirectory = Path.Combine(BaseOutputDirectory, "Binaries", Platform.ToString());
+ 			string IntermediateDirectory = Path.Combine(BaseOutputDirectory, BuildConfiguration.PlatformIntermediateFolder, AppName, ModuleConfiguration.ToString());
+
+			// Append a subdirectory if the module rules specifies one
+			ModuleRules ModuleRules;
+			if(RulesCompiler.TryCreateModuleRules(ModuleName, TargetInfo, out ModuleRules) && !String.IsNullOrEmpty(ModuleRules.BinariesSubFolder))
 			{
-				string ModuleFileName = RulesCompiler.GetModuleFilename(ModuleName);
-				if(String.IsNullOrEmpty(ModuleFileName))
-				{
-					bAllowCompilation = bHasModuleRules = false;
-				}
-				else if(!Directory.EnumerateFiles(Path.GetDirectoryName(ModuleFileName), "*.cpp", SearchOption.AllDirectories).Any())
-				{
-					bAllowCompilation = false;
-				}
+				OutputDirectory = Path.Combine(OutputDirectory, ModuleRules.BinariesSubFolder);
+				IntermediateDirectory = Path.Combine(IntermediateDirectory, ModuleRules.BinariesSubFolder);
 			}
 
+			// Get the output filenames
+			string BaseBinaryPath = Path.Combine(OutputDirectory, MakeBinaryFileName(AppName + "-" + ModuleName, Platform, ModuleConfiguration, Rules.UndecoratedConfiguration, BinaryType));
+			string[] OutputFilePaths = UEBuildPlatform.GetBuildPlatform(Platform).FinalizeBinaryPaths(BaseBinaryPath);
+
 			// Prepare the configuration object
 			UEBuildBinaryConfiguration Config = new UEBuildBinaryConfiguration(BinaryType);
 			Config.OutputFilePaths = OutputFilePaths;
 			Config.IntermediateDirectory = IntermediateDirectory;
-			Config.bHasModuleRules = bHasModuleRules;
+			Config.bHasModuleRules = (ModuleRules != null);
 			Config.bAllowExports = (BinaryType == UEBuildBinaryType.DynamicLinkLibrary);
 			Config.bAllowCompilation = bAllowCompilation;
 			Config.bIsCrossTarget = bIsCrossTarget;
@@ -2406,70 +2387,6 @@ namespace UnrealBuildTool
 			return Result.ToString();
 		}
 
-		/** Given a UBT-built binary name (e.g. "Core"), returns a relative path to the binary for the current build configuration (e.g. "../Binaries/Win64/Core-Win64-Debug.lib") */
-		private string[] MakeBinaryPaths(string ModuleName, string BinaryName, UEBuildBinaryType BinaryType, PluginInfo PluginInfo)
-		{
-			// Determine the binary extension for the platform and binary type.
-			var BuildPlatform = UEBuildPlatform.GetBuildPlatform(Platform);
-
-
-			UnrealTargetConfiguration LocalConfig = Configuration;
-			if(Configuration == UnrealTargetConfiguration.DebugGame && !RulesCompiler.IsGameModule(ModuleName))
-			{
-				LocalConfig = UnrealTargetConfiguration.Development;
-			}
-			
-			string ModuleBinariesSubDir = "";
-			if (BinaryType == UEBuildBinaryType.DynamicLinkLibrary || BinaryType == UEBuildBinaryType.StaticLibrary)
-			{
-				// Allow for modules to specify sub-folders in the Binaries folder
-				var RulesFilename = RulesCompiler.GetModuleFilename(ModuleName);
-				// Plugins can be binary-only and can have no rules object
-				if (PluginInfo == null || !String.IsNullOrEmpty(RulesFilename))
-				{
-					ModuleRules ModuleRulesObj = RulesCompiler.CreateModuleRules(ModuleName, TargetInfo, out RulesFilename);
-					if (ModuleRulesObj != null)
-					{
-						ModuleBinariesSubDir = ModuleRulesObj.BinariesSubFolder;
-					}
-				}
-			}
-
-			[REDACTED_EMAIL]: This just happens to work since exp and lib files go into intermediate...
-
-			// Build base directory string ("../Binaries/<Platform>/")
-			string BinariesDirName;
-			if(TargetType == TargetRules.TargetType.Program && UnrealBuildTool.HasUProjectFile() && !ProjectFileGenerator.bGenerateProjectFiles)
-			{
-				BinariesDirName = Path.Combine( UnrealBuildTool.GetUProjectPath(), "Binaries" );
-			}
-			else if( PluginInfo != null )
-			{
-				BinariesDirName = Path.Combine( PluginInfo.Directory, "Binaries" );
-			}
-			else
-			{
-				string ModuleFileName = RulesCompiler.GetModuleFilename(ModuleName);
-				if(!String.IsNullOrEmpty(ModuleFileName) && UnrealBuildTool.HasUProjectFile() && Utils.IsFileUnderDirectory(ModuleFileName, UnrealBuildTool.GetUProjectPath()))
-				{
-					BinariesDirName = Path.Combine( UnrealBuildTool.GetUProjectPath(), "Binaries" );
-				}
-				else
-				{
-					BinariesDirName = Path.Combine( "..", "Binaries" );
-				}
-			}
-			var BaseDirectory = Path.Combine( BinariesDirName, Platform.ToString());
-			if (ModuleBinariesSubDir.Length > 0)
-			{
-				BaseDirectory = Path.Combine(BaseDirectory, ModuleBinariesSubDir);
-			}
-
-			// append the binary name
-			string OutBinaryPath = Path.Combine(BaseDirectory, MakeBinaryFileName(BinaryName, Platform, LocalConfig, Rules.UndecoratedConfiguration, BinaryType));
-			return BuildPlatform.FinalizeBinaryPaths(OutBinaryPath);
-		}
-
 		/// <summary>
 		/// Determine the output path for a target's executable
 		/// </summary>
@@ -2681,13 +2598,7 @@ namespace UnrealBuildTool
 				LinkEnvironmentConfiguration RulesLinkEnvConfig = GlobalLinkEnvironment.Config;
 				CPPEnvironmentConfiguration RulesCPPEnvConfig = GlobalCompileEnvironment.Config;
 
-				// Don't let games override the global environment...
-				// This will potentially cause problems due to them generating a game-agnostic exe.
-				// RocketEditor is a special case.
-				// 
-				bool bAllowGameOverride = !TargetRules.IsGameType(TargetType);
-				if (bAllowGameOverride ||
-					Rules.ShouldCompileMonolithic(Platform, Configuration))
+				if(!bUseSharedBuildEnvironment)
 				{
 					Rules.SetupGlobalEnvironment(
 						TargetInfo,
diff --git a/Engine/Source/Programs/UnrealBuildTool/IOS/UEDeployIOS.cs b/Engine/Source/Programs/UnrealBuildTool/IOS/UEDeployIOS.cs
--- a/Engine/Source/Programs/UnrealBuildTool/IOS/UEDeployIOS.cs
+++ b/Engine/Source/Programs/UnrealBuildTool/IOS/UEDeployIOS.cs
@@ -658,7 +658,7 @@ namespace UnrealBuildTool.IOS
 
 		public override bool PrepTargetForDeployment(UEBuildTarget InTarget)
 		{
-			string GameName = InTarget.AppName;
+			string GameName = InTarget.TargetName;
 			string BuildPath = (GameName == "UE4Game" ? "../../Engine" : InTarget.ProjectDirectory) + "/Binaries/IOS";
 			string ProjectDirectory = InTarget.ProjectDirectory;
 			bool bIsUE4Game = GameName.Contains("UE4Game");
diff --git a/Engine/Source/Programs/UnrealBuildTool/System/RulesCompiler.cs b/Engine/Source/Programs/UnrealBuildTool/System/RulesCompiler.cs
--- a/Engine/Source/Programs/UnrealBuildTool/System/RulesCompiler.cs
+++ b/Engine/Source/Programs/UnrealBuildTool/System/RulesCompiler.cs
@@ -663,11 +663,10 @@ namespace UnrealBuildTool
 		{
 		}
 
-
 		/// <summary>
 		/// Setup the global environment for building this target
-		/// IMPORTANT: Game targets will *not* have this function called unless they are built as monolithic targets.
-		/// This is due to non-monolithic games generating a shared executable.
+		/// IMPORTANT: Game targets will *not* have this function called if they use the shared build environment.
+		/// See ShouldUseSharedBuildEnvironment().
 		/// </summary>
 		/// <param name="Target">The target information - such as platform and configuration</param>
 		/// <param name="OutLinkEnvironmentConfiguration">Output link environment settings</param>
@@ -680,6 +679,18 @@ namespace UnrealBuildTool
 		{
 		}
 
+		/// <summary>
+		/// Allows a target to choose whether to use the shared build environment for a given configuration. Using
+		/// the shared build environment allows binaries to be reused between targets, but prevents customizing the
+		/// compile environment through SetupGlobalEnvironment().
+		/// </summary>
+		/// <param name="Target">Information about the target</param>
+		/// <returns>True if the target should use the shared build environment</returns>
+		public virtual bool ShouldUseSharedBuildEnvironment(TargetInfo Target)
+		{
+			return UnrealBuildTool.RunningRocket() || (Target.Type != TargetType.Program && !Target.IsMonolithic);
+		}
+
 		/// <summary>
 		/// Allows the target to specify modules which can be precompiled with the -Precompile/-UsePrecompiled arguments to UBT. 
 		/// All dependencies of the specified modules will be included.
@@ -1420,6 +1431,39 @@ namespace UnrealBuildTool
 			return (Utils.IsFileUnderDirectory( Filename, UnrealBuildTool.GetUProjectPath() ));
 		}
 
+		/// <summary>
+		/// Creates an instance of a module rules descriptor object for the specified module name
+		/// </summary>
+		/// <param name="ModuleName">Name of the module</param>
+		/// <param name="Target">Information about the target associated with this module</param>
+		/// <param name="Rules">Output </param>
+		/// <returns>Compiled module rule info</returns>
+		public static bool TryCreateModuleRules( string ModuleName, TargetInfo Target, out ModuleRules Rules )
+		{
+			if(GetModuleFilename( ModuleName ) == null)
+			{
+				Rules = null;
+				return false;
+			}
+			else
+			{
+				Rules = CreateModuleRules( ModuleName, Target );
+				return true;
+			}
+		}
+
+		/// <summary>
+		/// Creates an instance of a module rules descriptor object for the specified module name
+		/// </summary>
+		/// <param name="ModuleName">Name of the module</param>
+		/// <param name="Target">Information about the target associated with this module</param>
+		/// <returns>Compiled module rule info</returns>
+		public static ModuleRules CreateModuleRules( string ModuleName, TargetInfo Target )
+		{
+			string ModuleFileName;
+			return CreateModuleRules(ModuleName, Target, out ModuleFileName);
+		}
+
 		/// <summary>
 		/// Creates an instance of a module rules descriptor object for the specified module name
 		/// </summary>
@@ -1891,7 +1935,7 @@ namespace UnrealBuildTool
 					BuildTarget = new UEBuildServer(Desc, RulesObject);
 					break;
 				case TargetRules.TargetType.Program:
-					BuildTarget = new UEBuildTarget(Desc.TargetName, Desc, RulesObject);
+					BuildTarget = new UEBuildTarget(Desc, RulesObject, null);
 					break;
 			}
 
diff --git a/Engine/Source/UE4Client.Target.cs b/Engine/Source/UE4Client.Target.cs
--- a/Engine/Source/UE4Client.Target.cs
+++ b/Engine/Source/UE4Client.Target.cs
@@ -33,6 +33,11 @@ public class UE4ClientTarget : TargetRules
     {
     }
 
+	public override bool ShouldUseSharedBuildEnvironment(TargetInfo Target)
+	{
+		return true;
+	}
+
     public override List<UnrealTargetPlatform> GUBP_GetPlatforms_MonolithicOnly(UnrealTargetPlatform HostPlatform)
     {
 		List<UnrealTargetPlatform> Platforms = null;
diff --git a/Engine/Source/UE4Editor.Target.cs b/Engine/Source/UE4Editor.Target.cs
--- a/Engine/Source/UE4Editor.Target.cs
+++ b/Engine/Source/UE4Editor.Target.cs
@@ -27,6 +27,11 @@ public class UE4EditorTarget : TargetRules
 	{
 	}
 
+	public override bool ShouldUseSharedBuildEnvironment(TargetInfo Target)
+	{
+		return true;
+	}
+
 	public override void GetModulesToPrecompile(TargetInfo Target, List<string> ModuleNames)
 	{
 		ModuleNames.Add("Launch");
diff --git a/Engine/Source/UE4Game.Target.cs b/Engine/Source/UE4Game.Target.cs
--- a/Engine/Source/UE4Game.Target.cs
+++ b/Engine/Source/UE4Game.Target.cs
@@ -52,6 +52,11 @@ public class UE4GameTarget : TargetRules
 		}
 	}
 
+	public override bool ShouldUseSharedBuildEnvironment(TargetInfo Target)
+	{
+		return true;
+	}
+
 	public override void GetModulesToPrecompile(TargetInfo Target, List<string> ModuleNames)
 	{
 		// Add all the precompiled modules for this target
diff --git a/Engine/Source/UE4Server.Target.cs b/Engine/Source/UE4Server.Target.cs
--- a/Engine/Source/UE4Server.Target.cs
+++ b/Engine/Source/UE4Server.Target.cs
@@ -32,6 +32,11 @@ public class UE4ServerTarget : TargetRules
     {
     }
 
+	public override bool ShouldUseSharedBuildEnvironment(TargetInfo Target)
+	{
+		return true;
+	}
+
 	public override bool GetSupportedPlatforms(ref List<UnrealTargetPlatform> OutPlatforms)
 	{
 		// It is valid for only server platforms
