Set bPromoteEditorOnly to true by default. None of our games actually require this, and it gums up the build farm trying to make so much stuff.

diff --git a/Engine/Source/Programs/UnrealBuildTool/System/RulesCompiler.cs b/Engine/Source/Programs/UnrealBuildTool/System/RulesCompiler.cs
--- a/Engine/Source/Programs/UnrealBuildTool/System/RulesCompiler.cs
+++ b/Engine/Source/Programs/UnrealBuildTool/System/RulesCompiler.cs
@@ -868,8 +868,8 @@ namespace UnrealBuildTool
             public bool bIsMassive = false;
             public bool bCustomWorkflowForPromotion = false;
 			public bool bIsNonCode = false;
-            public bool bPromoteEditorOnly = false;
-        }
+            public bool bPromoteEditorOnly = true;
+		}
         public virtual GUBPProjectOptions GUBP_IncludeProjectInPromotedBuild_EditorTypeOnly(UnrealTargetPlatform HostPlatform)
         {
             var Result = new GUBPProjectOptions();
