Fix exception if plugin module does not have a rules file.

diff --git a/Engine/Source/Programs/UnrealBuildTool/Configuration/UEBuildTarget.cs b/Engine/Source/Programs/UnrealBuildTool/Configuration/UEBuildTarget.cs
--- a/Engine/Source/Programs/UnrealBuildTool/Configuration/UEBuildTarget.cs
+++ b/Engine/Source/Programs/UnrealBuildTool/Configuration/UEBuildTarget.cs
@@ -2139,7 +2139,7 @@ namespace UnrealBuildTool
 				if (ShouldIncludePluginModule(Plugin, Module))
 				{
 					string ModuleFileName = RulesCompiler.GetModuleFilename(Module.Name);
-					bool bHasSource = (ModuleFileName != null && Directory.EnumerateFiles(Path.GetDirectoryName(ModuleFileName), "*.cpp", SearchOption.AllDirectories).Any());
+					bool bHasSource = (!String.IsNullOrEmpty(ModuleFileName) && Directory.EnumerateFiles(Path.GetDirectoryName(ModuleFileName), "*.cpp", SearchOption.AllDirectories).Any());
 					AddBinaryForModule(Module.Name, BinaryType, bHasSource, false);
 				}
 			}
