Disabled server configs in Rocket Xcode projects

diff --git a/Engine/Source/Programs/UnrealBuildTool/System/XcodeProjectFileGenerator.cs b/Engine/Source/Programs/UnrealBuildTool/System/XcodeProjectFileGenerator.cs
--- a/Engine/Source/Programs/UnrealBuildTool/System/XcodeProjectFileGenerator.cs
+++ b/Engine/Source/Programs/UnrealBuildTool/System/XcodeProjectFileGenerator.cs
@@ -1100,6 +1100,11 @@ namespace UnrealBuildTool
 							WantProjectFileForTarget = true;
 						}
 					}
+
+					if (bGeneratingRocketProjectFiles && TargetName.EndsWith("Server"))
+					{
+						WantProjectFileForTarget = false;
+					}
 				}
 
 				if (WantProjectFileForTarget)
