Move GetGeneratedCodeDirectoryForModule into UEBuildTarget, so it can make use of shared logic with figuring out output paths.

diff --git a/Engine/Source/Programs/UnrealBuildTool/Configuration/UEBuildModule.cs b/Engine/Source/Programs/UnrealBuildTool/Configuration/UEBuildModule.cs
--- a/Engine/Source/Programs/UnrealBuildTool/Configuration/UEBuildModule.cs
+++ b/Engine/Source/Programs/UnrealBuildTool/Configuration/UEBuildModule.cs
@@ -2070,67 +2070,7 @@ namespace UnrealBuildTool
 				throw new ModuleProcessingException(this, ex);
 			}
 		}
-
-
-		/// <summary>
-		/// Determines where generated code files will be stored for this module
-		/// </summary>
-		/// <param name="ModuleDirectory">Module's base directory</param>
-		/// <param name="ModuleName">Name of module</param>
-		/// <returns></returns>
-		public static string GetGeneratedCodeDirectoryForModule(UEBuildTarget Target, string ModuleDirectory, string ModuleName)
-		{
-			string BaseDirectory = null;
-			if (Target.ShouldCompileMonolithic() || Target.TargetType == TargetRules.TargetType.Program || UnrealBuildTool.RunningRocket())
-			{
-				// Monolithic configurations, programs and Rocket games have their intermediate headers stored under their respective project folders.
-				string RootDirectory = UnrealBuildTool.GetUProjectPath();
-				if (String.IsNullOrEmpty(RootDirectory))
-				{
-					// Intermediates under Engine intermediate folder (program name will be appended later)
-					RootDirectory = Path.GetFullPath(BuildConfiguration.RelativeEnginePath);
-				}
-				BaseDirectory = Path.Combine(RootDirectory, BuildConfiguration.PlatformIntermediateFolder, Target.GetTargetName(), "Inc");
-			}
-			else if (Plugins.IsPluginModule(ModuleName))
-			{
-				// Plugin module
-				string PluginIntermediateIncPath;
-				{ 
-					PluginInfo Plugin = Plugins.GetPluginInfoForModule(ModuleName);
-					if (Plugin.LoadedFrom == PluginInfo.LoadedFromType.Engine)
-					{
-						// Plugin folder is in the engine directory
-						PluginIntermediateIncPath = Path.GetFullPath(Path.Combine(BuildConfiguration.RelativeEnginePath, BuildConfiguration.PlatformIntermediateFolder));
-					}
-					else
-					{
-						// Plugin folder is in the project directory
-						PluginIntermediateIncPath   = Path.GetFullPath(Path.Combine( Target.ProjectDirectory, BuildConfiguration.PlatformIntermediateFolder));
-					}
-					PluginIntermediateIncPath = Path.Combine(PluginIntermediateIncPath, "Inc", "Plugins");
-				}
-				BaseDirectory = PluginIntermediateIncPath;
-			}
-			else
-			{
-				var AllProjectFolders = UEBuildTarget.DiscoverAllGameFolders();		// @todo ubtmake: This will be called again and again for every UObject module (80+)
-				BaseDirectory = AllProjectFolders.Find(ProjectFolder => Utils.IsFileUnderDirectory( ModuleDirectory, ProjectFolder ));
-				if (BaseDirectory == null)
-				{
-					// Must be an engine module or program module
-					BaseDirectory = ProjectFileGenerator.EngineRelativePath;
-				}
-
-				BaseDirectory = Path.GetFullPath(Path.Combine(BaseDirectory, BuildConfiguration.PlatformIntermediateFolder, "Inc"));
-			}
-
-			// Construct the intermediate path.
-			var GeneratedCodeDirectory = Path.Combine(BaseDirectory, ModuleName);
-			return GeneratedCodeDirectory + Path.DirectorySeparatorChar;
-		}
-
-	};
+	}
 
 	/** A module that is compiled from C++ CLR code. */
 	class UEBuildModuleCPPCLR : UEBuildModuleCPP
diff --git a/Engine/Source/Programs/UnrealBuildTool/Configuration/UEBuildTarget.cs b/Engine/Source/Programs/UnrealBuildTool/Configuration/UEBuildTarget.cs
--- a/Engine/Source/Programs/UnrealBuildTool/Configuration/UEBuildTarget.cs
+++ b/Engine/Source/Programs/UnrealBuildTool/Configuration/UEBuildTarget.cs
@@ -1009,7 +1009,7 @@ namespace UnrealBuildTool
 				foreach (var ModuleName in ModuleList)
 				{
 					var Module = GetModuleByName(ModuleName);
-					var ModuleIncludeDir = UEBuildModuleCPP.GetGeneratedCodeDirectoryForModule(this, Module.ModuleDirectory, ModuleName).Replace("\\", "/");
+					var ModuleIncludeDir = GetGeneratedCodeDirectoryForModule(Module.ModuleDirectory, ModuleName).Replace("\\", "/");
 					if (!UnrealBuildTool.RunningRocket() || !ModuleIncludeDir.StartsWith(BaseEngineBuildDataFolder, StringComparison.InvariantCultureIgnoreCase))
 					{
 						if (Directory.Exists(ModuleIncludeDir))
@@ -1457,7 +1457,7 @@ namespace UnrealBuildTool
 							// If we've got this far and there are no source files then it's likely we're running Rocket and ignoring
 							// engine files, so we don't need a .generated.cpp either
 							UEBuildModuleCPP.AutoGenerateCppInfoClass.BuildInfoClass BuildInfo = null;
-							UHTModuleInfo.GeneratedCPPFilenameBase = Path.Combine( UEBuildModuleCPP.GetGeneratedCodeDirectoryForModule(this, UHTModuleInfo.ModuleDirectory, UHTModuleInfo.ModuleName), UHTModuleInfo.ModuleName ) + ".generated";
+							UHTModuleInfo.GeneratedCPPFilenameBase = Path.Combine( GetGeneratedCodeDirectoryForModule(UHTModuleInfo.ModuleDirectory, UHTModuleInfo.ModuleName), UHTModuleInfo.ModuleName ) + ".generated";
 							if (DependencyModuleCPP.SourceFilesToBuild.Count != 0)
 							{
 								BuildInfo = new UEBuildModuleCPP.AutoGenerateCppInfoClass.BuildInfoClass( UHTModuleInfo.GeneratedCPPFilenameBase + "*.cpp" );
@@ -2322,6 +2322,64 @@ namespace UnrealBuildTool
 			return Binary;
 		}
 
+		/// <summary>
+		/// Determines where generated code files will be stored for this module
+		/// </summary>
+		/// <param name="ModuleDirectory">Module's base directory</param>
+		/// <param name="ModuleName">Name of module</param>
+		/// <returns></returns>
+		string GetGeneratedCodeDirectoryForModule(string ModuleDirectory, string ModuleName)
+		{
+			string BaseDirectory = null;
+			if (ShouldCompileMonolithic() || TargetType == TargetRules.TargetType.Program || UnrealBuildTool.RunningRocket())
+			{
+				// Monolithic configurations, programs and Rocket games have their intermediate headers stored under their respective project folders.
+				string RootDirectory = UnrealBuildTool.GetUProjectPath();
+				if (String.IsNullOrEmpty(RootDirectory))
+				{
+					// Intermediates under Engine intermediate folder (program name will be appended later)
+					RootDirectory = Path.GetFullPath(BuildConfiguration.RelativeEnginePath);
+				}
+				BaseDirectory = Path.Combine(RootDirectory, BuildConfiguration.PlatformIntermediateFolder, GetTargetName(), "Inc");
+			}
+			else if (Plugins.IsPluginModule(ModuleName))
+			{
+				// Plugin module
+				string PluginIntermediateIncPath;
+				{ 
+					PluginInfo Plugin = Plugins.GetPluginInfoForModule(ModuleName);
+					if (Plugin.LoadedFrom == PluginInfo.LoadedFromType.Engine)
+					{
+						// Plugin folder is in the engine directory
+						PluginIntermediateIncPath = Path.GetFullPath(Path.Combine(BuildConfiguration.RelativeEnginePath, BuildConfiguration.PlatformIntermediateFolder));
+					}
+					else
+					{
+						// Plugin folder is in the project directory
+						PluginIntermediateIncPath   = Path.GetFullPath(Path.Combine( ProjectDirectory, BuildConfiguration.PlatformIntermediateFolder));
+					}
+					PluginIntermediateIncPath = Path.Combine(PluginIntermediateIncPath, "Inc", "Plugins");
+				}
+				BaseDirectory = PluginIntermediateIncPath;
+			}
+			else
+			{
+				var AllProjectFolders = UEBuildTarget.DiscoverAllGameFolders();		// @todo ubtmake: This will be called again and again for every UObject module (80+)
+				BaseDirectory = AllProjectFolders.Find(ProjectFolder => Utils.IsFileUnderDirectory( ModuleDirectory, ProjectFolder ));
+				if (BaseDirectory == null)
+				{
+					// Must be an engine module or program module
+					BaseDirectory = ProjectFileGenerator.EngineRelativePath;
+				}
+
+				BaseDirectory = Path.GetFullPath(Path.Combine(BaseDirectory, BuildConfiguration.PlatformIntermediateFolder, "Inc"));
+			}
+
+			// Construct the intermediate path.
+			var GeneratedCodeDirectory = Path.Combine(BaseDirectory, ModuleName);
+			return GeneratedCodeDirectory + Path.DirectorySeparatorChar;
+		}
+
 		/**
 		 * @return true if debug information is created, false otherwise.
 		 */
@@ -2890,11 +2948,16 @@ namespace UnrealBuildTool
 					}
 				}
 
-				// Don't generate include paths for third party modules.  They don't follow our conventions.
-				if (RulesObject.Type != ModuleRules.ModuleType.External)
+				// Don't generate include paths for third party modules; they don't follow our conventions. Core is a special-case... leave it alone
+				if (RulesObject.Type != ModuleRules.ModuleType.External && ModuleName != "Core")
 				{
 					// Add the default include paths to the module rules, if they exist.
-					RulesCompiler.AddDefaultIncludePathsToModuleRules(this, ModuleName, ModuleFileName, ModuleFileRelativeToEngineDirectory, IsGameModule: IsGameModule, RulesObject: ref RulesObject);
+					RulesCompiler.AddDefaultIncludePathsToModuleRules(ModuleName, ModuleFileName, ModuleFileRelativeToEngineDirectory, IsGameModule, RulesObject);
+
+					// Add the path to the generated headers 
+					string ModuleGeneratedCodePath = GetGeneratedCodeDirectoryForModule(Path.GetDirectoryName(ModuleFileName), ModuleName);
+					ModuleGeneratedCodePath = Utils.CleanDirectorySeparators(Utils.MakePathRelativeTo(ModuleGeneratedCodePath, Path.Combine(ProjectFileGenerator.RootRelativePath, "Engine/Source")), '/');
+					RulesObject.PublicIncludePaths.Add(ModuleGeneratedCodePath);
 				}
 
 				// Figure out whether we need to build this module
diff --git a/Engine/Source/Programs/UnrealBuildTool/System/RulesCompiler.cs b/Engine/Source/Programs/UnrealBuildTool/System/RulesCompiler.cs
--- a/Engine/Source/Programs/UnrealBuildTool/System/RulesCompiler.cs
+++ b/Engine/Source/Programs/UnrealBuildTool/System/RulesCompiler.cs
@@ -1600,14 +1600,8 @@ namespace UnrealBuildTool
 		/// <param name="InModuleFileRelativeToEngineDirectory">The module file relative to the engine directory</param>
 		/// <param name="IsGameModule">true if it is a game module, false if not</param>
 		/// <param name="RulesObject">The module rules object itself</param>
-		public static void AddDefaultIncludePathsToModuleRules(UEBuildTarget BuildTarget, string InModuleName, string InModuleFilename, string InModuleFileRelativeToEngineDirectory, bool IsGameModule, ref ModuleRules RulesObject)
+		public static void AddDefaultIncludePathsToModuleRules(string InModuleName, string InModuleFilename, string InModuleFileRelativeToEngineDirectory, bool IsGameModule, ModuleRules RulesObject)
 		{
-			// Core is a special-case... leave it alone
-			if (InModuleName == "Core")
-			{
-				return;
-			}
-
 			// Grab the absolute path of the Engine/Source folder for use later
 			string AbsEngineSourceDirectory = Path.Combine(ProjectFileGenerator.RootRelativePath, "Engine/Source");
 			AbsEngineSourceDirectory = Path.GetFullPath(AbsEngineSourceDirectory);
@@ -1649,8 +1643,6 @@ namespace UnrealBuildTool
 			// Setup the directories for Classes, Public, and Intermediate			
 			string ClassesDirectory = Path.Combine(ModuleDirectoryRelativeToEngineSourceDirectory, "Classes/");	// @todo uht: Deprecate eventually.  Or force it to be manually specified...
 			string PublicDirectory = Path.Combine(ModuleDirectoryRelativeToEngineSourceDirectory, "Public/");
-			string IntermediateDirectory = UEBuildModuleCPP.GetGeneratedCodeDirectoryForModule(BuildTarget, Path.GetDirectoryName(InModuleFilename), InModuleName);
-			IntermediateDirectory = Utils.CleanDirectorySeparators(Utils.MakePathRelativeTo(IntermediateDirectory, Path.Combine(ProjectFileGenerator.RootRelativePath, "Engine/Source")), '/');
 
 			if (IsGameModule || IsPluginModule)
 			{
@@ -1704,7 +1696,6 @@ namespace UnrealBuildTool
 			}
 
 			string IncludePath_Classes = "";
-			string IncludePath_Intermediates = IntermediateDirectory;
 			string IncludePath_Public = "";
 
 			bool bModulePathIsRooted = Path.IsPathRooted(ModuleDirectoryRelativeToEngineSourceDirectory);
@@ -1726,10 +1717,6 @@ namespace UnrealBuildTool
 			{
 				RulesObject.PublicIncludePaths.Add(IncludePath_Classes);
 			}
-			if (IncludePath_Intermediates.Length > 0)
-			{
-				RulesObject.PublicIncludePaths.Add(IncludePath_Intermediates);
-			}
 			if (IncludePath_Public.Length > 0)
 			{
 				RulesObject.PublicIncludePaths.Add(IncludePath_Public);
