Use the shared build environment settings to pick an output directory for generated headers, and put generated headers for plugins in the plugin directory. Doing so allows Rocket to install generated headers again (which should fix some intellisense issues), and allows read-only plugins to be installed into the engine directory.

diff --git a/Engine/Source/Programs/UnrealBuildTool/Configuration/UEBuildTarget.cs b/Engine/Source/Programs/UnrealBuildTool/Configuration/UEBuildTarget.cs
--- a/Engine/Source/Programs/UnrealBuildTool/Configuration/UEBuildTarget.cs
+++ b/Engine/Source/Programs/UnrealBuildTool/Configuration/UEBuildTarget.cs
@@ -1009,7 +1009,7 @@ namespace UnrealBuildTool
 				foreach (var ModuleName in ModuleList)
 				{
 					var Module = GetModuleByName(ModuleName);
-					var ModuleIncludeDir = GetGeneratedCodeDirectoryForModule(Module.ModuleDirectory, ModuleName).Replace("\\", "/");
+					var ModuleIncludeDir = GetGeneratedCodeDirectoryForModule(ModuleName).Replace("\\", "/");
 					if (!UnrealBuildTool.RunningRocket() || !ModuleIncludeDir.StartsWith(BaseEngineBuildDataFolder, StringComparison.InvariantCultureIgnoreCase))
 					{
 						if (Directory.Exists(ModuleIncludeDir))
@@ -1457,7 +1457,7 @@ namespace UnrealBuildTool
 							// If we've got this far and there are no source files then it's likely we're running Rocket and ignoring
 							// engine files, so we don't need a .generated.cpp either
 							UEBuildModuleCPP.AutoGenerateCppInfoClass.BuildInfoClass BuildInfo = null;
-							UHTModuleInfo.GeneratedCPPFilenameBase = Path.Combine( GetGeneratedCodeDirectoryForModule(UHTModuleInfo.ModuleDirectory, UHTModuleInfo.ModuleName), UHTModuleInfo.ModuleName ) + ".generated";
+							UHTModuleInfo.GeneratedCPPFilenameBase = Path.Combine( GetGeneratedCodeDirectoryForModule(UHTModuleInfo.ModuleName), UHTModuleInfo.ModuleName ) + ".generated";
 							if (DependencyModuleCPP.SourceFilesToBuild.Count != 0)
 							{
 								BuildInfo = new UEBuildModuleCPP.AutoGenerateCppInfoClass.BuildInfoClass( UHTModuleInfo.GeneratedCPPFilenameBase + "*.cpp" );
@@ -2325,58 +2325,31 @@ namespace UnrealBuildTool
 		/// <summary>
 		/// Determines where generated code files will be stored for this module
 		/// </summary>
-		/// <param name="ModuleDirectory">Module's base directory</param>
-		/// <param name="ModuleName">Name of module</param>
+		/// <param name="ModuleName">Name of the module</param>
 		/// <returns></returns>
-		string GetGeneratedCodeDirectoryForModule(string ModuleDirectory, string ModuleName)
+		string GetGeneratedCodeDirectoryForModule(string ModuleName)
 		{
-			string BaseDirectory = null;
-			if (ShouldCompileMonolithic() || TargetType == TargetRules.TargetType.Program || UnrealBuildTool.RunningRocket())
+			// Get the plugin info for the module
+			PluginInfo Plugin = Plugins.GetPluginInfoForModule(ModuleName);
+
+			// Get the output directory. Plugins always write to their own directory so they can be copied between projects, shared engine intermediates go in the engine 
+			// intermediate folder, and anything else goes in the project folder.
+			string GeneratedCodeDirectory;
+			if(Plugin != null)
 			{
-				// Monolithic configurations, programs and Rocket games have their intermediate headers stored under their respective project folders.
-				string RootDirectory = UnrealBuildTool.GetUProjectPath();
-				if (String.IsNullOrEmpty(RootDirectory))
-				{
-					// Intermediates under Engine intermediate folder (program name will be appended later)
-					RootDirectory = Path.GetFullPath(BuildConfiguration.RelativeEnginePath);
-				}
-				BaseDirectory = Path.Combine(RootDirectory, BuildConfiguration.PlatformIntermediateFolder, GetTargetName(), "Inc");
+				GeneratedCodeDirectory = Plugin.Directory;
 			}
-			else if (Plugins.IsPluginModule(ModuleName))
+			else if(bUseSharedBuildEnvironment && Utils.IsFileUnderDirectory(RulesCompiler.GetModuleFilename(ModuleName), BuildConfiguration.RelativeEnginePath))
 			{
-				// Plugin module
-				string PluginIntermediateIncPath;
-				{ 
-					PluginInfo Plugin = Plugins.GetPluginInfoForModule(ModuleName);
-					if (Plugin.LoadedFrom == PluginInfo.LoadedFromType.Engine)
-					{
-						// Plugin folder is in the engine directory
-						PluginIntermediateIncPath = Path.GetFullPath(Path.Combine(BuildConfiguration.RelativeEnginePath, BuildConfiguration.PlatformIntermediateFolder));
-					}
-					else
-					{
-						// Plugin folder is in the project directory
-						PluginIntermediateIncPath   = Path.GetFullPath(Path.Combine( ProjectDirectory, BuildConfiguration.PlatformIntermediateFolder));
-					}
-					PluginIntermediateIncPath = Path.Combine(PluginIntermediateIncPath, "Inc", "Plugins");
-				}
-				BaseDirectory = PluginIntermediateIncPath;
+				GeneratedCodeDirectory = BuildConfiguration.RelativeEnginePath;
 			}
 			else
 			{
-				var AllProjectFolders = UEBuildTarget.DiscoverAllGameFolders();		// @todo ubtmake: This will be called again and again for every UObject module (80+)
-				BaseDirectory = AllProjectFolders.Find(ProjectFolder => Utils.IsFileUnderDirectory( ModuleDirectory, ProjectFolder ));
-				if (BaseDirectory == null)
-				{
-					// Must be an engine module or program module
-					BaseDirectory = ProjectFileGenerator.EngineRelativePath;
-				}
-
-				BaseDirectory = Path.GetFullPath(Path.Combine(BaseDirectory, BuildConfiguration.PlatformIntermediateFolder, "Inc"));
+				GeneratedCodeDirectory = ProjectDirectory;
 			}
+			GeneratedCodeDirectory = Path.GetFullPath(Path.Combine(GeneratedCodeDirectory, BuildConfiguration.PlatformIntermediateFolder, AppName, "Inc", ModuleName));
 
-			// Construct the intermediate path.
-			var GeneratedCodeDirectory = Path.Combine(BaseDirectory, ModuleName);
+			// Return the path with a directory separator
 			return GeneratedCodeDirectory + Path.DirectorySeparatorChar;
 		}
 
@@ -2955,7 +2928,7 @@ namespace UnrealBuildTool
 					RulesCompiler.AddDefaultIncludePathsToModuleRules(ModuleName, ModuleFileName, ModuleFileRelativeToEngineDirectory, IsGameModule, RulesObject);
 
 					// Add the path to the generated headers 
-					string ModuleGeneratedCodePath = GetGeneratedCodeDirectoryForModule(Path.GetDirectoryName(ModuleFileName), ModuleName);
+					string ModuleGeneratedCodePath = GetGeneratedCodeDirectoryForModule(ModuleName);
 					ModuleGeneratedCodePath = Utils.CleanDirectorySeparators(Utils.MakePathRelativeTo(ModuleGeneratedCodePath, Path.Combine(ProjectFileGenerator.RootRelativePath, "Engine/Source")), '/');
 					RulesObject.PublicIncludePaths.Add(ModuleGeneratedCodePath);
 				}
diff --git a/Engine/Source/Programs/UnrealBuildTool/System/ExternalExecution.cs b/Engine/Source/Programs/UnrealBuildTool/System/ExternalExecution.cs
--- a/Engine/Source/Programs/UnrealBuildTool/System/ExternalExecution.cs
+++ b/Engine/Source/Programs/UnrealBuildTool/System/ExternalExecution.cs
@@ -134,7 +134,7 @@ namespace UnrealBuildTool
 				PrivateHeaders           = Info.PrivateUObjectHeaders      .Select((Header) => Header.AbsolutePath).ToList(),
 				PCH                      = Info.PCH,
 				GeneratedCPPFilenameBase = Info.GeneratedCPPFilenameBase,
-				SaveExportedHeaders      = true, // Rocket just regenerates all headers to the project directory now
+				SaveExportedHeaders      = !UnrealBuildTool.IsEngineInstalled() || !Utils.IsFileUnderDirectory(Info.ModuleDirectory, BuildConfiguration.RelativeEnginePath),
 				UHTGeneratedCodeVersion = Info.GeneratedCodeVersion,
 			}).ToList();
 		}
