Implement PluginDescriptor and ModuleDescriptor in UBT in an identical way to the engine. Also add a wrapper for objects parsed by fastJSON, that allows easier retrieval of typed fields.

diff --git a/Engine/Source/Programs/AutomationTool/ProjectUtils.cs b/Engine/Source/Programs/AutomationTool/ProjectUtils.cs
--- a/Engine/Source/Programs/AutomationTool/ProjectUtils.cs
+++ b/Engine/Source/Programs/AutomationTool/ProjectUtils.cs
@@ -213,9 +213,9 @@ namespace AutomationTool
 						{
 							if (Plugin.Name == PluginName && !Plugin.bEnabledByDefault && !RetVal)
 							{
-								foreach (var Module in Plugin.Modules)
+								foreach (var Module in Plugin.Descriptor.Modules)
 								{
-									if (Module.Platforms.Count > 0 && Module.Platforms.Contains(TargetPlatformType) && !RetVal)
+									if(Module.IsCompiledInConfiguration(TargetPlatformType, TargetRules.TargetType.Game))
 									{
 										RetVal = true;
 										break;
diff --git a/Engine/Source/Programs/UnrealBuildTool/Configuration/UEBuildModule.cs b/Engine/Source/Programs/UnrealBuildTool/Configuration/UEBuildModule.cs
--- a/Engine/Source/Programs/UnrealBuildTool/Configuration/UEBuildModule.cs
+++ b/Engine/Source/Programs/UnrealBuildTool/Configuration/UEBuildModule.cs
@@ -66,24 +66,24 @@ namespace UnrealBuildTool
 			{
 				return ModuleType;
 			}
-			foreach (PluginInfo.PluginModuleInfo ModuleInfo in Info.Modules)
+			foreach (ModuleDescriptor ModuleInfo in Info.Descriptor.Modules)
 			{
 				if (ModuleInfo.Name == ModuleName)
 				{
 					switch (ModuleInfo.Type)
 					{
-						case PluginInfo.PluginModuleType.Developer:
+						case ModuleHostType.Developer:
 							ModuleType = UEBuildModuleType.Developer;
 							break;
-						case PluginInfo.PluginModuleType.Editor:
-						case PluginInfo.PluginModuleType.EditorNoCommandlet:
+						case ModuleHostType.Editor:
+						case ModuleHostType.EditorNoCommandlet:
 							ModuleType = UEBuildModuleType.Editor;
 							break;
-						case PluginInfo.PluginModuleType.Program:
+						case ModuleHostType.Program:
 							ModuleType = UEBuildModuleType.Program;
 							break;
-						case PluginInfo.PluginModuleType.Runtime:
-						case PluginInfo.PluginModuleType.RuntimeNoCommandlet:
+						case ModuleHostType.Runtime:
+						case ModuleHostType.RuntimeNoCommandlet:
 							ModuleType = UEBuildModuleType.Runtime;
 							break;
 						default:
diff --git a/Engine/Source/Programs/UnrealBuildTool/Configuration/UEBuildTarget.cs b/Engine/Source/Programs/UnrealBuildTool/Configuration/UEBuildTarget.cs
--- a/Engine/Source/Programs/UnrealBuildTool/Configuration/UEBuildTarget.cs
+++ b/Engine/Source/Programs/UnrealBuildTool/Configuration/UEBuildTarget.cs
@@ -1852,9 +1852,9 @@ namespace UnrealBuildTool
 			}
 			foreach (PluginInfo Plugin in DependentPlugins)
 			{
-				foreach (PluginInfo.PluginModuleInfo Module in Plugin.Modules)
+				foreach (ModuleDescriptor Module in Plugin.Descriptor.Modules)
 				{
-					if(ShouldIncludePluginModule(Plugin, Module))
+					if(Module.IsCompiledInConfiguration(Platform, TargetType))
 					{
 						PrivateDependencyModuleNames.Add(Module.Name);
 					}
@@ -2134,9 +2134,9 @@ namespace UnrealBuildTool
         public void AddPlugin(PluginInfo Plugin)
 		{
 			UEBuildBinaryType BinaryType = ShouldCompileMonolithic() ? UEBuildBinaryType.StaticLibrary : UEBuildBinaryType.DynamicLinkLibrary;
-			foreach(PluginInfo.PluginModuleInfo Module in Plugin.Modules)
+			foreach(ModuleDescriptor Module in Plugin.Descriptor.Modules)
 			{
-				if (ShouldIncludePluginModule(Plugin, Module))
+				if (Module.IsCompiledInConfiguration(Platform, TargetType))
 				{
 					string ModuleFileName = RulesCompiler.GetModuleFilename(Module.Name);
 					bool bHasSource = (!String.IsNullOrEmpty(ModuleFileName) && Directory.EnumerateFiles(Path.GetDirectoryName(ModuleFileName), "*.cpp", SearchOption.AllDirectories).Any());
@@ -2171,9 +2171,9 @@ namespace UnrealBuildTool
 				// Add all the enabled plugins to the precompiled module list. Plugins are always precompiled, even if bPrecompile is not set, so we should precompile their dependencies.
 				foreach(PluginInfo Plugin in BuildPlugins)
 				{
-					foreach(PluginInfo.PluginModuleInfo Module in Plugin.Modules)
+					foreach(ModuleDescriptor Module in Plugin.Descriptor.Modules)
 					{
-						if (ShouldIncludePluginModule(Plugin, Module))
+						if (Module.IsCompiledInConfiguration(Platform, TargetType))
 						{
 							PrecompiledModuleNames.Add(Module.Name);
 						}
@@ -2450,35 +2450,6 @@ namespace UnrealBuildTool
 			return UEBuildPlatform.GetBuildPlatform(Platform).FinalizeBinaryPaths(BinaryPath);
 		}
 
-		/// <summary>
-		/// Determines whether the given plugin module is part of the current build.
-		/// </summary>
-		private bool ShouldIncludePluginModule(PluginInfo Plugin, PluginInfo.PluginModuleInfo Module)
-		{
-			// Check it can be built for this platform...
-			if (Module.Platforms.Contains(Platform))
-			{
-				// ...and that it's appropriate for the current build environment.
-				switch (Module.Type)
-				{
-					case PluginInfo.PluginModuleType.Runtime:
-					case PluginInfo.PluginModuleType.RuntimeNoCommandlet:
-						return true;
-
-					case PluginInfo.PluginModuleType.Developer:
-						return UEBuildConfiguration.bBuildDeveloperTools;
-
-					case PluginInfo.PluginModuleType.Editor:
-					case PluginInfo.PluginModuleType.EditorNoCommandlet:
-						return UEBuildConfiguration.bBuildEditor;
-
-					case PluginInfo.PluginModuleType.Program:
-						return TargetType == TargetRules.TargetType.Program;
-				}
-			}
-			return false;
-		}
-
 		/** Sets up the modules for the target. */
 		protected virtual void SetupModules()
 		{
diff --git a/Engine/Source/Programs/UnrealBuildTool/System/ModuleDescriptor.cs b/Engine/Source/Programs/UnrealBuildTool/System/ModuleDescriptor.cs
new file mode 100644
--- /dev/null
+++ b/Engine/Source/Programs/UnrealBuildTool/System/ModuleDescriptor.cs
@@ -0,0 +1,136 @@
+// Copyright 1998-2015 Epic Games, Inc. All Rights Reserved.
+
+using System;
+using System.Collections.Generic;
+using System.IO;
+using System.Diagnostics;
+using System.Linq;
+
+namespace UnrealBuildTool
+{
+	public enum ModuleHostType
+	{
+		Runtime,
+		RuntimeNoCommandlet,
+		Developer,
+		Editor,
+		EditorNoCommandlet,
+		Program,
+	}
+
+	public enum ModuleLoadingPhase
+	{
+		Default,
+		PostDefault,
+		PreDefault,
+		PostConfigInit,
+		PreLoadingScreen,
+		PostEngineInit,
+	}
+
+	[DebuggerDisplay("Name={Name}")]
+	public class ModuleDescriptor
+	{
+		// Name of this module
+		public readonly string Name;
+
+		// Usage type of module
+		public ModuleHostType Type;
+
+		// When should the module be loaded during the startup sequence?  This is sort of an advanced setting.
+		public ModuleLoadingPhase LoadingPhase = ModuleLoadingPhase.Default;
+
+		// List of allowed platforms
+		public UnrealTargetPlatform[] WhitelistPlatforms;
+
+		// List of disallowed platforms
+		public UnrealTargetPlatform[] BlacklistPlatforms;
+
+		// List of additional dependencies for building this module.
+		public string[] AdditionalDependencies;
+
+		/// <summary>
+		/// Constructor
+		/// </summary>
+		/// <param name="InName">Name of the module</param>
+		/// <param name="InType">Type of target that can host this module</param>
+		public ModuleDescriptor(string InName, ModuleHostType InType)
+		{
+			Name = InName;
+			Type = InType;
+		}
+
+		/// <summary>
+		/// Constructs a ModuleDescriptor from a Json object
+		/// </summary>
+		/// <param name="InObject"></param>
+		/// <returns>The new module descriptor</returns>
+		public static ModuleDescriptor FromJsonObject(JsonObject InObject)
+		{
+			ModuleDescriptor Module = new ModuleDescriptor(InObject.GetStringField("Name"), InObject.GetEnumField<ModuleHostType>("Type"));
+
+			ModuleLoadingPhase LoadingPhase;
+			if(InObject.TryGetEnumField<ModuleLoadingPhase>("LoadingPhase", out LoadingPhase))
+			{
+				Module.LoadingPhase = LoadingPhase;
+			}
+
+			UnrealTargetPlatform[] WhitelistPlatforms;
+			if(InObject.TryGetEnumArrayField<UnrealTargetPlatform>("WhitelistPlatforms", out WhitelistPlatforms))
+			{
+				Module.WhitelistPlatforms = WhitelistPlatforms;
+			}
+
+			UnrealTargetPlatform[] BlacklistPlatforms;
+			if(InObject.TryGetEnumArrayField<UnrealTargetPlatform>("BlacklistPlatforms", out BlacklistPlatforms))
+			{
+				Module.BlacklistPlatforms = BlacklistPlatforms;
+			}
+
+			string[] AdditionalDependencies;
+			if(InObject.TryGetStringArrayField("AdditionalDependencies", out AdditionalDependencies))
+			{
+				Module.AdditionalDependencies = AdditionalDependencies;
+			}
+
+			return Module;
+		}
+
+		/// <summary>
+		/// Determines whether the given plugin module is part of the current build.
+		/// </summary>
+		/// <param name="Platform">The platform being compiled for</param>
+		/// <param name="TargetType">The type of the target being compiled</param>
+		public bool IsCompiledInConfiguration(UnrealTargetPlatform Platform, TargetRules.TargetType TargetType)
+		{
+			// Check the platform is whitelisted
+			if(WhitelistPlatforms != null && WhitelistPlatforms.Length > 0 && !WhitelistPlatforms.Contains(Platform))
+			{
+				return false;
+			}
+
+			// Check the platform is not blacklisted
+			if(BlacklistPlatforms != null && BlacklistPlatforms.Contains(Platform))
+			{
+				return false;
+			}
+
+			// Check the module is compatible with this target.
+			switch (Type)
+			{
+				case ModuleHostType.Runtime:
+				case ModuleHostType.RuntimeNoCommandlet:
+					return true;
+				case ModuleHostType.Developer:
+					return UEBuildConfiguration.bBuildDeveloperTools;
+				case ModuleHostType.Editor:
+				case ModuleHostType.EditorNoCommandlet:
+					return UEBuildConfiguration.bBuildEditor;
+				case ModuleHostType.Program:
+					return TargetType == TargetRules.TargetType.Program;
+			}
+
+			return false;
+		}
+	}
+}
diff --git a/Engine/Source/Programs/UnrealBuildTool/System/PluginDescriptor.cs b/Engine/Source/Programs/UnrealBuildTool/System/PluginDescriptor.cs
new file mode 100644
--- /dev/null
+++ b/Engine/Source/Programs/UnrealBuildTool/System/PluginDescriptor.cs
@@ -0,0 +1,135 @@
+﻿using System;
+using System.Collections.Generic;
+using System.Linq;
+using System.Text;
+
+namespace UnrealBuildTool
+{
+	public enum PluginDescriptorVersion
+	{
+		Invalid = 0,
+		Initial = 1,
+		NameHash = 2,
+		ProjectPluginUnification = 3,
+		// !!!!!!!!!! IMPORTANT: Remember to also update LatestPluginDescriptorFileVersion in Plugins.cs (and Plugin system documentation) when this changes!!!!!!!!!!!
+		// -----<new versions can be added before this line>-------------------------------------------------
+		// - this needs to be the last line (see note below)
+		LatestPlusOne,
+		Latest = LatestPlusOne - 1
+	};
+
+	public class PluginDescriptor
+	{
+		// Descriptor version number
+		public int FileVersion;
+
+		// Version number for the plugin.  The version number must increase with every version of the plugin, so that the system 
+		// can determine whether one version of a plugin is newer than another, or to enforce other requirements.  This version
+		// number is not displayed in front-facing UI.  Use the VersionName for that.
+		public int Version;
+
+		// Name of the version for this plugin.  This is the front-facing part of the version number.  It doesn't need to match
+		// the version number numerically, but should be updated when the version number is increased accordingly.
+		public string VersionName;
+		 
+		// Friendly name of the plugin
+		public string FriendlyName;
+
+		// Description of the plugin
+		public string Description;
+
+		// The name of the category this plugin
+		public string Category;
+
+		// The company or individual who created this plugin.  This is an optional field that may be displayed in the user interface.
+		public string CreatedBy;
+
+		// Hyperlink URL string for the company or individual who created this plugin.  This is optional.
+		public string CreatedByURL;
+
+		// Documentation URL string.
+		public string DocsURL;
+
+		// List of all modules associated with this plugin
+		public ModuleDescriptor[] Modules;
+
+		// Whether this plugin should be enabled by default for all projects
+		public bool bEnabledByDefault;
+
+		// Can this plugin contain content?
+		public bool bCanContainContent;
+
+		// Marks the plugin as beta in the UI
+		public bool bIsBetaVersion;
+
+		/// <summary>
+		/// Private constructor. This object should not be created directly; read it from disk using FromFile() instead.
+		/// </summary>
+		private PluginDescriptor()
+		{
+			FileVersion = (int)PluginDescriptorVersion.Latest;
+		}
+
+		/// <summary>
+		/// Creates a plugin descriptor from a file on disk
+		/// </summary>
+		/// <param name="FileName">The filename to read</param>
+		/// <returns>New plugin descriptor</returns>
+		public static PluginDescriptor FromFile(string FileName)
+		{
+			JsonObject RawObject = JsonObject.FromFile(FileName);
+
+			PluginDescriptor Descriptor = new PluginDescriptor();
+
+			// Read the version
+			if(!RawObject.TryGetNumericField("FileVersion", out Descriptor.FileVersion))
+			{
+				if(!RawObject.TryGetNumericField("PluginFileVersion", out Descriptor.FileVersion))
+				{
+					throw new BuildException("Plugin descriptor file '{0}' does not contain a valid FileVersion entry", FileName);
+				}
+			}
+
+			// Check it's not newer than the latest version we can parse
+			if(Descriptor.FileVersion > (int)PluginDescriptorVersion.Latest)
+			{
+				throw new BuildException( "Plugin descriptor file '{0}' appears to be in a newer version ({1}) of the file format that we can load (max version: {2}).", FileName, Descriptor.FileVersion, (int)PluginDescriptorVersion.Latest);
+			}
+
+			// Read the other fields
+			RawObject.TryGetNumericField("Version", out Descriptor.Version);
+			RawObject.TryGetStringField("VersionName", out Descriptor.VersionName);
+			RawObject.TryGetStringField("FriendlyName", out Descriptor.FriendlyName);
+			RawObject.TryGetStringField("Description", out Descriptor.Description);
+
+			if (!RawObject.TryGetStringField("Category", out Descriptor.Category))
+			{
+				// Category used to be called CategoryPath in .uplugin files
+				RawObject.TryGetStringField("CategoryPath", out Descriptor.Category);
+			}
+        
+			// Due to a difference in command line parsing between Windows and Mac, we shipped a few Mac samples containing
+			// a category name with escaped quotes. Remove them here to make sure we can list them in the right category.
+			if (Descriptor.Category != null && Descriptor.Category.Length >= 2 && Descriptor.Category.StartsWith("\"") && Descriptor.Category.EndsWith("\""))
+			{
+				Descriptor.Category = Descriptor.Category.Substring(1, Descriptor.Category.Length - 2);
+			}
+
+			RawObject.TryGetStringField("CreatedBy", out Descriptor.CreatedBy);
+			RawObject.TryGetStringField("CreatedByURL", out Descriptor.CreatedByURL);
+			RawObject.TryGetStringField("DocsURL", out Descriptor.DocsURL);
+
+			JsonObject[] ModulesArray;
+			if(RawObject.TryGetObjectArrayField("Modules", out ModulesArray))
+			{
+				Descriptor.Modules = Array.ConvertAll(ModulesArray, x => ModuleDescriptor.FromJsonObject(x));
+			}
+
+			RawObject.TryGetBoolField("EnabledByDefault", out Descriptor.bEnabledByDefault);
+			RawObject.TryGetBoolField("CanContainContent", out Descriptor.bCanContainContent);
+			RawObject.TryGetBoolField("IsBetaVersion", out Descriptor.bIsBetaVersion);
+
+			return Descriptor;
+		}
+	};
+}
diff --git a/Engine/Source/Programs/UnrealBuildTool/System/Plugins.cs b/Engine/Source/Programs/UnrealBuildTool/System/Plugins.cs
--- a/Engine/Source/Programs/UnrealBuildTool/System/Plugins.cs
+++ b/Engine/Source/Programs/UnrealBuildTool/System/Plugins.cs
@@ -10,17 +10,6 @@ namespace UnrealBuildTool
 {
 	public class PluginInfo
 	{
-		public enum PluginModuleType
-		{
-			Runtime,
-			RuntimeNoCommandlet,
-			Developer,
-			Editor,
-			EditorNoCommandlet,
-			/** Program-only plugin */
-			Program,
-		}
-
 		public enum LoadedFromType
 		{
 			// Plugin is built-in to the engine
@@ -29,18 +18,6 @@ namespace UnrealBuildTool
 			// Project-specific plugin, stored within a game project directory
 			GameProject
 		};
-
-		public struct PluginModuleInfo
-		{
-			// Name of this module
-			public string Name;
-
-			// Type of module
-			public PluginModuleType Type;
-
-			// List of platforms supported by this modules
-			public List<UnrealTargetPlatform> Platforms;
-		}
 		
 		// Whether or not the plugin should be built
 		public bool bShouldBuild;
@@ -51,8 +28,8 @@ namespace UnrealBuildTool
 		// Path to the plugin's root directory
 		public string Directory;
 
-		// List of modules in this plugin
-		public readonly List<PluginModuleInfo> Modules = new List<PluginModuleInfo>();
+		// The plugin descriptor
+		public PluginDescriptor Descriptor;
 
 		// Whether this plugin is enabled by default
 		public bool bEnabledByDefault;
@@ -127,10 +104,6 @@ namespace UnrealBuildTool
 
 	public class Plugins
 	{
-		/// Latest supported version for plugin descriptor files.  We can still usually load older versions, but not newer version.
-		/// NOTE: This constant exists in PluginManager C++ code as well.
-		private static int LatestPluginDescriptorFileVersion = 3;	// IMPORTANT: Remember to also update EProjectFileVersion in PluginManagerShared.h when this changes!
-		
 		/// File extension of plugin descriptor files.  NOTE: This constant exists in UnrealBuildTool code as well.
 		/// NOTE: This constant exists in PluginManager C++ code as well.
 		private static string PluginDescriptorFileExtension = ".uplugin";
@@ -145,228 +118,15 @@ namespace UnrealBuildTool
 		/// <returns>New PluginInfo for the loaded descriptor.</returns>
 		private static PluginInfo LoadPluginDescriptor( FileInfo PluginFileInfo, PluginInfo.LoadedFromType LoadedFrom )
 		{
-			// Load the file up (JSon format)
-			Dictionary<string, object> PluginDescriptorDict;
-			{
-				string FileContent;
-				using( var StreamReader = new StreamReader( PluginFileInfo.FullName ) )
-				{
-					FileContent = StreamReader.ReadToEnd();
-				}
-
-				// Parse the Json into a dictionary
-				var CaseSensitiveJSonDict = fastJSON.JSON.Instance.ToObject< Dictionary< string, object > >( FileContent );
-
-				// Convert to a case-insensitive dictionary, so that we can be more tolerant of hand-typed files
-				PluginDescriptorDict = new Dictionary<string, object>( CaseSensitiveJSonDict, StringComparer.InvariantCultureIgnoreCase );
-			}
-
-
-			// File version check
-			long PluginVersionNumber;
-			{
-				// Try to get the version of the plugin
-				object PluginVersionObject;
-				if( !PluginDescriptorDict.TryGetValue( "FileVersion", out PluginVersionObject ) )
-				{
-					if( !PluginDescriptorDict.TryGetValue( "PluginFileVersion", out PluginVersionObject ) )
-					{
-						throw new BuildException( "Plugin descriptor file '{0}' does not contain a valid FileVersion entry", PluginFileInfo.FullName );
-					}
-				}
-
-				if( !( PluginVersionObject is long ) )
-				{
-					throw new BuildException( "Unable to parse the version number of the plugin descriptor file '{0}'", PluginFileInfo.FullName );
-				}
-
-				PluginVersionNumber = (long)PluginVersionObject;
-				if( PluginVersionNumber > LatestPluginDescriptorFileVersion )
-				{
-					throw new BuildException( "Plugin descriptor file '{0}' appears to be in a newer version ({1}) of the file format that we can load (max version: {2}).", PluginFileInfo.FullName, PluginVersionNumber, LatestPluginDescriptorFileVersion );
-				}
-
-				// @todo plugin: Should we also test the engine version here?  (we would need to load it from build.properties)
-			}
-			
-
-			// NOTE: At this point, we can use PluginVersionNumber to handle backwards compatibility when loading the rest of the file!
-
-
-			var PluginInfo = new PluginInfo();
-			PluginInfo.LoadedFrom = LoadedFrom;
-			PluginInfo.Directory = PluginFileInfo.Directory.FullName;
-			PluginInfo.Name = Path.GetFileName(PluginInfo.Directory);
-
-			// Determine whether the plugin should be enabled by default
-			object EnabledByDefaultObject;
-			if(PluginDescriptorDict.TryGetValue("EnabledByDefault", out EnabledByDefaultObject) && (EnabledByDefaultObject is bool))
-			{
-				PluginInfo.bEnabledByDefault = (bool)EnabledByDefaultObject;
-			}
-
-			// This plugin might have some modules that we need to know about.  Let's take a look.
-			{
-				object ModulesObject;
-				if( PluginDescriptorDict.TryGetValue( "Modules", out ModulesObject ) )
-				{
-					if( !( ModulesObject is Array ) )
-					{
-						throw new BuildException( "Found a 'Modules' entry in plugin descriptor file '{0}', but it doesn't appear to be in the array format that we were expecting.", PluginFileInfo.FullName );
-					}
-					
-					var ModulesArray = (Array)ModulesObject;
-					foreach( var ModuleObject in ModulesArray )
-					{
-						var ModuleDict = new Dictionary<string,object>( (Dictionary< string, object >)ModuleObject, StringComparer.InvariantCultureIgnoreCase );
-
-						var PluginModuleInfo = new PluginInfo.PluginModuleInfo();
-
-						// Module name
-						{
-							// All modules require a name to be set
-							object ModuleNameObject;
-							if( !ModuleDict.TryGetValue( "Name", out ModuleNameObject ) )
-							{
-								throw new BuildException( "Found a 'Module' entry with a missing 'Name' field in plugin descriptor file '{0}'", PluginFileInfo.FullName );
-							}
-							string ModuleName = (string)ModuleNameObject;
-
-							// @todo plugin: Locate this module right now and validate it?  Repair case?
-							PluginModuleInfo.Name = ModuleName;
-						}
-
-						
-						// Module type
-						{
-							// Check to see if the user specified the module's type
-							object ModuleTypeObject;
-							if( !ModuleDict.TryGetValue( "Type", out ModuleTypeObject ) )
-							{
-								throw new BuildException( "Found a Module entry '{0}' with a missing 'Type' field in plugin descriptor file '{1}'", PluginModuleInfo.Name, PluginFileInfo.FullName );
-							}
-							string ModuleTypeString = (string)ModuleTypeObject;
-
-							// Check to see if this is a valid type
-							bool FoundValidType = false;
-							foreach( PluginInfo.PluginModuleType PossibleType in Enum.GetValues( typeof( PluginInfo.PluginModuleType ) ) )
-							{
-								if( ModuleTypeString.Equals( PossibleType.ToString(), StringComparison.InvariantCultureIgnoreCase ) )
-								{
-									FoundValidType = true;
-									PluginModuleInfo.Type = PossibleType;
-									break;
-								}
-							}
-							if( !FoundValidType )
-							{
-								throw new BuildException( "Module entry '{0}' specified an unrecognized module Type '{1}' in plugin descriptor file '{0}'", PluginModuleInfo.Name, ModuleTypeString, PluginFileInfo.FullName );
-							}
-						}
-
-						// Supported platforms
-						PluginModuleInfo.Platforms = new List<UnrealTargetPlatform>();
-
-						// look for white and blacklists
-						object WhitelistObject, BlacklistObject;
-						ModuleDict.TryGetValue( "WhitelistPlatforms", out WhitelistObject );
-						ModuleDict.TryGetValue( "BlacklistPlatforms", out BlacklistObject );
-
-						if (WhitelistObject != null && BlacklistObject != null)
-						{
-							throw new BuildException( "Found a module '{0}' with both blacklist and whitelist platform lists in plugin file '{1}'", PluginModuleInfo.Name, PluginFileInfo.FullName );
-						}
-
-						// now process the whitelist
-						if (WhitelistObject != null)
-						{
-							if (!(WhitelistObject is Array))
-							{
-								throw new BuildException("Found a 'WhitelistPlatforms' entry in plugin descriptor file '{0}', but it doesn't appear to be in the array format that we were expecting.", PluginFileInfo.FullName);
-							}
-
-							// put the whitelist array directly into the plugin's modulelist
-							ConvertPlatformArrayToList((Array)WhitelistObject, ref PluginModuleInfo.Platforms, PluginFileInfo.FullName);
-						}
-						// handle the blacklist (or lack of blacklist and whitelist which means all platforms)
-						else
-						{
-							// start with all platforms supported
-							foreach (UnrealTargetPlatform Platform in Enum.GetValues( typeof( UnrealTargetPlatform ) ) )
-							{
-								PluginModuleInfo.Platforms.Add(Platform);
-							}
-
-							// if we want to disallow some platforms, then pull them out now
-							if (BlacklistObject != null)
-							{
-								if (!(BlacklistObject is Array))
-								{
-									throw new BuildException("Found a 'BlacklistPlatforms' entry in plugin descriptor file '{0}', but it doesn't appear to be in the array format that we were expecting.", PluginFileInfo.FullName);
-								}
-
-								// put the whitelist array directly into the plugin's modulelist
-								List<UnrealTargetPlatform> Blacklist = new List<UnrealTargetPlatform>();
-								ConvertPlatformArrayToList((Array)BlacklistObject, ref Blacklist, PluginFileInfo.FullName);
-
-								// now remove them from the module platform list
-								foreach (UnrealTargetPlatform Platform in Blacklist)
-								{
-									PluginModuleInfo.Platforms.Remove(Platform);
-								}
-							}
-						}
-						
-						object ModuleShouldBuild;
-						if( ModuleDict.TryGetValue( "bShouldBuild", out ModuleShouldBuild ) )
-						{
-							PluginInfo.bShouldBuild = (Int64)ModuleShouldBuild == 1 ? true : false;
-						}
-						else
-						{
-							PluginInfo.bShouldBuild = true;
-						}
-
-						if (PluginInfo.bShouldBuild)
-						{
-							// add to list of modules
-							PluginInfo.Modules.Add(PluginModuleInfo);
-						}
-					}
-				}
-				else
-				{
-					// Plugin contains no modules array.  That's fine.
-				}
-			}
-
-			return PluginInfo;
+			// Create the plugin info object
+			PluginInfo Info = new PluginInfo();
+			Info.LoadedFrom = LoadedFrom;
+			Info.Directory = PluginFileInfo.Directory.FullName;
+			Info.Name = Path.GetFileName(Info.Directory);
+			Info.Descriptor = PluginDescriptor.FromFile(PluginFileInfo.FullName);
+			return Info;
 		}
 
-		private static void ConvertPlatformArrayToList(Array PlatformNameList, ref List<UnrealTargetPlatform> Platforms, string ModuleFilename)
-		{
-			// look up each platform name in the array
-			foreach (string PlatformName in PlatformNameList)
-			{
-				// case-insensitive enum matching
-				bool bFoundValidType = false;
-				foreach( UnrealTargetPlatform PossiblePlatform in Enum.GetValues( typeof( UnrealTargetPlatform ) ) )
-				{
-					if( PlatformName.Equals( PossiblePlatform.ToString(), StringComparison.InvariantCultureIgnoreCase ) )
-					{
-						bFoundValidType = true;
-						// add it to the output!
-						Platforms.Add(PossiblePlatform);
-						break;
-					}
-				}
-
-				if( !bFoundValidType )
-				{
-					throw new BuildException( "Module entry specified unknown platform '{0}' in plugin descriptor file '{1}'", PlatformName, ModuleFilename );
-				}
-			}
-		}
 
 
 		/// <summary>
@@ -457,7 +217,7 @@ namespace UnrealBuildTool
 				ModulesToPluginMapVar = new Dictionary<string,PluginInfo>( StringComparer.InvariantCultureIgnoreCase );
 				foreach( var CurPluginInfo in AllPlugins )
 				{
-					foreach( var Module in CurPluginInfo.Modules )
+					foreach( var Module in CurPluginInfo.Descriptor.Modules )
 					{
 						// Make sure a different plugin doesn't already have a module with this name
 						// @todo plugin: Collisions like this could happen because of third party plugins added to a project, which isn't really ideal.
diff --git a/Engine/Source/Programs/UnrealBuildTool/UnrealBuildTool.csproj b/Engine/Source/Programs/UnrealBuildTool/UnrealBuildTool.csproj
--- a/Engine/Source/Programs/UnrealBuildTool/UnrealBuildTool.csproj
+++ b/Engine/Source/Programs/UnrealBuildTool/UnrealBuildTool.csproj
@@ -197,7 +197,7 @@
     <Compile Include="fastJSON\SafeDictionary.cs" />
     <Compile Include="Properties\AssemblyInfo.cs" />
     <Compile Include="Configuration\UEBuildConfiguration.cs" />
-    <Compile Include="Configuration\UEBuildTarget.cs" />    
+    <Compile Include="Configuration\UEBuildTarget.cs" />
     <Compile Include="System\ActionGraph.cs" />
     <Compile Include="System\ActionHistory.cs" />
     <Compile Include="System\BuildException.cs" />
@@ -207,6 +207,8 @@
     <Compile Include="System\FlatCPPIncludeDepencencyCache.cs" />
     <Compile Include="System\Distcc.cs" />
     <Compile Include="System\CMakefileGenerator.cs" />
+    <Compile Include="System\ModuleDescriptor.cs" />
+    <Compile Include="System\PluginDescriptor.cs" />
     <Compile Include="System\QMakefileGenerator.cs" />
     <Compile Include="System\ModuleProcessingException.cs" />
     <Compile Include="System\MakefileGenerator.cs" />
@@ -249,6 +251,7 @@
       <SubType>Code</SubType>
     </Compile>
     <Compile Include="Configuration\EngineConfiguration.cs" />
+    <Compile Include="Utilities\JsonObject.cs" />
     <Compile Include="Utilities\XmlConfigLoader.cs" />
     <Compile Include="Windows\UEBuildWindows.cs" />
     <Compile Include="Windows\VCEnvironment.cs" />
diff --git a/Engine/Source/Programs/UnrealBuildTool/Utilities/JsonObject.cs b/Engine/Source/Programs/UnrealBuildTool/Utilities/JsonObject.cs
new file mode 100644
--- /dev/null
+++ b/Engine/Source/Programs/UnrealBuildTool/Utilities/JsonObject.cs
@@ -0,0 +1,150 @@
+﻿using System;
+using System.Collections.Generic;
+using System.IO;
+using System.Linq;
+using System.Text;
+
+namespace UnrealBuildTool
+{
+	public class JsonObject
+	{
+		Dictionary<string, object> RawObject;
+
+		public JsonObject(Dictionary<string, object> InRawObject)
+		{
+			RawObject = new Dictionary<string,object>(InRawObject, StringComparer.InvariantCultureIgnoreCase);
+		}
+
+		public static JsonObject FromFile(string FileName)
+		{
+			string Text = File.ReadAllText(FileName);
+			Dictionary<string, object> CaseSensitiveRawObject = fastJSON.JSON.Instance.ToObject<Dictionary<string, object>>(Text);
+			return new JsonObject(CaseSensitiveRawObject);
+		}
+
+		public string GetStringField(string FieldName)
+		{
+			string StringValue;
+			if(!TryGetStringField(FieldName, out StringValue))
+			{
+				throw new BuildException("Missing or invalid '{0}' field", FieldName);
+			}
+			return StringValue;
+		}
+
+		public bool TryGetStringField(string FieldName, out string Result)
+		{
+			object RawValue;
+			if(RawObject.TryGetValue(FieldName, out RawValue) && (RawValue is string))
+			{
+				Result = (string)RawValue;
+				return true;
+			}
+			else
+			{
+				Result = null;
+				return false;
+			}
+		}
+
+		public bool TryGetStringArrayField(string FieldName, out string[] Result)
+		{
+			object RawValue;
+			if(RawObject.TryGetValue(FieldName, out RawValue) && (RawValue is object[]) && ((object[])RawValue).All(x => x is string))
+			{
+				Result = Array.ConvertAll((object[])RawValue, x => (string)x);
+				return true;
+			}
+			else
+			{
+				Result = null;
+				return false;
+			}
+		}
+
+		public bool TryGetBoolField(string FieldName, out bool Result)
+		{
+			object RawValue;
+			if(RawObject.TryGetValue(FieldName, out RawValue) && (RawValue is Boolean))
+			{
+				Result = (bool)RawValue;
+				return true;
+			}
+			else
+			{
+				Result = false;
+				return false;
+			}
+		}
+
+		public bool TryGetNumericField(string FieldName, out int Result)
+		{
+			object RawValue;
+			if(!RawObject.TryGetValue(FieldName, out RawValue) || !int.TryParse(RawValue.ToString(), out Result))
+			{
+				Result = 0;
+				return false;
+			}
+			return true;
+		}
+
+		public T GetEnumField<T>(string FieldName) where T : struct
+		{
+			T EnumValue;
+			if(!TryGetEnumField(FieldName, out EnumValue))
+			{
+				throw new BuildException("Missing or invalid '{0}' field", FieldName);
+			}
+			return EnumValue;
+		}
+
+		public bool TryGetEnumField<T>(string FieldName, out T Result) where T : struct
+		{
+			string StringValue;
+			if(!TryGetStringField(FieldName, out StringValue) || !Enum.TryParse<T>(StringValue, true, out Result))
+			{
+				Result = default(T);
+				return false;
+			}
+			return true;
+		}
+
+		public bool TryGetEnumArrayField<T>(string FieldName, out T[] Result) where T : struct
+		{
+			string[] StringValues;
+			if(!TryGetStringArrayField(FieldName, out StringValues))
+			{
+				Result = null;
+				return false;
+			}
+
+			T[] EnumValues = new T[StringValues.Length];
+			for(int Idx = 0; Idx < StringValues.Length; Idx++)
+			{
+				if(!Enum.TryParse<T>(StringValues[Idx], true, out EnumValues[Idx]))
+				{
+					Result = null;
+					return false;
+				}
+			}
+
+			Result = EnumValues;
+			return true;
+		}
+
+		public bool TryGetObjectArrayField(string FieldName, out JsonObject[] Result)
+		{
+			object RawValue;
+			if(RawObject.TryGetValue(FieldName, out RawValue) && (RawValue is object[]) && ((object[])RawValue).All(x => x is Dictionary<string, object>))
+			{
+				Result = Array.ConvertAll((object[])RawValue, x => new JsonObject((Dictionary<string, object>)x));
+				return true;
+			}
+			else
+			{
+				Result = null;
+				return false;
+			}
+		}
+	}
+}
