Added implementation of ProjectDescriptor to UBT, which mirrors definition in the engine.

diff --git a/Engine/Source/Programs/UnrealBuildTool/System/PluginDescriptor.cs b/Engine/Source/Programs/UnrealBuildTool/System/PluginDescriptor.cs
--- a/Engine/Source/Programs/UnrealBuildTool/System/PluginDescriptor.cs
+++ b/Engine/Source/Programs/UnrealBuildTool/System/PluginDescriptor.cs
@@ -1,5 +1,6 @@
 ﻿using System;
 using System.Collections.Generic;
+using System.Diagnostics;
 using System.Linq;
 using System.Text;
 
@@ -78,58 +79,130 @@ namespace UnrealBuildTool
 		public static PluginDescriptor FromFile(string FileName)
 		{
 			JsonObject RawObject = JsonObject.FromFile(FileName);
+			try
+			{
+				PluginDescriptor Descriptor = new PluginDescriptor();
 
-			PluginDescriptor Descriptor = new PluginDescriptor();
+				// Read the version
+				if(!RawObject.TryGetIntegerField("FileVersion", out Descriptor.FileVersion))
+				{
+					if(!RawObject.TryGetIntegerField("PluginFileVersion", out Descriptor.FileVersion))
+					{
+						throw new BuildException("Plugin descriptor file '{0}' does not contain a valid FileVersion entry", FileName);
+					}
+				}
 
-			// Read the version
-			if(!RawObject.TryGetNumericField("FileVersion", out Descriptor.FileVersion))
-			{
-				if(!RawObject.TryGetNumericField("PluginFileVersion", out Descriptor.FileVersion))
+				// Check it's not newer than the latest version we can parse
+				if(Descriptor.FileVersion > (int)PluginDescriptorVersion.Latest)
 				{
-					throw new BuildException("Plugin descriptor file '{0}' does not contain a valid FileVersion entry", FileName);
+					throw new BuildException( "Plugin descriptor file '{0}' appears to be in a newer version ({1}) of the file format that we can load (max version: {2}).", FileName, Descriptor.FileVersion, (int)PluginDescriptorVersion.Latest);
+				}
+
+				// Read the other fields
+				RawObject.TryGetIntegerField("Version", out Descriptor.Version);
+				RawObject.TryGetStringField("VersionName", out Descriptor.VersionName);
+				RawObject.TryGetStringField("FriendlyName", out Descriptor.FriendlyName);
+				RawObject.TryGetStringField("Description", out Descriptor.Description);
+
+				if (!RawObject.TryGetStringField("Category", out Descriptor.Category))
+				{
+					// Category used to be called CategoryPath in .uplugin files
+					RawObject.TryGetStringField("CategoryPath", out Descriptor.Category);
+				}
+        
+				// Due to a difference in command line parsing between Windows and Mac, we shipped a few Mac samples containing
+				// a category name with escaped quotes. Remove them here to make sure we can list them in the right category.
+				if (Descriptor.Category != null && Descriptor.Category.Length >= 2 && Descriptor.Category.StartsWith("\"") && Descriptor.Category.EndsWith("\""))
+				{
+					Descriptor.Category = Descriptor.Category.Substring(1, Descriptor.Category.Length - 2);
 				}
-			}
 
-			// Check it's not newer than the latest version we can parse
-			if(Descriptor.FileVersion > (int)PluginDescriptorVersion.Latest)
+				RawObject.TryGetStringField("CreatedBy", out Descriptor.CreatedBy);
+				RawObject.TryGetStringField("CreatedByURL", out Descriptor.CreatedByURL);
+				RawObject.TryGetStringField("DocsURL", out Descriptor.DocsURL);
+
+				JsonObject[] ModulesArray;
+				if(RawObject.TryGetObjectArrayField("Modules", out ModulesArray))
+				{
+					Descriptor.Modules = Array.ConvertAll(ModulesArray, x => ModuleDescriptor.FromJsonObject(x));
+				}
+
+				RawObject.TryGetBoolField("EnabledByDefault", out Descriptor.bEnabledByDefault);
+				RawObject.TryGetBoolField("CanContainContent", out Descriptor.bCanContainContent);
+				RawObject.TryGetBoolField("IsBetaVersion", out Descriptor.bIsBetaVersion);
+
+				return Descriptor;
+			}
+			catch(JsonParseException ParseException)
 			{
-				throw new BuildException( "Plugin descriptor file '{0}' appears to be in a newer version ({1}) of the file format that we can load (max version: {2}).", FileName, Descriptor.FileVersion, (int)PluginDescriptorVersion.Latest);
+				throw new JsonParseException("{0} (in {1})", ParseException.Message, FileName);
 			}
+		}
+	}
+
+	[DebuggerDisplay("Name={Name}")]
+	public class PluginReferenceDescriptor
+	{
+		// Name of the plugin
+		public string Name;
+
+		// Whether it should be enabled by default
+		public bool bEnabled;
+
+		// Description of the plugin for users that do not have it installed.
+		public string Description;
+
+		// If enabled, list of platforms for which the plugin should be enabled (or all platforms if blank).
+		UnrealTargetPlatform[] WhitelistPlatforms;
+
+		// If enabled, list of platforms for which the plugin should be disabled.
+		UnrealTargetPlatform[] BlacklistPlatforms;
+
+		/// <summary>
+		/// Constructor
+		/// </summary>
+		/// <param name="InName">Name of the plugin</param>
+		/// <param name="bInEnabled">Whether the plugin is enabled</param>
+		public PluginReferenceDescriptor(string InName, bool bInEnabled)
+		{
+			Name = InName;
+			bEnabled = bInEnabled;
+		}
 
-			// Read the other fields
-			RawObject.TryGetNumericField("Version", out Descriptor.Version);
-			RawObject.TryGetStringField("VersionName", out Descriptor.VersionName);
-			RawObject.TryGetStringField("FriendlyName", out Descriptor.FriendlyName);
+		/// <summary>
+		/// Construct a PluginReferenceDescriptor from a Json object
+		/// </summary>
+		/// <param name="RawObject">The Json object containing a plugin reference descriptor</param>
+		/// <returns>New PluginReferenceDescriptor object</returns>
+		public static PluginReferenceDescriptor FromJsonObject(JsonObject RawObject)
+		{
+			PluginReferenceDescriptor Descriptor = new PluginReferenceDescriptor(RawObject.GetStringField("Name"), RawObject.GetBoolField("Enabled"));
 			RawObject.TryGetStringField("Description", out Descriptor.Description);
+			RawObject.TryGetEnumArrayField<UnrealTargetPlatform>("WhitelistPlatforms", out Descriptor.WhitelistPlatforms);
+			RawObject.TryGetEnumArrayField<UnrealTargetPlatform>("BlacklistPlatforms", out Descriptor.BlacklistPlatforms);
+			return Descriptor;
+		}
 
-			if (!RawObject.TryGetStringField("Category", out Descriptor.Category))
+		/// <summary>
+		/// Determines if this reference enables the plugin for a given platform
+		/// </summary>
+		/// <param name="Platform">The platform to check</param>
+		/// <returns>True if the plugin should be enabled</returns>
+		public bool IsEnabledForPlatform(UnrealTargetPlatform Platform)
+		{
+			if(!bEnabled)
 			{
-				// Category used to be called CategoryPath in .uplugin files
-				RawObject.TryGetStringField("CategoryPath", out Descriptor.Category);
+				return false;
 			}
-        
-			// Due to a difference in command line parsing between Windows and Mac, we shipped a few Mac samples containing
-			// a category name with escaped quotes. Remove them here to make sure we can list them in the right category.
-			if (Descriptor.Category != null && Descriptor.Category.Length >= 2 && Descriptor.Category.StartsWith("\"") && Descriptor.Category.EndsWith("\""))
+			if(WhitelistPlatforms != null && WhitelistPlatforms.Length > 0 && !WhitelistPlatforms.Contains(Platform))
 			{
-				Descriptor.Category = Descriptor.Category.Substring(1, Descriptor.Category.Length - 2);
+				return false;
 			}
-
-			RawObject.TryGetStringField("CreatedBy", out Descriptor.CreatedBy);
-			RawObject.TryGetStringField("CreatedByURL", out Descriptor.CreatedByURL);
-			RawObject.TryGetStringField("DocsURL", out Descriptor.DocsURL);
-
-			JsonObject[] ModulesArray;
-			if(RawObject.TryGetObjectArrayField("Modules", out ModulesArray))
+			if(BlacklistPlatforms != null && BlacklistPlatforms.Contains(Platform))
 			{
-				Descriptor.Modules = Array.ConvertAll(ModulesArray, x => ModuleDescriptor.FromJsonObject(x));
+				return false;
 			}
-
-			RawObject.TryGetBoolField("EnabledByDefault", out Descriptor.bEnabledByDefault);
-			RawObject.TryGetBoolField("CanContainContent", out Descriptor.bCanContainContent);
-			RawObject.TryGetBoolField("IsBetaVersion", out Descriptor.bIsBetaVersion);
-
-			return Descriptor;
+			return true;
 		}
-	};
+	}
 }
diff --git a/Engine/Source/Programs/UnrealBuildTool/System/Plugins.cs b/Engine/Source/Programs/UnrealBuildTool/System/Plugins.cs
--- a/Engine/Source/Programs/UnrealBuildTool/System/Plugins.cs
+++ b/Engine/Source/Programs/UnrealBuildTool/System/Plugins.cs
@@ -43,65 +43,6 @@ namespace UnrealBuildTool
 		}
 	}
 
-	class PluginReferenceDescriptor
-	{
-		public string Name;
-		public bool bEnabled;
-		public List<UnrealTargetPlatform> WhitelistPlatforms = new List<UnrealTargetPlatform>();
-		public List<UnrealTargetPlatform> BlacklistPlatforms = new List<UnrealTargetPlatform>();
-
-		public PluginReferenceDescriptor(string InName, bool bInEnabled)
-		{
-			Name = InName;
-			bEnabled = bInEnabled;
-		}
-
-		public static PluginReferenceDescriptor FromJson(Dictionary<string, object> Dictionary)
-		{
-			PluginReferenceDescriptor Descriptor = new PluginReferenceDescriptor((string)Dictionary["Name"], (bool)Dictionary["Enabled"]);
-			TryParsePlatformList(Dictionary, "WhitelistPlatforms", Descriptor.WhitelistPlatforms);
-			TryParsePlatformList(Dictionary, "BlacklistPlatforms", Descriptor.BlacklistPlatforms);
-			return Descriptor;
-		}
-
-		public bool IsEnabledForPlatform(UnrealTargetPlatform Platform)
-		{
-			if(!bEnabled)
-			{
-				return false;
-			}
-			if(WhitelistPlatforms.Count > 0 && !WhitelistPlatforms.Contains(Platform))
-			{
-				return false;
-			}
-			if(BlacklistPlatforms.Contains(Platform))
-			{
-				return false;
-			}
-			return true;
-		}
-
-		static void TryParsePlatformList(Dictionary<string, object> Dictionary, string FieldName, List<UnrealTargetPlatform> Platforms)
-		{
-			object ArrayObject;
-			if(Dictionary.TryGetValue(FieldName, out ArrayObject))
-			{
-				foreach(string PlatformName in ((object[])ArrayObject).Select(x => (string)x))
-				{
-					UnrealTargetPlatform Platform;
-					if(Enum.TryParse(PlatformName, true, out Platform))
-					{
-						Platforms.Add(Platform);
-					}
-					else
-					{
-						throw new BuildException("Unknown platform name '{0}'", PlatformName);
-					}
-				}
-			}
-		}
-	}
-
 	public class Plugins
 	{
 		/// File extension of plugin descriptor files.  NOTE: This constant exists in UnrealBuildTool code as well.
diff --git a/Engine/Source/Programs/UnrealBuildTool/System/ProjectDescriptor.cs b/Engine/Source/Programs/UnrealBuildTool/System/ProjectDescriptor.cs
new file mode 100644
--- /dev/null
+++ b/Engine/Source/Programs/UnrealBuildTool/System/ProjectDescriptor.cs
@@ -0,0 +1,115 @@
+﻿using System;
+using System.Collections.Generic;
+using System.Linq;
+using System.Text;
+
+namespace UnrealBuildTool
+{
+	enum ProjectDescriptorVersion
+	{
+		Invalid = 0,
+		Initial = 1,
+		NameHash = 2,
+		ProjectPluginUnification = 3,
+		// !!!!!!!!!! IMPORTANT: Remember to also update LatestPluginDescriptorFileVersion in Plugins.cs (and Plugin system documentation) when this changes!!!!!!!!!!!
+		// -----<new versions can be added before this line>-------------------------------------------------
+		// - this needs to be the last line (see note below)
+		LatestPlusOne,
+		Latest = LatestPlusOne - 1
+	}
+
+	class ProjectDescriptor
+	{
+		// Descriptor version number. */
+		public int FileVersion;
+
+		// The engine to open this project with.
+		public string EngineAssociation;
+
+		// Category to show under the project browser
+		public string Category;
+
+		// Description to show in the project browser
+		public string Description;
+
+		// List of all modules associated with this project
+		public ModuleDescriptor[] Modules;
+
+		// List of plugins for this project (may be enabled/disabled)
+		public PluginReferenceDescriptor[] Plugins;
+
+		// Array of platforms that this project is targeting
+		public string[] TargetPlatforms;
+
+		// A hash that is used to determine if the project was forked from a sample
+		public uint EpicSampleNameHash;
+
+		/// <summary>
+		/// Constructor.
+		/// </summary>
+		public ProjectDescriptor()
+		{
+			FileVersion = (int)ProjectDescriptorVersion.Latest;
+		}
+
+		/// <summary>
+		/// Creates a plugin descriptor from a file on disk
+		/// </summary>
+		/// <param name="FileName">The filename to read</param>
+		/// <returns>New plugin descriptor</returns>
+		public static ProjectDescriptor FromFile(string FileName)
+		{
+			JsonObject RawObject = JsonObject.FromFile(FileName);
+			try
+			{
+				ProjectDescriptor Descriptor = new ProjectDescriptor();
+
+				// Read the version
+				if(!RawObject.TryGetIntegerField("FileVersion", out Descriptor.FileVersion))
+				{
+					if(!RawObject.TryGetIntegerField("ProjectFileVersion", out Descriptor.FileVersion))
+					{
+						throw new BuildException("Project descriptor '{0}' does not contain a valid FileVersion entry", FileName);
+					}
+				}
+
+				// Check it's not newer than the latest version we can parse
+				if(Descriptor.FileVersion > (int)PluginDescriptorVersion.Latest)
+				{
+					throw new BuildException("Project descriptor '{0}' appears to be in a newer version ({1}) of the file format that we can load (max version: {2}).", FileName, Descriptor.FileVersion, (int)ProjectDescriptorVersion.Latest);
+				}
+
+				// Read simple fields
+				RawObject.TryGetStringField("EngineAssociation", out Descriptor.EngineAssociation);
+				RawObject.TryGetStringField("Category", out Descriptor.Category);
+				RawObject.TryGetStringField("Description", out Descriptor.Description);
+
+				// Read the modules
+				JsonObject[] ModulesArray;
+				if(RawObject.TryGetObjectArrayField("Modules", out ModulesArray))
+				{
+					Descriptor.Modules = Array.ConvertAll(ModulesArray, x => ModuleDescriptor.FromJsonObject(x));
+				}
+
+				// Read the plugins
+				JsonObject[] PluginsArray;
+				if(RawObject.TryGetObjectArrayField("Plugins", out PluginsArray))
+				{
+					Descriptor.Plugins = Array.ConvertAll(PluginsArray, x => PluginReferenceDescriptor.FromJsonObject(x));
+				}
+
+				// Read the target platforms
+				RawObject.TryGetStringArrayField("TargetPlatforms", out Descriptor.TargetPlatforms);
+
+				// Get the sample name hash
+				RawObject.TryGetUnsignedIntegerField("EpicSampleNameHash", out Descriptor.EpicSampleNameHash);
+
+				return Descriptor;
+			}
+			catch(JsonParseException ParseException)
+			{
+				throw new JsonParseException("{0} (in {1})", ParseException.Message, FileName);
+			}
+		}
+	}
+}
diff --git a/Engine/Source/Programs/UnrealBuildTool/System/UProjectInfo.cs b/Engine/Source/Programs/UnrealBuildTool/System/UProjectInfo.cs
--- a/Engine/Source/Programs/UnrealBuildTool/System/UProjectInfo.cs
+++ b/Engine/Source/Programs/UnrealBuildTool/System/UProjectInfo.cs
@@ -307,39 +307,22 @@ namespace UnrealBuildTool
 		public static List<string> GetEnabledPlugins(string ProjectFileName, List<string> DefaultEnabledPlugins, UnrealTargetPlatform Platform)
 		{
 			// Parse the project file
-			string ProjectFile = File.ReadAllText(ProjectFileName);
-			try
-			{
-				List<string> EnabledPlugins = new List<string>(DefaultEnabledPlugins);
-
-				// Create a case-insensitive dictionary of the contents
-				Dictionary<string, object> Descriptor = fastJSON.JSON.Instance.ToObject<Dictionary<string, object>>(ProjectFile);
-				Descriptor = new Dictionary<string,object>(Descriptor, StringComparer.InvariantCultureIgnoreCase);
+			ProjectDescriptor Project = ProjectDescriptor.FromFile(ProjectFileName);
 
-				// Get the list of plugins
-				object PluginsObject;
-				if(Descriptor.TryGetValue("Plugins", out PluginsObject))
+			// Get the list of plugins
+			List<string> EnabledPlugins = new List<string>(DefaultEnabledPlugins);
+			foreach(PluginReferenceDescriptor Plugin in Project.Plugins)
+			{
+				if(Plugin.IsEnabledForPlatform(Platform))
 				{
-					object[] PluginsArrayObject = (object[])PluginsObject;
-					foreach(object PluginObject in PluginsArrayObject)
-					{
-						PluginReferenceDescriptor Plugin = PluginReferenceDescriptor.FromJson((Dictionary<string, object>)PluginObject);
-						if(Plugin.IsEnabledForPlatform(Platform))
-						{
-							EnabledPlugins.Add(Plugin.Name);
-						}
-						else
-						{
-							EnabledPlugins.RemoveAll(x => x.Equals(Plugin.Name, StringComparison.CurrentCultureIgnoreCase));
-						}
-					}
+					EnabledPlugins.Add(Plugin.Name);
+				}
+				else
+				{
+					EnabledPlugins.RemoveAll(x => x.Equals(Plugin.Name, StringComparison.CurrentCultureIgnoreCase));
 				}
-				return EnabledPlugins;
-			}
-			catch(BuildException Exception)
-			{
-				throw new BuildException("While parsing '{0}': {1}", ProjectFileName, Exception.ToString());
 			}
+			return EnabledPlugins;
 		}
 	}
 }
diff --git a/Engine/Source/Programs/UnrealBuildTool/UnrealBuildTool.csproj b/Engine/Source/Programs/UnrealBuildTool/UnrealBuildTool.csproj
--- a/Engine/Source/Programs/UnrealBuildTool/UnrealBuildTool.csproj
+++ b/Engine/Source/Programs/UnrealBuildTool/UnrealBuildTool.csproj
@@ -209,6 +209,7 @@
     <Compile Include="System\CMakefileGenerator.cs" />
     <Compile Include="System\ModuleDescriptor.cs" />
     <Compile Include="System\PluginDescriptor.cs" />
+    <Compile Include="System\ProjectDescriptor.cs" />
     <Compile Include="System\QMakefileGenerator.cs" />
     <Compile Include="System\ModuleProcessingException.cs" />
     <Compile Include="System\MakefileGenerator.cs" />
diff --git a/Engine/Source/Programs/UnrealBuildTool/Utilities/JsonObject.cs b/Engine/Source/Programs/UnrealBuildTool/Utilities/JsonObject.cs
--- a/Engine/Source/Programs/UnrealBuildTool/Utilities/JsonObject.cs
+++ b/Engine/Source/Programs/UnrealBuildTool/Utilities/JsonObject.cs
@@ -6,6 +6,13 @@ using System.Text;
 
 namespace UnrealBuildTool
 {
+	public class JsonParseException : Exception
+	{
+		public JsonParseException(string Format, params string[] Args) : base(String.Format(Format, Args))
+		{
+		}
+	}
+
 	public class JsonObject
 	{
 		Dictionary<string, object> RawObject;
@@ -27,7 +34,7 @@ namespace UnrealBuildTool
 			string StringValue;
 			if(!TryGetStringField(FieldName, out StringValue))
 			{
-				throw new BuildException("Missing or invalid '{0}' field", FieldName);
+				throw new JsonParseException("Missing or invalid '{0}' field", FieldName);
 			}
 			return StringValue;
 		}
@@ -47,6 +54,16 @@ namespace UnrealBuildTool
 			}
 		}
 
+		public string[] GetStringArrayField(string FieldName)
+		{
+			string[] StringValues;
+			if(!TryGetStringArrayField(FieldName, out StringValues))
+			{
+				throw new JsonParseException("Missing or invalid '{0}' field", FieldName);
+			}
+			return StringValues;
+		}
+
 		public bool TryGetStringArrayField(string FieldName, out string[] Result)
 		{
 			object RawValue;
@@ -62,6 +79,16 @@ namespace UnrealBuildTool
 			}
 		}
 
+		public bool GetBoolField(string FieldName)
+		{
+			bool BoolValue;
+			if(!TryGetBoolField(FieldName, out BoolValue))
+			{
+				throw new JsonParseException("Missing or invalid '{0}' field", FieldName);
+			}
+			return BoolValue;
+		}
+
 		public bool TryGetBoolField(string FieldName, out bool Result)
 		{
 			object RawValue;
@@ -77,7 +104,7 @@ namespace UnrealBuildTool
 			}
 		}
 
-		public bool TryGetNumericField(string FieldName, out int Result)
+		public bool TryGetIntegerField(string FieldName, out int Result)
 		{
 			object RawValue;
 			if(!RawObject.TryGetValue(FieldName, out RawValue) || !int.TryParse(RawValue.ToString(), out Result))
@@ -88,12 +115,23 @@ namespace UnrealBuildTool
 			return true;
 		}
 
+		public bool TryGetUnsignedIntegerField(string FieldName, out uint Result)
+		{
+			object RawValue;
+			if(!RawObject.TryGetValue(FieldName, out RawValue) || !uint.TryParse(RawValue.ToString(), out Result))
+			{
+				Result = 0;
+				return false;
+			}
+			return true;
+		}
+
 		public T GetEnumField<T>(string FieldName) where T : struct
 		{
 			T EnumValue;
 			if(!TryGetEnumField(FieldName, out EnumValue))
 			{
-				throw new BuildException("Missing or invalid '{0}' field", FieldName);
+				throw new JsonParseException("Missing or invalid '{0}' field", FieldName);
 			}
 			return EnumValue;
 		}
