Add null check when getting list of enabled plugins.

diff --git a/Engine/Source/Programs/UnrealBuildTool/System/UProjectInfo.cs b/Engine/Source/Programs/UnrealBuildTool/System/UProjectInfo.cs
--- a/Engine/Source/Programs/UnrealBuildTool/System/UProjectInfo.cs
+++ b/Engine/Source/Programs/UnrealBuildTool/System/UProjectInfo.cs
@@ -311,15 +311,18 @@ namespace UnrealBuildTool
 
 			// Get the list of plugins
 			List<string> EnabledPlugins = new List<string>(DefaultEnabledPlugins);
-			foreach(PluginReferenceDescriptor Plugin in Project.Plugins)
+			if(Project.Plugins != null)
 			{
-				if(Plugin.IsEnabledForPlatform(Platform))
+				foreach(PluginReferenceDescriptor Plugin in Project.Plugins)
 				{
-					EnabledPlugins.Add(Plugin.Name);
-				}
-				else
-				{
-					EnabledPlugins.RemoveAll(x => x.Equals(Plugin.Name, StringComparison.CurrentCultureIgnoreCase));
+					if(Plugin.IsEnabledForPlatform(Platform))
+					{
+						EnabledPlugins.Add(Plugin.Name);
+					}
+					else
+					{
+						EnabledPlugins.RemoveAll(x => x.Equals(Plugin.Name, StringComparison.CurrentCultureIgnoreCase));
+					}
 				}
 			}
 			return EnabledPlugins;
