Fix "EnabledByDefault" plugins not being compiled as part of targets.

diff --git a/Engine/Source/Programs/UnrealBuildTool/Configuration/UEBuildTarget.cs b/Engine/Source/Programs/UnrealBuildTool/Configuration/UEBuildTarget.cs
--- a/Engine/Source/Programs/UnrealBuildTool/Configuration/UEBuildTarget.cs
+++ b/Engine/Source/Programs/UnrealBuildTool/Configuration/UEBuildTarget.cs
@@ -2487,7 +2487,7 @@ namespace UnrealBuildTool
 			// Add the list of plugins enabled by default
 			if (UEBuildConfiguration.bCompileAgainstEngine)
 			{
-				EnabledPluginNames.AddRange(ValidPlugins.Where(x => x.bEnabledByDefault).Select(x => x.Name));
+				EnabledPluginNames.AddRange(ValidPlugins.Where(x => x.Descriptor.bEnabledByDefault).Select(x => x.Name));
 			}
 
 			// Update the plugin list for game targets
diff --git a/Engine/Source/Programs/UnrealBuildTool/System/Plugins.cs b/Engine/Source/Programs/UnrealBuildTool/System/Plugins.cs
--- a/Engine/Source/Programs/UnrealBuildTool/System/Plugins.cs
+++ b/Engine/Source/Programs/UnrealBuildTool/System/Plugins.cs
@@ -19,9 +19,6 @@ namespace UnrealBuildTool
 			GameProject
 		};
 		
-		// Whether or not the plugin should be built
-		public bool bShouldBuild;
-
 		// Plugin name
 		public string Name;
 
@@ -31,9 +28,6 @@ namespace UnrealBuildTool
 		// The plugin descriptor
 		public PluginDescriptor Descriptor;
 
-		// Whether this plugin is enabled by default
-		public bool bEnabledByDefault;
-
 		// Where does this plugin live?
 		public LoadedFromType LoadedFrom;
 
