Don't treat VC++ compiler differently than other executables for XGE
- Experimenting with disabling this to see if we notice any drawbacks (seems fine after a few days of testing)

diff --git a/Engine/Source/Programs/UnrealBuildTool/Configuration/UEBuildBinary.cs b/Engine/Source/Programs/UnrealBuildTool/Configuration/UEBuildBinary.cs
--- a/Engine/Source/Programs/UnrealBuildTool/Configuration/UEBuildBinary.cs
+++ b/Engine/Source/Programs/UnrealBuildTool/Configuration/UEBuildBinary.cs
@@ -847,7 +847,6 @@ namespace UnrealBuildTool
 			Action LinkAction = new Action(ActionType.Compile);
 			LinkAction.WorkingDirectory = Path.GetFullPath(".");
 			LinkAction.CommandPath = System.IO.Path.Combine(LinkAction.WorkingDirectory, @"..", @"Binaries", @"Win32", @"NotForLicensees", @"UnrealCodeAnalyzer.exe");
-			LinkAction.bIsVCCompiler = false;
 			LinkAction.ProducedItems.Add(OutputFile);
 			LinkAction.PrerequisiteItems.AddRange(BinaryLinkEnvironment.InputFiles);
 			LinkAction.CommandArguments = @"-AnalyzePCHFile -PCHFile=""" + ModulePrivatePCH.AbsolutePath + @""" -OutputFile=""" + OutputFileName + @""" -HeaderDataPath=""" + IntermediatePath + @""" -UsageThreshold " + BuildConfiguration.UCAUsageThreshold.ToString(CultureInfo.InvariantCulture);
diff --git a/Engine/Source/Programs/UnrealBuildTool/System/ActionGraph.cs b/Engine/Source/Programs/UnrealBuildTool/System/ActionGraph.cs
--- a/Engine/Source/Programs/UnrealBuildTool/System/ActionGraph.cs
+++ b/Engine/Source/Programs/UnrealBuildTool/System/ActionGraph.cs
@@ -62,9 +62,6 @@ namespace UnrealBuildTool
 		/** True if this action is allowed to be run on a remote machine when a distributed build system is being used, such as XGE */
 		public bool bCanExecuteRemotely = false;
 
-		/** True if this action is using the Visual C++ compiler.  Some build systems may be able to optimize for this case. */
-		public bool bIsVCCompiler = false;
-
 		/** True if this action is using the GCC compiler.  Some build systems may be able to optimize for this case. */
 		public bool bIsGCCCompiler = false;
 
diff --git a/Engine/Source/Programs/UnrealBuildTool/System/XGE.cs b/Engine/Source/Programs/UnrealBuildTool/System/XGE.cs
--- a/Engine/Source/Programs/UnrealBuildTool/System/XGE.cs
+++ b/Engine/Source/Programs/UnrealBuildTool/System/XGE.cs
@@ -295,6 +295,9 @@ namespace UnrealBuildTool
 				ToolElement.SetAttribute("Name", string.Format("Tool{0}", ActionIndex));
 				ToolElement.SetAttribute("AllowRemote", Action.bCanExecuteRemotely.ToString());
 
+				// The XGE documentation says that 'AllowIntercept' must be set to 'true' for all tools where 'AllowRemote' is enabled
+				ToolElement.SetAttribute("AllowIntercept", Action.bCanExecuteRemotely.ToString());
+
 				string OutputPrefix = "";
 				if (bProgressMarkup)
 				{
@@ -335,7 +338,6 @@ namespace UnrealBuildTool
 
 				ToolElement.SetAttribute("Params", CommandArguments);
 				ToolElement.SetAttribute("Path", CommandPath);
-				ToolElement.SetAttribute("VCCompiler", Action.bIsVCCompiler.ToString());
 				ToolElement.SetAttribute("SkipIfProjectFailed", "true");
 				if (Action.bIsGCCCompiler)
 				{
diff --git a/Engine/Source/Programs/UnrealBuildTool/WinRT/WinRTToolChain.cs b/Engine/Source/Programs/UnrealBuildTool/WinRT/WinRTToolChain.cs
--- a/Engine/Source/Programs/UnrealBuildTool/WinRT/WinRTToolChain.cs
+++ b/Engine/Source/Programs/UnrealBuildTool/WinRT/WinRTToolChain.cs
@@ -645,7 +645,6 @@ namespace UnrealBuildTool
 
 				CompileAction.WorkingDirectory = Path.GetFullPath(".");
 				CompileAction.CommandPath = GetVCToolPath(CompileEnvironment.Config.Target.Platform, CompileEnvironment.Config.Target.Configuration, "cl");
-				CompileAction.bIsVCCompiler = true;
 				CompileAction.CommandArguments = Arguments + FileArguments + CompileEnvironment.Config.AdditionalArguments;
 				CompileAction.StatusDescription = string.Format("{0}", Path.GetFileName(SourceFile.AbsolutePath));
 
diff --git a/Engine/Source/Programs/UnrealBuildTool/Windows/VCToolChain.cs b/Engine/Source/Programs/UnrealBuildTool/Windows/VCToolChain.cs
--- a/Engine/Source/Programs/UnrealBuildTool/Windows/VCToolChain.cs
+++ b/Engine/Source/Programs/UnrealBuildTool/Windows/VCToolChain.cs
@@ -1250,11 +1250,6 @@ namespace UnrealBuildTool
 					CompileAction.CommandPath = EnvVars.CompilerPath;
 				}
 
-				if( !WindowsPlatform.bCompileWithClang && !BuildConfiguration.bRunUnrealCodeAnalyzer )
-				{
-					CompileAction.bIsVCCompiler = true;
-				}
-
 				string UnrealCodeAnalyzerArguments = "";
 				if (BuildConfiguration.bRunUnrealCodeAnalyzer)
 				{
