Added support for programs to specify their own Solution Directory.

Merging using UE4-To-UE4-LauncherDev

diff --git a/Engine/Source/Programs/UnrealBuildTool/System/ProjectFileGenerator.cs b/Engine/Source/Programs/UnrealBuildTool/System/ProjectFileGenerator.cs
--- a/Engine/Source/Programs/UnrealBuildTool/System/ProjectFileGenerator.cs
+++ b/Engine/Source/Programs/UnrealBuildTool/System/ProjectFileGenerator.cs
@@ -42,17 +42,30 @@ namespace UnrealBuildTool
 		/// <returns>The newly-added folder</returns>
 		public MasterProjectFolder AddSubFolder( string SubFolderName )
 		{
-			foreach( var CurFolder in SubFolders )
-			{
-				if( CurFolder.FolderName.Equals( SubFolderName, StringComparison.InvariantCultureIgnoreCase ) )
-				{
-					// Already exists!
-					return CurFolder;
-				}
-			}
-			var NewFolder = OwnerProjectFileGenerator.AllocateMasterProjectFolder( OwnerProjectFileGenerator, SubFolderName );
-			SubFolders.Add( NewFolder );
-			return NewFolder;
+            MasterProjectFolder ResultFolder = null;
+
+            foreach (var FolderName in SubFolderName.Split(new char[1] {'\\'}, StringSplitOptions.RemoveEmptyEntries))
+            {
+                bool AlreadyExists = false;
+                foreach (var ExistingFolder in SubFolders)
+                {
+                    if (ExistingFolder.FolderName.Equals(FolderName, StringComparison.InvariantCultureIgnoreCase))
+                    {
+                        // Already exists!
+                        ResultFolder = ExistingFolder;
+                        AlreadyExists = true;
+                        break;
+                    }
+                }
+
+                if (!AlreadyExists)
+                {
+                    ResultFolder = OwnerProjectFileGenerator.AllocateMasterProjectFolder(OwnerProjectFileGenerator, SubFolderName);
+                    SubFolders.Add(ResultFolder);
+                }
+            }
+
+            return ResultFolder;
 		}
 
 
@@ -475,7 +488,16 @@ namespace UnrealBuildTool
 
 					foreach( var CurProgramProject in ProgramProjects.Values )
 					{
-						RootFolder.AddSubFolder( "Programs" ).ChildProjects.Add( CurProgramProject );
+                        ProjectTarget Target = CurProgramProject.ProjectTargets.FirstOrDefault(t => !String.IsNullOrEmpty(t.TargetRules.SolutionDirectory));
+
+                        if (Target != null)
+                        {
+                            RootFolder.AddSubFolder(Target.TargetRules.SolutionDirectory).ChildProjects.Add(CurProgramProject);
+                        }
+                        else
+                        {
+						    RootFolder.AddSubFolder( "Programs" ).ChildProjects.Add( CurProgramProject );
+                        }
 					}
 
 					// Add all of the config files for generated program targets
diff --git a/Engine/Source/Programs/UnrealBuildTool/System/RulesCompiler.cs b/Engine/Source/Programs/UnrealBuildTool/System/RulesCompiler.cs
--- a/Engine/Source/Programs/UnrealBuildTool/System/RulesCompiler.cs
+++ b/Engine/Source/Programs/UnrealBuildTool/System/RulesCompiler.cs
@@ -536,6 +536,11 @@ namespace UnrealBuildTool
 		private string ConfigurationNameVar = String.Empty;
 
 
+        /// <summary>
+        /// Allows a Program Target to specify it's own solution folder path
+        /// </summary>
+        public string SolutionDirectory = String.Empty;
+
 		/// <summary>
 		/// If true, the built target goes into the Engine/Binaries/<PLATFORM> folder
 		/// </summary>
