Store the generated code directory on the module object in UBT, so we don't need to query global plugin state to figure it out.

diff --git a/Engine/Source/Programs/UnrealBuildTool/Configuration/UEBuildModule.cs b/Engine/Source/Programs/UnrealBuildTool/Configuration/UEBuildModule.cs
--- a/Engine/Source/Programs/UnrealBuildTool/Configuration/UEBuildModule.cs
+++ b/Engine/Source/Programs/UnrealBuildTool/Configuration/UEBuildModule.cs
@@ -1014,6 +1014,9 @@ namespace UnrealBuildTool
 		/** A list of the absolute paths of source files to be built in this module. */
 		public readonly SourceFilesClass SourceFilesToBuild = new SourceFilesClass();
 
+		/** The directory for this module's generated code */
+		public readonly string GeneratedCodeDirectory;
+
 		/** The preprocessor definitions used to compile this module's private implementation. */
 		HashSet<string> Definitions;
 
@@ -1102,6 +1105,7 @@ namespace UnrealBuildTool
 			string InName,
 			UEBuildModuleType InType,
 			string InModuleDirectory,
+			string InGeneratedCodeDirectory,
 			bool? InIsRedistributableOverride,
 			IntelliSenseGatherer InIntelliSenseGatherer,
 			IEnumerable<FileItem> InSourceFiles,
@@ -1162,6 +1166,7 @@ namespace UnrealBuildTool
 					InRuntimeDependencies
 				)
 		{
+			GeneratedCodeDirectory = InGeneratedCodeDirectory;
 			IntelliSenseGatherer = InIntelliSenseGatherer;
 
 			if (bInBuildSourceFiles)
@@ -2083,6 +2088,7 @@ namespace UnrealBuildTool
 			string InName,
 			UEBuildModuleType InType,
 			string InModuleDirectory,
+			string InGeneratedCodeDirectory,
 			bool? InIsRedistributableOverride,
 			IntelliSenseGatherer InIntelliSenseGatherer,
 			IEnumerable<FileItem> InSourceFiles,
@@ -2116,7 +2122,7 @@ namespace UnrealBuildTool
 			bool InEnableExceptions,
 			bool bInBuildSourceFiles
 			)
-			: base(InTarget,InName,InType,InModuleDirectory,InIsRedistributableOverride,InIntelliSenseGatherer,
+			: base(InTarget,InName,InType,InModuleDirectory,InGeneratedCodeDirectory,InIsRedistributableOverride,InIntelliSenseGatherer,
 			InSourceFiles,InPublicIncludePaths,InPublicSystemIncludePaths,null,InDefinitions,
 			InPublicIncludePathModuleNames,InPublicDependencyModuleNames,InPublicDelayLoadDLLs,InPublicAdditionalLibraries,InPublicFrameworks,InPublicWeakFrameworks,InPublicAdditionalFrameworks,InPublicAdditionalShadowFiles,InPublicAdditionalBundleResources,
 			InPrivateIncludePaths,InPrivateIncludePathModuleNames,InPrivateDependencyModuleNames,
diff --git a/Engine/Source/Programs/UnrealBuildTool/Configuration/UEBuildTarget.cs b/Engine/Source/Programs/UnrealBuildTool/Configuration/UEBuildTarget.cs
--- a/Engine/Source/Programs/UnrealBuildTool/Configuration/UEBuildTarget.cs
+++ b/Engine/Source/Programs/UnrealBuildTool/Configuration/UEBuildTarget.cs
@@ -1008,14 +1008,13 @@ namespace UnrealBuildTool
 				// Delete generated header files
 				foreach (var ModuleName in ModuleList)
 				{
-					var Module = GetModuleByName(ModuleName);
-					var ModuleIncludeDir = GetGeneratedCodeDirectoryForModule(ModuleName).Replace("\\", "/");
-					if (!UnrealBuildTool.RunningRocket() || !ModuleIncludeDir.StartsWith(BaseEngineBuildDataFolder, StringComparison.InvariantCultureIgnoreCase))
+					UEBuildModuleCPP Module = GetModuleByName(ModuleName) as UEBuildModuleCPP;
+					if (Module != null && Module.GeneratedCodeDirectory != null && Directory.Exists(Module.GeneratedCodeDirectory))
 					{
-						if (Directory.Exists(ModuleIncludeDir))
+						if(!UnrealBuildTool.IsEngineInstalled() || !Utils.IsFileUnderDirectory(Module.GeneratedCodeDirectory, BuildConfiguration.RelativeEnginePath))
 						{
-							Log.TraceVerbose("\t\tDeleting " + ModuleIncludeDir);
-							CleanDirectory(ModuleIncludeDir);
+							Log.TraceVerbose("\t\tDeleting " + Module.GeneratedCodeDirectory);
+							CleanDirectory(Module.GeneratedCodeDirectory);
 						}
 					}
 				}
@@ -1457,7 +1456,7 @@ namespace UnrealBuildTool
 							// If we've got this far and there are no source files then it's likely we're running Rocket and ignoring
 							// engine files, so we don't need a .generated.cpp either
 							UEBuildModuleCPP.AutoGenerateCppInfoClass.BuildInfoClass BuildInfo = null;
-							UHTModuleInfo.GeneratedCPPFilenameBase = Path.Combine( GetGeneratedCodeDirectoryForModule(UHTModuleInfo.ModuleName), UHTModuleInfo.ModuleName ) + ".generated";
+							UHTModuleInfo.GeneratedCPPFilenameBase = Path.Combine( DependencyModuleCPP.GeneratedCodeDirectory, UHTModuleInfo.ModuleName ) + ".generated";
 							if (DependencyModuleCPP.SourceFilesToBuild.Count != 0)
 							{
 								BuildInfo = new UEBuildModuleCPP.AutoGenerateCppInfoClass.BuildInfoClass( UHTModuleInfo.GeneratedCPPFilenameBase + "*.cpp" );
@@ -1995,6 +1994,7 @@ namespace UnrealBuildTool
 				InName: Name,
 				InType: UEBuildModuleType.Game,
 				InModuleDirectory: Directory,
+				InGeneratedCodeDirectory: null,
 				InIsRedistributableOverride: null,
 				InIntelliSenseGatherer: null,
 				InSourceFiles: SourceFiles.ToList(),
@@ -2322,37 +2322,6 @@ namespace UnrealBuildTool
 			return Binary;
 		}
 
-		/// <summary>
-		/// Determines where generated code files will be stored for this module
-		/// </summary>
-		/// <param name="ModuleName">Name of the module</param>
-		/// <returns></returns>
-		string GetGeneratedCodeDirectoryForModule(string ModuleName)
-		{
-			// Get the plugin info for the module
-			PluginInfo Plugin = Plugins.GetPluginInfoForModule(ModuleName);
-
-			// Get the output directory. Plugins always write to their own directory so they can be copied between projects, shared engine intermediates go in the engine 
-			// intermediate folder, and anything else goes in the project folder.
-			string GeneratedCodeDirectory;
-			if(Plugin != null)
-			{
-				GeneratedCodeDirectory = Plugin.Directory;
-			}
-			else if(bUseSharedBuildEnvironment && Utils.IsFileUnderDirectory(RulesCompiler.GetModuleFilename(ModuleName), BuildConfiguration.RelativeEnginePath))
-			{
-				GeneratedCodeDirectory = BuildConfiguration.RelativeEnginePath;
-			}
-			else
-			{
-				GeneratedCodeDirectory = ProjectDirectory;
-			}
-			GeneratedCodeDirectory = Path.GetFullPath(Path.Combine(GeneratedCodeDirectory, BuildConfiguration.PlatformIntermediateFolder, AppName, "Inc", ModuleName));
-
-			// Return the path with a directory separator
-			return GeneratedCodeDirectory + Path.DirectorySeparatorChar;
-		}
-
 		/**
 		 * @return true if debug information is created, false otherwise.
 		 */
@@ -2892,6 +2861,27 @@ namespace UnrealBuildTool
 					}
 				}
 
+				// Get the generated code directory. Plugins always write to their own intermediate directory so they can be copied between projects, shared engine 
+				// intermediates go in the engine intermediate folder, and anything else goes in the project folder.
+				string GeneratedCodeDirectory = null;
+				if(RulesObject.Type != ModuleRules.ModuleType.External)
+				{
+					PluginInfo Plugin = Plugins.GetPluginInfoForModule(ModuleName);
+					if(Plugin != null)
+					{
+						GeneratedCodeDirectory = Plugin.Directory;
+					}
+					else if(bUseSharedBuildEnvironment && Utils.IsFileUnderDirectory(RulesCompiler.GetModuleFilename(ModuleName), BuildConfiguration.RelativeEnginePath))
+					{
+						GeneratedCodeDirectory = BuildConfiguration.RelativeEnginePath;
+					}
+					else
+					{
+						GeneratedCodeDirectory = ProjectDirectory;
+					}
+					GeneratedCodeDirectory = Path.GetFullPath(Path.Combine(GeneratedCodeDirectory, BuildConfiguration.PlatformIntermediateFolder, AppName, "Inc", ModuleName));
+				}
+
 				// Don't generate include paths for third party modules; they don't follow our conventions. Core is a special-case... leave it alone
 				if (RulesObject.Type != ModuleRules.ModuleType.External && ModuleName != "Core")
 				{
@@ -2899,9 +2889,12 @@ namespace UnrealBuildTool
 					RulesCompiler.AddDefaultIncludePathsToModuleRules(ModuleName, ModuleFileName, ModuleFileRelativeToEngineDirectory, IsGameModule, RulesObject);
 
 					// Add the path to the generated headers 
-					string ModuleGeneratedCodePath = GetGeneratedCodeDirectoryForModule(ModuleName);
-					ModuleGeneratedCodePath = Utils.CleanDirectorySeparators(Utils.MakePathRelativeTo(ModuleGeneratedCodePath, Path.Combine(ProjectFileGenerator.RootRelativePath, "Engine/Source")), '/');
-					RulesObject.PublicIncludePaths.Add(ModuleGeneratedCodePath);
+					if(GeneratedCodeDirectory != null)
+					{
+						string EngineSourceDirectory = Path.Combine(ProjectFileGenerator.RootRelativePath, "Engine/Source");
+						string RelativeGeneratedCodeDirectory = Utils.CleanDirectorySeparators(Utils.MakePathRelativeTo(GeneratedCodeDirectory, EngineSourceDirectory), '/');
+						RulesObject.PublicIncludePaths.Add(RelativeGeneratedCodeDirectory);
+					}
 				}
 
 				// Figure out whether we need to build this module
@@ -2944,7 +2937,7 @@ namespace UnrealBuildTool
 				}
 
 				// Now, go ahead and create the module builder instance
-				Module = InstantiateModule(RulesObject, ModuleName, ModuleType, ModuleDirectory, IntelliSenseGatherer, FoundSourceFiles, bBuildFiles);
+				Module = InstantiateModule(RulesObject, ModuleName, ModuleType, ModuleDirectory, GeneratedCodeDirectory, IntelliSenseGatherer, FoundSourceFiles, bBuildFiles);
 				if(Module == null)
 				{
 					throw new BuildException("Unrecognized module type specified by 'Rules' object {0}", RulesObject.ToString());
@@ -2970,6 +2963,7 @@ namespace UnrealBuildTool
 			string               ModuleName,
 			UEBuildModuleType    ModuleType,
 			string               ModuleDirectory,
+			string				 GeneratedCodeDirectory,
 			IntelliSenseGatherer IntelliSenseGatherer,
 			List<FileItem>       ModuleSourceFiles,
 			bool                 bBuildSourceFiles)
@@ -2982,6 +2976,7 @@ namespace UnrealBuildTool
 							InName: ModuleName,
 							InType: ModuleType,
 							InModuleDirectory: ModuleDirectory,
+							InGeneratedCodeDirectory: GeneratedCodeDirectory,
 							InIsRedistributableOverride: RulesObject.IsRedistributableOverride,
 							InIntelliSenseGatherer: IntelliSenseGatherer,
 							InSourceFiles: ModuleSourceFiles,
@@ -3022,6 +3017,7 @@ namespace UnrealBuildTool
 							InName: ModuleName,
 							InType: ModuleType,
 							InModuleDirectory: ModuleDirectory,
+							InGeneratedCodeDirectory: GeneratedCodeDirectory,
 							InIsRedistributableOverride: RulesObject.IsRedistributableOverride,
 							InIntelliSenseGatherer: IntelliSenseGatherer,
 							InSourceFiles: ModuleSourceFiles,
