- Removed code that shouldn't have been checked in

diff --git a/Engine/Source/Programs/UnrealBuildTool/Configuration/UEBuildTarget.cs b/Engine/Source/Programs/UnrealBuildTool/Configuration/UEBuildTarget.cs
--- a/Engine/Source/Programs/UnrealBuildTool/Configuration/UEBuildTarget.cs
+++ b/Engine/Source/Programs/UnrealBuildTool/Configuration/UEBuildTarget.cs
@@ -779,8 +779,7 @@ namespace UnrealBuildTool
 			ProjectIntermediateDirectory = Path.GetFullPath(Path.Combine(ProjectDirectory, BuildConfiguration.PlatformIntermediateFolder, GetTargetName(), Configuration.ToString()));
 
 			// Build the engine intermediate directory. If we're building agnostic engine binaries, we can use the engine intermediates folder. Otherwise we need to use the project intermediates directory.
-			// If we are building a content project with UE4Game as the base executable, then we don't want to use the Project as the intermediate, it should use UE4's intermedate
-			if (!bUseSharedBuildEnvironment && AppName != "UE4Game")
+			if (!bUseSharedBuildEnvironment)
 			{
 				EngineIntermediateDirectory = ProjectIntermediateDirectory;
 			}
@@ -2275,20 +2274,20 @@ namespace UnrealBuildTool
 				BaseOutputDirectory = Path.GetFullPath(Plugin.Directory);
 			}
 			else if(RulesCompiler.IsGameModule(ModuleName) || !bUseSharedBuildEnvironment)
-			{
+				{
 				BaseOutputDirectory = Path.GetFullPath(ProjectDirectory);
-			}
-			else
-			{
+				}
+				else
+				{
 				BaseOutputDirectory = Path.GetFullPath(BuildConfiguration.RelativeEnginePath);
-			}
+				}
 
 			// Get the configuration that this module will be built in. Engine modules compiled in DebugGame will use Development.
 			UnrealTargetConfiguration ModuleConfiguration = Configuration;
 			if(Configuration == UnrealTargetConfiguration.DebugGame && !RulesCompiler.IsGameModule(ModuleName))
 			{
 				ModuleConfiguration = UnrealTargetConfiguration.Development;
-			}
+				}
 
 			// Get the output and intermediate directories for this module
 			string OutputDirectory = Path.Combine(BaseOutputDirectory, "Binaries", Platform.ToString());
@@ -2297,7 +2296,7 @@ namespace UnrealBuildTool
 			// Append a subdirectory if the module rules specifies one
 			ModuleRules ModuleRules;
 			if(RulesCompiler.TryCreateModuleRules(ModuleName, TargetInfo, out ModuleRules) && !String.IsNullOrEmpty(ModuleRules.BinariesSubFolder))
-			{
+				{
 				OutputDirectory = Path.Combine(OutputDirectory, ModuleRules.BinariesSubFolder);
 				IntermediateDirectory = Path.Combine(IntermediateDirectory, ModuleRules.BinariesSubFolder);
 			}
@@ -2864,7 +2863,7 @@ namespace UnrealBuildTool
 				// Get the generated code directory. Plugins always write to their own intermediate directory so they can be copied between projects, shared engine 
 				// intermediates go in the engine intermediate folder, and anything else goes in the project folder.
 				string GeneratedCodeDirectory = null;
-				if(RulesObject.Type != ModuleRules.ModuleType.External)
+				if (RulesObject.Type != ModuleRules.ModuleType.External)
 				{
 					PluginInfo Plugin = Plugins.GetPluginInfoForModule(ModuleName);
 					if(Plugin != null)
