Remove the IsPluginModule() function; use a more explicit check against files in the engine source directory instead.

diff --git a/Engine/Source/Programs/UnrealBuildTool/System/Plugins.cs b/Engine/Source/Programs/UnrealBuildTool/System/Plugins.cs
--- a/Engine/Source/Programs/UnrealBuildTool/System/Plugins.cs
+++ b/Engine/Source/Programs/UnrealBuildTool/System/Plugins.cs
@@ -167,17 +167,6 @@ namespace UnrealBuildTool
 		}
 
 
-		/// <summary>
-		/// Returns true if the specified module name is part of a plugin
-		/// </summary>
-		/// <param name="ModuleName">Name of the module to check</param>
-		/// <returns>True if this is a plugin module</returns>
-		public static bool IsPluginModule( string ModuleName )
-		{
-			return ModuleNameToPluginMap.ContainsKey( ModuleName );
-		}
-
-		
 		/// <summary>
 		/// Checks to see if this module is a plugin module, and if it is, returns the PluginInfo for that module, otherwise null.
 		/// </summary>
diff --git a/Engine/Source/Programs/UnrealBuildTool/System/RulesCompiler.cs b/Engine/Source/Programs/UnrealBuildTool/System/RulesCompiler.cs
--- a/Engine/Source/Programs/UnrealBuildTool/System/RulesCompiler.cs
+++ b/Engine/Source/Programs/UnrealBuildTool/System/RulesCompiler.cs
@@ -1643,13 +1643,11 @@ namespace UnrealBuildTool
 				}
 			}
 
-			var IsPluginModule = Plugins.IsPluginModule( InModuleName );
-
 			// Setup the directories for Classes, Public, and Intermediate			
 			string ClassesDirectory = Path.Combine(ModuleDirectoryRelativeToEngineSourceDirectory, "Classes/");	// @todo uht: Deprecate eventually.  Or force it to be manually specified...
 			string PublicDirectory = Path.Combine(ModuleDirectoryRelativeToEngineSourceDirectory, "Public/");
 
-			if (IsGameModule || IsPluginModule)
+			if(!Utils.IsFileUnderDirectory(InModuleFilename, AbsEngineSourceDirectory))
 			{
 				// This will be either the format 
 				//		../<Game>/Source/<Module>
