Simplify check for whether we can write to the engine directory in installed builds.

diff --git a/Engine/Source/Programs/UnrealBuildTool/System/ExternalExecution.cs b/Engine/Source/Programs/UnrealBuildTool/System/ExternalExecution.cs
--- a/Engine/Source/Programs/UnrealBuildTool/System/ExternalExecution.cs
+++ b/Engine/Source/Programs/UnrealBuildTool/System/ExternalExecution.cs
@@ -462,26 +462,10 @@ namespace UnrealBuildTool
 				{
 					if (GeneratedCodeDirectoryInfo.Exists)
 					{
-						if (UnrealBuildTool.RunningRocket())
+						// Don't write anything to the engine directory if we're running an installed build
+						if(UnrealBuildTool.IsEngineInstalled() && Utils.IsFileUnderDirectory(Module.ModuleDirectory, BuildConfiguration.RelativeEnginePath))
 						{
-							// If it is an Engine folder and we are building a rocket project do NOT update the timestamp!
-							// @todo Rocket: This contains check is hacky/fragile
-							string FullGeneratedCodeDirectory = GeneratedCodeDirectoryInfo.FullName;
-							FullGeneratedCodeDirectory = FullGeneratedCodeDirectory.Replace("\\", "/");
-							if (FullGeneratedCodeDirectory.Contains("Engine/Intermediate/Build"))
-							{
-								continue;
-							}
-
-							// Skip checking timestamps for engine plugin intermediate headers in Rocket
-							PluginInfo Info = Plugins.GetPluginInfoForModule( Module.ModuleName );
-							if( Info != null )
-							{
-								if( Info.LoadedFrom == PluginInfo.LoadedFromType.Engine )
-								{
-									continue;
-								}
-							}
+							continue;
 						}
 
 						// Touch the include directory since we have technically 'generated' the headers
