Formatting changes.

diff --git a/Engine/Source/Programs/UnrealBuildTool/Configuration/UEBuildTarget.cs b/Engine/Source/Programs/UnrealBuildTool/Configuration/UEBuildTarget.cs
--- a/Engine/Source/Programs/UnrealBuildTool/Configuration/UEBuildTarget.cs
+++ b/Engine/Source/Programs/UnrealBuildTool/Configuration/UEBuildTarget.cs
@@ -2139,7 +2139,7 @@ namespace UnrealBuildTool
 				{
 					string ModuleFileName = RulesCompiler.GetModuleFilename(Module.Name);
 					bool bHasSource = (!String.IsNullOrEmpty(ModuleFileName) && Directory.EnumerateFiles(Path.GetDirectoryName(ModuleFileName), "*.cpp", SearchOption.AllDirectories).Any());
-					AddBinaryForModule(Module.Name, BinaryType, bHasSource, false);
+					AddBinaryForModule(Module.Name, BinaryType, bAllowCompilation: bHasSource, bIsCrossTarget: false);
 				}
 			}
 		}
@@ -2274,13 +2274,13 @@ namespace UnrealBuildTool
 				BaseOutputDirectory = Path.GetFullPath(Plugin.Directory);
 			}
 			else if(RulesCompiler.IsGameModule(ModuleName) || !bUseSharedBuildEnvironment)
-				{
+			{
 				BaseOutputDirectory = Path.GetFullPath(ProjectDirectory);
-				}
-				else
-				{
+			}
+			else
+			{
 				BaseOutputDirectory = Path.GetFullPath(BuildConfiguration.RelativeEnginePath);
-				}
+			}
 
 			// Get the configuration that this module will be built in. Engine modules compiled in DebugGame will use Development.
 			UnrealTargetConfiguration ModuleConfiguration = Configuration;
