Make Mac builds generate .dSYMs by default for Development & Shipping, making sure to add the .dSYM's Info.plist and internal Mach-O binary to the build receipt.

diff --git a/Engine/Source/Programs/UnrealBuildTool/Mac/MacToolChain.cs b/Engine/Source/Programs/UnrealBuildTool/Mac/MacToolChain.cs
--- a/Engine/Source/Programs/UnrealBuildTool/Mac/MacToolChain.cs
+++ b/Engine/Source/Programs/UnrealBuildTool/Mac/MacToolChain.cs
@@ -1097,7 +1097,7 @@ namespace UnrealBuildTool
 			}
 
 			// For Mac, generate the dSYM file if the config file is set to do so
-			if (BuildConfiguration.bGeneratedSYMFile == true && (!bIsBuildingLibrary || LinkEnvironment.Config.bIsBuildingDLL))
+			if ((BuildConfiguration.bGeneratedSYMFile == true || BuildConfiguration.bUsePDBFiles == true) && (!bIsBuildingLibrary || LinkEnvironment.Config.bIsBuildingDLL))
 			{
 				Log.TraceInformation("Generating dSYM file for {0} - this will add some time to your build...", Path.GetFileName(OutputFile.AbsolutePath));
 				RemoteOutputFile = GenerateDebugInfo(OutputFile);
@@ -1326,6 +1326,19 @@ namespace UnrealBuildTool
 
         public override void AddFilesToReceipt(BuildReceipt Receipt, UEBuildBinary Binary)
 		{
+			string DebugExtension = UEBuildPlatform.GetBuildPlatform(Binary.Target.Platform).GetDebugInfoExtension(Binary.Config.Type);
+			if(DebugExtension == ".dsym")
+			{
+				foreach (string OutputFilePath in Binary.Config.OutputFilePaths)
+				{
+					string DsymInfo = OutputFilePath + ".dSYM/Contents/Info.plist";
+					Receipt.AddBuildProduct(DsymInfo, BuildProductType.SymbolFile);
+
+					string DsymDylib = OutputFilePath + ".dSYM/Contents/Resources/DWARF/" + Path.GetFileName(OutputFilePath);
+					Receipt.AddBuildProduct(DsymDylib, BuildProductType.SymbolFile);
+				}
+			}
+
 			if (Binary.Target.GlobalLinkEnvironment.Config.bIsBuildingConsoleApplication)
 			{
 				return;
diff --git a/Engine/Source/Programs/UnrealBuildTool/Mac/UEBuildMac.cs b/Engine/Source/Programs/UnrealBuildTool/Mac/UEBuildMac.cs
--- a/Engine/Source/Programs/UnrealBuildTool/Mac/UEBuildMac.cs
+++ b/Engine/Source/Programs/UnrealBuildTool/Mac/UEBuildMac.cs
@@ -89,7 +89,7 @@ namespace UnrealBuildTool
 		 */
 		public override string GetDebugInfoExtension(UEBuildBinaryType InBinaryType)
 		{
-			return BuildConfiguration.bGeneratedSYMFile ? ".dsym" : "";
+			return BuildConfiguration.bGeneratedSYMFile || BuildConfiguration.bUsePDBFiles ? ".dsym" : "";
 		}
 
 		public override void ModifyNewlyLoadedModule(UEBuildModule InModule, TargetInfo Target)
@@ -173,6 +173,8 @@ namespace UnrealBuildTool
 			BuildConfiguration.ProcessorCountMultiplier = MacToolChain.GetAdjustedProcessorCountMultiplier();
 			BuildConfiguration.bUseSharedPCHs = false;
 
+			BuildConfiguration.bUsePDBFiles = bCreateDebugInfo && Configuration != CPPTargetConfiguration.Debug && Platform == CPPTargetPlatform.Mac;
+
 			// we always deploy - the build machines need to be able to copy the files back, which needs the full bundle
 			BuildConfiguration.bDeployAfterCompile = true;
 		}
