Simplify check for detecting an engine module whether to disable optimizations in DebugGame; just look if it's under the engine directory.

diff --git a/Engine/Source/Programs/UnrealBuildTool/Configuration/UEBuildModule.cs b/Engine/Source/Programs/UnrealBuildTool/Configuration/UEBuildModule.cs
--- a/Engine/Source/Programs/UnrealBuildTool/Configuration/UEBuildModule.cs
+++ b/Engine/Source/Programs/UnrealBuildTool/Configuration/UEBuildModule.cs
@@ -1839,8 +1839,7 @@ namespace UnrealBuildTool
 			// Switch the optimization flag if we're building a game module. Also pass the definition for building in DebugGame along (see ModuleManager.h for notes).
 			if (Target.Configuration == UnrealTargetConfiguration.DebugGame)
 			{
-				PluginInfo Plugin = Plugins.GetPluginInfoForModule(Name);
-				if((Plugin != null && Plugin.LoadedFrom == PluginInfo.LoadedFromType.GameProject) || Type == UEBuildModuleType.Game)
+				if(!Utils.IsFileUnderDirectory(ModuleDirectory, BuildConfiguration.RelativeEnginePath))
 				{
 					Result.Config.Target.Configuration = CPPTargetConfiguration.Debug;
 					Result.Config.Definitions.Add("UE_BUILD_DEVELOPMENT_WITH_DEBUGGAME=1");
