Remove .dSYM bundle path from BuildReceipt as only files are expected in the list, not folders.

diff --git a/Engine/Source/Programs/UnrealBuildTool/Mac/MacToolChain.cs b/Engine/Source/Programs/UnrealBuildTool/Mac/MacToolChain.cs
--- a/Engine/Source/Programs/UnrealBuildTool/Mac/MacToolChain.cs
+++ b/Engine/Source/Programs/UnrealBuildTool/Mac/MacToolChain.cs
@@ -1337,6 +1337,14 @@ namespace UnrealBuildTool
 					string DsymDylib = OutputFilePath + ".dSYM/Contents/Resources/DWARF/" + Path.GetFileName(OutputFilePath);
 					Receipt.AddBuildProduct(DsymDylib, BuildProductType.SymbolFile);
 				}
+
+				for (int i = 0; i < Receipt.BuildProducts.Count; i++)
+				{
+					if(Path.GetExtension(Receipt.BuildProducts[i].Path) == DebugExtension)
+					{
+						Receipt.BuildProducts.RemoveAt(i--);
+					}
+				}
 			}
 
 			if (Binary.Target.GlobalLinkEnvironment.Config.bIsBuildingConsoleApplication)
