Move LinuxNativeDialogs to a more appropriate location.

- Build script corrected.
- Unfortunately, yet another hardcoded rpath.

diff --git a/Engine/Build/BatchFiles/Linux/BuildThirdParty.sh b/Engine/Build/BatchFiles/Linux/BuildThirdParty.sh
--- a/Engine/Build/BatchFiles/Linux/BuildThirdParty.sh
+++ b/Engine/Build/BatchFiles/Linux/BuildThirdParty.sh
@@ -227,7 +227,7 @@ BuildLND()
   make $MAKE_ARGS
   local LIB_DIR=../lib/Linux/$TARGET_ARCH/
   mv *.so $LIB_DIR
-  cp -P --remove-destination $LIB_DIR/*.so ${TOP_DIR}/Binaries/Linux/
+  cp -P --remove-destination $LIB_DIR/*.so ${TOP_DIR}/Binaries/ThirdParty/Linux/$TARGET_ARCH/
   set +x
 }
 
diff --git a/Engine/Source/Programs/UnrealBuildTool/Linux/LinuxToolChain.cs b/Engine/Source/Programs/UnrealBuildTool/Linux/LinuxToolChain.cs
--- a/Engine/Source/Programs/UnrealBuildTool/Linux/LinuxToolChain.cs
+++ b/Engine/Source/Programs/UnrealBuildTool/Linux/LinuxToolChain.cs
@@ -469,6 +469,7 @@ namespace UnrealBuildTool
 			Result += " -Wl,-rpath=${ORIGIN}/..";	// for modules that are in sub-folders of the main Engine/Binary/Linux folder
 			// FIXME: really ugly temp solution. Modules need to be able to specify this
             Result += " -Wl,-rpath=${ORIGIN}/../../../Engine/Binaries/ThirdParty/ICU/icu4c-53_1/Linux/x86_64-unknown-linux-gnu";
+            Result += " -Wl,-rpath=${ORIGIN}/../../../Engine/Binaries/ThirdParty/LinuxNativeDialogs/Linux/x86_64-unknown-linux-gnu";
 
             if (CrossCompiling())
             {
