Remove erroneous .dSYM entries for static libraries, only executables & dylibs have corresponding .dSYMs on OS X.

diff --git a/Engine/Source/Programs/UnrealBuildTool/Mac/MacToolChain.cs b/Engine/Source/Programs/UnrealBuildTool/Mac/MacToolChain.cs
--- a/Engine/Source/Programs/UnrealBuildTool/Mac/MacToolChain.cs
+++ b/Engine/Source/Programs/UnrealBuildTool/Mac/MacToolChain.cs
@@ -1329,13 +1329,17 @@ namespace UnrealBuildTool
 			string DebugExtension = UEBuildPlatform.GetBuildPlatform(Binary.Target.Platform).GetDebugInfoExtension(Binary.Config.Type);
 			if(DebugExtension == ".dsym")
 			{
-				foreach (string OutputFilePath in Binary.Config.OutputFilePaths)
+				for (int i = 0; i < Receipt.BuildProducts.Count; i++)
 				{
-					string DsymInfo = OutputFilePath + ".dSYM/Contents/Info.plist";
-					Receipt.AddBuildProduct(DsymInfo, BuildProductType.SymbolFile);
+					if(Receipt.BuildProducts[i].Type == BuildProductType.Executable || Receipt.BuildProducts[i].Type == BuildProductType.DynamicLibrary)
+					{
+						string OutputFilePath = Receipt.BuildProducts[i].Path;
+						string DsymInfo = OutputFilePath + ".dSYM/Contents/Info.plist";
+						Receipt.AddBuildProduct(DsymInfo, BuildProductType.SymbolFile);
 
-					string DsymDylib = OutputFilePath + ".dSYM/Contents/Resources/DWARF/" + Path.GetFileName(OutputFilePath);
-					Receipt.AddBuildProduct(DsymDylib, BuildProductType.SymbolFile);
+						string DsymDylib = OutputFilePath + ".dSYM/Contents/Resources/DWARF/" + Path.GetFileName(OutputFilePath);
+						Receipt.AddBuildProduct(DsymDylib, BuildProductType.SymbolFile);
+					}
 				}
 
 				for (int i = 0; i < Receipt.BuildProducts.Count; i++)
