Add plugins to projects without needing to reference the AllPlugins array.

diff --git a/Engine/Source/Programs/UnrealBuildTool/System/Plugins.cs b/Engine/Source/Programs/UnrealBuildTool/System/Plugins.cs
--- a/Engine/Source/Programs/UnrealBuildTool/System/Plugins.cs
+++ b/Engine/Source/Programs/UnrealBuildTool/System/Plugins.cs
@@ -62,7 +62,41 @@ namespace UnrealBuildTool
 			return Info;
 		}
 
+		/// <summary>
+		/// Find paths to all the plugins under a given parent directory (recursively)
+		/// </summary>
+		/// <param name="ParentDirectory">Parent directory to look in. Plugins will be found in any *subfolders* of this directory.</param>
+		/// <param name="FileNames">List of filenames. Will have all the discovered .uplugin files appended to it.</param>
+		public static IEnumerable<string> EnumeratePlugins(string ParentDirectory)
+		{
+			List<string> FileNames = new List<string>();
+			if(Directory.Exists(ParentDirectory))
+			{
+				EnumeratePluginsInternal(new DirectoryInfo(ParentDirectory), FileNames);
+			}
+			return FileNames;
+		}
 
+		/// <summary>
+		/// Find paths to all the plugins under a given parent directory (recursively)
+		/// </summary>
+		/// <param name="ParentDirectory">Parent directory to look in. Plugins will be found in any *subfolders* of this directory.</param>
+		/// <param name="FileNames">List of filenames. Will have all the discovered .uplugin files appended to it.</param>
+		static void EnumeratePluginsInternal(DirectoryInfo ParentDirectory, List<string> FileNames)
+		{
+			foreach(DirectoryInfo ChildDirectory in ParentDirectory.EnumerateDirectories())
+			{
+				int InitialFileNamesCount = FileNames.Count;
+				foreach(FileInfo PluginFileInfo in ParentDirectory.EnumerateFiles("*.uplugin", SearchOption.TopDirectoryOnly))
+				{
+					FileNames.Add(PluginFileInfo.FullName);
+				}
+				if(FileNames.Count == InitialFileNamesCount)
+				{
+					EnumeratePluginsInternal(ChildDirectory, FileNames);
+				}
+			}
+		}
 
 		/// <summary>
 		/// Recursively locates all plugins in the specified folder, appending to the incoming list
diff --git a/Engine/Source/Programs/UnrealBuildTool/System/ProjectFileGenerator.cs b/Engine/Source/Programs/UnrealBuildTool/System/ProjectFileGenerator.cs
--- a/Engine/Source/Programs/UnrealBuildTool/System/ProjectFileGenerator.cs
+++ b/Engine/Source/Programs/UnrealBuildTool/System/ProjectFileGenerator.cs
@@ -1528,6 +1528,7 @@ namespace UnrealBuildTool
 		/// <param name="bGatherThirdPartySource">True to gather source code from third party projects too</param>
 		protected void AddProjectsForAllModules( List<UProjectInfo> AllGames, Dictionary<string, ProjectFile> ProgramProjects, List<string> AllModuleFiles, bool bGatherThirdPartySource )
 		{
+			HashSet<ProjectFile> ProjectsWithPlugins = new HashSet<ProjectFile>();
 			foreach( var CurModuleFile in AllModuleFiles )
 			{
 				Log.TraceVerbose("AddProjectsForAllModules " + CurModuleFile);
@@ -1630,35 +1631,27 @@ namespace UnrealBuildTool
 					var FoundFiles = SourceFileSearch.FindModuleSourceFiles( CurModuleFile, ExcludeNoRedistFiles: bExcludeNoRedistFiles, SearchSubdirectories:SearchSubdirectories, IncludePrivateSourceCode:IncludePrivateSourceCode );
 					ProjectFile.AddFilesToProject( FoundFiles, BaseFolder );
 
-
-					// Is this module part of a plugin?  If so then we'll make sure to add other plugin-related fiels
-					var PluginInfo = Plugins.GetPluginInfoForModule( ModuleName );
-					if( PluginInfo != null )
+					// Check if there's a plugin directory here
+					if(!ProjectsWithPlugins.Contains(ProjectFile))
 					{
-						// NOTE: For plugins with multiple modules, we may attempt to re-add the same plugin files here.  That's OK,
-						//  this is handled safely in AddFileToProject()!
-						var UPluginFilePath = Path.Combine( PluginInfo.Directory, Path.GetFileName( PluginInfo.Directory ) + ".uplugin" );
-						if( File.Exists( UPluginFilePath ) )
+						string PluginFolder = Path.Combine(BaseFolder, "Plugins");
+						if(Directory.Exists(PluginFolder))
 						{
-							ProjectFile.AddFileToProject( UPluginFilePath, BaseFolder );
-						}
-						else
-						{
-							throw new BuildException( "Not expecting to find a plugin module with no corresponding .uplugin file.  File '{0}' doesn't exist", UPluginFilePath );
-						}
+							// Add all the plugin files for this project
+							foreach(string PluginFileName in Plugins.EnumeratePlugins(Path.Combine(BaseFolder, "Plugins")))
+							{
+								// Add the .uplugin file
+								ProjectFile.AddFileToProject(PluginFileName, BaseFolder);
 
-						// Add plugin "resource" files if we have any
-						var PluginResourcesFolder = Path.Combine( PluginInfo.Directory, "Resources" );
-						if( Directory.Exists( PluginResourcesFolder ) )
-						{
-							var DirectoriesToSearch = new List<string>();
-							DirectoriesToSearch.Add( PluginResourcesFolder );
-							ProjectFile.AddFilesToProject( 
-								SourceFileSearch.FindFiles(
-									DirectoriesToSearch: DirectoriesToSearch,
-									ExcludeNoRedistFiles: bExcludeNoRedistFiles ), 
-								BaseFolder );
+								// Add plugin "resource" files if we have any
+								string PluginResourcesFolder = Path.Combine(Path.GetDirectoryName(PluginFileName), "Resources");
+								if(Directory.Exists(PluginResourcesFolder))
+								{
+									ProjectFile.AddFilesToProject(SourceFileSearch.FindFiles(DirectoriesToSearch: new List<string>{ PluginResourcesFolder }, ExcludeNoRedistFiles: bExcludeNoRedistFiles ), BaseFolder );
+								}
+							}
 						}
+						ProjectsWithPlugins.Add(ProjectFile);
 					}
 				}
 			}
