Remove another global plugin module name loookup.

diff --git a/Engine/Source/Programs/UnrealBuildTool/Configuration/UEBuildModule.cs b/Engine/Source/Programs/UnrealBuildTool/Configuration/UEBuildModule.cs
--- a/Engine/Source/Programs/UnrealBuildTool/Configuration/UEBuildModule.cs
+++ b/Engine/Source/Programs/UnrealBuildTool/Configuration/UEBuildModule.cs
@@ -58,42 +58,23 @@ namespace UnrealBuildTool
 		public static readonly string EditorFolder = String.Format("{0}Editor{0}", Path.DirectorySeparatorChar);
 		public static readonly string ProgramsFolder = String.Format("{0}Programs{0}", Path.DirectorySeparatorChar);
 
-		public static UEBuildModuleType GetPluginModuleType(string ModuleName)
-		{
-			UEBuildModuleType ModuleType = UEBuildModuleType.Unknown;
-			PluginInfo Info = Plugins.GetPluginInfoForModule(ModuleName);
-			if (null == Info)
-			{
-				return ModuleType;
+		public static UEBuildModuleType GetModuleTypeFromDescriptor(ModuleDescriptor Module)
+		{
+			switch (Module.Type)
+			{
+				case ModuleHostType.Developer:
+					return UEBuildModuleType.Developer;
+				case ModuleHostType.Editor:
+				case ModuleHostType.EditorNoCommandlet:
+					return UEBuildModuleType.Editor;
+				case ModuleHostType.Program:
+					return UEBuildModuleType.Program;
+				case ModuleHostType.Runtime:
+				case ModuleHostType.RuntimeNoCommandlet:
+					return UEBuildModuleType.Runtime;
+				default:
+					throw new BuildException("Unhandled module type {0}", Module.Type.ToString());
 			}
-			foreach (ModuleDescriptor ModuleInfo in Info.Descriptor.Modules)
-			{
-				if (ModuleInfo.Name == ModuleName)
-				{
-					switch (ModuleInfo.Type)
-					{
-						case ModuleHostType.Developer:
-							ModuleType = UEBuildModuleType.Developer;
-							break;
-						case ModuleHostType.Editor:
-						case ModuleHostType.EditorNoCommandlet:
-							ModuleType = UEBuildModuleType.Editor;
-							break;
-						case ModuleHostType.Program:
-							ModuleType = UEBuildModuleType.Program;
-							break;
-						case ModuleHostType.Runtime:
-						case ModuleHostType.RuntimeNoCommandlet:
-							ModuleType = UEBuildModuleType.Runtime;
-							break;
-						default:
-							throw new BuildException("Unhandled plugin module type {0}", ModuleInfo.Type.ToString());
-					}
-					// break out of loop
-					break;
-				}
-			}
-			return ModuleType;
 		}
 
 		public static UEBuildModuleType GetEngineModuleTypeBasedOnLocation(string ModuleName, UEBuildModuleType ModuleType, string ModuleFileRelativeToEngineDirectory)
diff --git a/Engine/Source/Programs/UnrealBuildTool/Configuration/UEBuildTarget.cs b/Engine/Source/Programs/UnrealBuildTool/Configuration/UEBuildTarget.cs
--- a/Engine/Source/Programs/UnrealBuildTool/Configuration/UEBuildTarget.cs
+++ b/Engine/Source/Programs/UnrealBuildTool/Configuration/UEBuildTarget.cs
@@ -2287,7 +2287,7 @@ namespace UnrealBuildTool
 			if(Configuration == UnrealTargetConfiguration.DebugGame && !RulesCompiler.IsGameModule(ModuleName))
 			{
 				ModuleConfiguration = UnrealTargetConfiguration.Development;
-				}
+			}
 
 			// Get the output and intermediate directories for this module
 			string OutputDirectory = Path.Combine(BaseOutputDirectory, "Binaries", Platform.ToString());
@@ -2296,7 +2296,7 @@ namespace UnrealBuildTool
 			// Append a subdirectory if the module rules specifies one
 			ModuleRules ModuleRules;
 			if(RulesCompiler.TryCreateModuleRules(ModuleName, TargetInfo, out ModuleRules) && !String.IsNullOrEmpty(ModuleRules.BinariesSubFolder))
-				{
+			{
 				OutputDirectory = Path.Combine(OutputDirectory, ModuleRules.BinariesSubFolder);
 				IntermediateDirectory = Path.Combine(IntermediateDirectory, ModuleRules.BinariesSubFolder);
 			}
@@ -2811,6 +2811,9 @@ namespace UnrealBuildTool
 				var RulesObject = RulesCompiler.CreateModuleRules(ModuleName, TargetInfo, out ModuleFileName);
 				var ModuleDirectory = Path.GetDirectoryName(ModuleFileName);
 
+				// Get the plugin info
+				PluginInfo Plugin = Plugins.GetPluginInfoForModule(ModuleName);
+
 				// Making an assumption here that any project file path that contains '../../'
 				// is NOT from the engine and therefore must be an application-specific module.
 				// Get the type of module we're creating
@@ -2825,7 +2828,14 @@ namespace UnrealBuildTool
 				else
 				{
 					// Check if it's a plugin
-					ModuleType = UEBuildModule.GetPluginModuleType(ModuleName);
+					if(Plugin != null)
+					{
+						ModuleDescriptor Descriptor = Plugin.Descriptor.Modules.FirstOrDefault(x => x.Name == ModuleName);
+						if(Descriptor != null)
+						{
+							ModuleType = UEBuildModule.GetModuleTypeFromDescriptor(Descriptor);
+						}
+					}
 					if (ModuleType == UEBuildModuleType.Unknown)
 					{
 						// not a plugin, see if it is a game module 
@@ -2865,7 +2875,6 @@ namespace UnrealBuildTool
 				string GeneratedCodeDirectory = null;
 				if (RulesObject.Type != ModuleRules.ModuleType.External)
 				{
-					PluginInfo Plugin = Plugins.GetPluginInfoForModule(ModuleName);
 					if(Plugin != null)
 					{
 						GeneratedCodeDirectory = Plugin.Directory;
