Remove global map for plugin module names. Each target now uses plugin information only when it's initialized.

diff --git a/Engine/Source/Programs/UnrealBuildTool/Configuration/UEBuildTarget.cs b/Engine/Source/Programs/UnrealBuildTool/Configuration/UEBuildTarget.cs
--- a/Engine/Source/Programs/UnrealBuildTool/Configuration/UEBuildTarget.cs
+++ b/Engine/Source/Programs/UnrealBuildTool/Configuration/UEBuildTarget.cs
@@ -2137,9 +2137,13 @@ namespace UnrealBuildTool
 			{
 				if (Module.IsCompiledInConfiguration(Platform, TargetType))
 				{
+					// Create the module, to ensure that it's set up with the right paths for a plugin module
+					FindOrCreateModuleByName(Module.Name, Plugin);
+
+					// Add the corresponding binary for it
 					string ModuleFileName = RulesCompiler.GetModuleFilename(Module.Name);
 					bool bHasSource = (!String.IsNullOrEmpty(ModuleFileName) && Directory.EnumerateFiles(Path.GetDirectoryName(ModuleFileName), "*.cpp", SearchOption.AllDirectories).Any());
-					AddBinaryForModule(Module.Name, BinaryType, bAllowCompilation: bHasSource, bIsCrossTarget: false);
+					AddBinaryForModule(Module.Name, BinaryType, Plugin, bAllowCompilation: bHasSource, bIsCrossTarget: false);
 				}
 			}
 		}
@@ -2230,7 +2234,7 @@ namespace UnrealBuildTool
 					if(PrecompiledModule is UEBuildModuleCPP && !BoundModuleNames.Contains(PrecompiledModule.Name))
 					{
 						UEBuildBinaryType BinaryType = bCompileMonolithic? UEBuildBinaryType.StaticLibrary : UEBuildBinaryType.DynamicLinkLibrary;
-						UEBuildBinary Binary = AddBinaryForModule(PrecompiledModule.Name, BinaryType, bAllowCompilation: !bUsePrecompiled, bIsCrossTarget: false);
+						UEBuildBinary Binary = AddBinaryForModule(PrecompiledModule.Name, BinaryType, Plugin: null, bAllowCompilation: !bUsePrecompiled, bIsCrossTarget: false);
 						PrecompiledBinaries.Add(Binary);
 						BoundModuleNames.Add(PrecompiledModule.Name);
 					}
@@ -2249,7 +2253,7 @@ namespace UnrealBuildTool
 			else
 			{
 				// Otherwise create a new module for it
-				Binary = AddBinaryForModule(ModuleName, UEBuildBinaryType.DynamicLinkLibrary, bAllowCompilation: true, bIsCrossTarget: bIsCrossTarget);
+				Binary = AddBinaryForModule(ModuleName, UEBuildBinaryType.DynamicLinkLibrary, Plugin: null, bAllowCompilation: true, bIsCrossTarget: bIsCrossTarget);
 			}
 			return Binary;
 		}
@@ -2262,11 +2266,8 @@ namespace UnrealBuildTool
 		/// <param name="bAllowCompilation">Whether this binary can be compiled. The function will check whether plugin binaries can be compiled.</param>
 		/// <param name="bIsCrossTarget">True if module is for supporting a different target-platform</param>
 		/// <returns>The new binary</returns>
-		private UEBuildBinaryCPP AddBinaryForModule(string ModuleName, UEBuildBinaryType BinaryType, bool bAllowCompilation, bool bIsCrossTarget)
+		private UEBuildBinaryCPP AddBinaryForModule(string ModuleName, UEBuildBinaryType BinaryType, PluginInfo Plugin, bool bAllowCompilation, bool bIsCrossTarget)
 		{
-			// Get the plugin info for the module
-			PluginInfo Plugin = Plugins.GetPluginInfoForModule(ModuleName);
-
 			// Get the root output directory and base name (target name/app name) for this binary
 			string BaseOutputDirectory;
 			if(Plugin != null)
@@ -2799,6 +2800,12 @@ namespace UnrealBuildTool
 
 		/** Finds a module given its name.  Throws an exception if the module couldn't be found. */
 		public UEBuildModule FindOrCreateModuleByName(string ModuleName)
+		{
+			return FindOrCreateModuleByName(ModuleName, null);
+		}
+
+		/** Finds a module given its name.  Throws an exception if the module couldn't be found. */
+		private UEBuildModule FindOrCreateModuleByName(string ModuleName, PluginInfo Plugin)
 		{
 			UEBuildModule Module;
 			if (!Modules.TryGetValue(ModuleName, out Module))
@@ -2811,9 +2818,6 @@ namespace UnrealBuildTool
 				var RulesObject = RulesCompiler.CreateModuleRules(ModuleName, TargetInfo, out ModuleFileName);
 				var ModuleDirectory = Path.GetDirectoryName(ModuleFileName);
 
-				// Get the plugin info
-				PluginInfo Plugin = Plugins.GetPluginInfoForModule(ModuleName);
-
 				// Making an assumption here that any project file path that contains '../../'
 				// is NOT from the engine and therefore must be an application-specific module.
 				// Get the type of module we're creating
diff --git a/Engine/Source/Programs/UnrealBuildTool/System/Plugins.cs b/Engine/Source/Programs/UnrealBuildTool/System/Plugins.cs
--- a/Engine/Source/Programs/UnrealBuildTool/System/Plugins.cs
+++ b/Engine/Source/Programs/UnrealBuildTool/System/Plugins.cs
@@ -181,45 +181,9 @@ namespace UnrealBuildTool
 						Plugins.FindPluginsIn(GamePluginsDirectory, PluginInfo.LoadedFromType.GameProject, ref AllPluginsVar);
 					}
 				}
-
-				// Also keep track of which modules map to which plugins
-				ModulesToPluginMapVar = new Dictionary<string,PluginInfo>( StringComparer.InvariantCultureIgnoreCase );
-				foreach( var CurPluginInfo in AllPlugins )
-				{
-					foreach( var Module in CurPluginInfo.Descriptor.Modules )
-					{
-						// Make sure a different plugin doesn't already have a module with this name
-						// @todo plugin: Collisions like this could happen because of third party plugins added to a project, which isn't really ideal.
-						if( ModuleNameToPluginMap.ContainsKey( Module.Name ) )
-						{
-							throw new BuildException( "Found a plugin in '{0}' which describes a module '{1}', but a module with this name already exists in plugin '{2}'!", CurPluginInfo.Directory, Module.Name, ModuleNameToPluginMap[ Module.Name ].Directory );
-						}
-						ModulesToPluginMapVar.Add( Module.Name, CurPluginInfo );						
-					}
-				}
-			}
-		}
-
-
-		/// <summary>
-		/// Checks to see if this module is a plugin module, and if it is, returns the PluginInfo for that module, otherwise null.
-		/// </summary>
-		/// <param name="ModuleName">Name of the module to check</param>
-		/// <returns>The PluginInfo for this module, if the module is a plugin module.  Otherwise, returns null</returns>
-		public static PluginInfo GetPluginInfoForModule( string ModuleName )
-		{
-			PluginInfo FoundPluginInfo;
-			if( ModuleNameToPluginMap.TryGetValue( ModuleName, out FoundPluginInfo ) )
-			{
-				return FoundPluginInfo;
-			}
-			else
-			{
-				return null;
 			}
 		}
 
-	
 		/// Access the list of all plugins.  We'll scan for plugins when this is called the first time.
 		public static List< PluginInfo > AllPlugins
 		{
@@ -230,22 +194,7 @@ namespace UnrealBuildTool
 			}
 		}
 
-		/// Access a mapping of modules to their respective owning plugin.  Dictionary is case-insensitive.
-		private static Dictionary< string, PluginInfo > ModuleNameToPluginMap
-		{
-			get
-			{
-				DiscoverAllPlugins();
-				return ModulesToPluginMapVar;
-			}
-		}
-
-	
-
 		/// List of all plugins we've found so far in this session
 		private static List< PluginInfo > AllPluginsVar = null;
-
-		/// Maps plugin modules to the plugin that owns them
-		private static Dictionary< string, PluginInfo > ModulesToPluginMapVar = null;
 	}
 }
