Enable HTML5 for Mac Rocket

diff --git a/Engine/Source/Programs/AutomationTool/Scripts/RocketBuild.Automation.cs b/Engine/Source/Programs/AutomationTool/Scripts/RocketBuild.Automation.cs
--- a/Engine/Source/Programs/AutomationTool/Scripts/RocketBuild.Automation.cs
+++ b/Engine/Source/Programs/AutomationTool/Scripts/RocketBuild.Automation.cs
@@ -158,7 +158,7 @@ namespace Rocket
 				{
 					TargetPlatforms.Add(UnrealTargetPlatform.Linux);
 				}
-				if(HostPlatform == UnrealTargetPlatform.Win64)
+				if(HostPlatform == UnrealTargetPlatform.Win64 || HostPlatform == UnrealTargetPlatform.Mac )
 				{
 					TargetPlatforms.Add(UnrealTargetPlatform.HTML5);
 				}
@@ -217,6 +217,10 @@ namespace Rocket
 
 		public static UnrealTargetPlatform GetSourceHostPlatform(GUBP bp, UnrealTargetPlatform HostPlatform, UnrealTargetPlatform TargetPlatform)
 		{
+            if (TargetPlatform == UnrealTargetPlatform.HTML5 && HostPlatform == UnrealTargetPlatform.Mac && bp.HostPlatforms.Contains(UnrealTargetPlatform.Win64))
+            {
+                return UnrealTargetPlatform.Win64;
+            }
 			if (TargetPlatform == UnrealTargetPlatform.Android && HostPlatform == UnrealTargetPlatform.Mac && bp.HostPlatforms.Contains(UnrealTargetPlatform.Win64))
 			{
 				return UnrealTargetPlatform.Win64;
@@ -406,7 +410,7 @@ namespace Rocket
 			AddDependency(GUBP.RootEditorNode.StaticGetFullName(HostPlatform));
 			AddDependency(GUBP.ToolsNode.StaticGetFullName(HostPlatform));
 
-			// Get all the external file lists. Android on Mac uses the same files as Win64.
+			// Get all the external file lists. Android/HTML5 on Mac uses the same files as Win64.
 			foreach(UnrealTargetPlatform CodeTargetPlatform in CodeTargetPlatforms)
 			{
 				UnrealTargetPlatform SourceHostPlatform = RocketBuild.GetSourceHostPlatform(bp, HostPlatform, CodeTargetPlatform);
diff --git a/Engine/Source/Programs/UnrealBuildTool/HTML5/UEBuildHTML5.cs b/Engine/Source/Programs/UnrealBuildTool/HTML5/UEBuildHTML5.cs
--- a/Engine/Source/Programs/UnrealBuildTool/HTML5/UEBuildHTML5.cs
+++ b/Engine/Source/Programs/UnrealBuildTool/HTML5/UEBuildHTML5.cs
@@ -14,9 +14,6 @@ namespace UnrealBuildTool
         [XmlConfig]
         public static string HTML5Architecture = "";
 
-        // Enable and link in code for serving HTTP request via NFS. 
-        [XmlConfig]
-        public static bool EnableHTTPForNFS = true;
 
         static private string TargetPlatformName = "HTML5";
 
@@ -345,22 +342,6 @@ namespace UnrealBuildTool
                     InModule.AddPublicDependencyModule("UEOgg");
                     InModule.AddPublicDependencyModule("Vorbis");
                 }
-
-                if (InModule.ToString() == "NetworkFile")
-                {
-                    if (EnableHTTPForNFS)
-                    {   
-                        InModule.AddPublicDefinition("ENABLE_HTTP_FOR_NFS=1");
-                        if (Target.Architecture == "-win32")
-                        {
-                            InModule.AddPrivateDependencyModule("HTML5Win32");
-                        }
-                        else
-                        {
-                            InModule.AddPrivateDependencyModule("HTML5JS");
-                        }
-                    }
-                }
             }
             else if (Target.Platform == UnrealTargetPlatform.Win32 || Target.Platform == UnrealTargetPlatform.Win64 || Target.Platform == UnrealTargetPlatform.Mac )
             {
@@ -372,26 +353,6 @@ namespace UnrealBuildTool
                 {
                     InModule.AddPlatformSpecificDynamicallyLoadedModule("HTML5TargetPlatform");
                 }
-
-                if (EnableHTTPForNFS && (Target.Platform == UnrealTargetPlatform.Win64  
-                                         || Target.Platform == UnrealTargetPlatform.Mac)
-                    )
-                {
-                    if (InModule.ToString() == "NetworkFile") // client
-                    {
-                        InModule.AddPrivateDependencyModule("HTTP");
-                        InModule.AddPublicDefinition("ENABLE_HTTP_FOR_NFS=1");
-                    }
-                    else if (InModule.ToString() == "NetworkFileSystem") // server 
-                    {
-                        if (UnrealBuildTool.RunningRocket() == false || 
-                            Target.Type == TargetRules.TargetType.Game )
-                        {
-                            InModule.AddPrivateDependencyModule("WebSockets");
-                            InModule.AddPublicDefinition("ENABLE_HTTP_FOR_NFS=1");
-                        }
-                    }
-                }
             }
         }
 
diff --git a/Engine/Source/Runtime/NetworkFile/NetworkFile.Build.cs b/Engine/Source/Runtime/NetworkFile/NetworkFile.Build.cs
--- a/Engine/Source/Runtime/NetworkFile/NetworkFile.Build.cs
+++ b/Engine/Source/Runtime/NetworkFile/NetworkFile.Build.cs
@@ -21,5 +21,18 @@ public class NetworkFile : ModuleRules
 				}
 				);
 		}
+
+        if (Target.Platform == UnrealTargetPlatform.HTML5)
+        { 
+            Definitions.Add("ENABLE_HTTP_FOR_NFS=1");
+            if (Target.Architecture == "-win32")
+            {
+                PrivateDependencyModuleNames.Add("HTML5Win32");
+            }
+            else
+            {
+                PrivateDependencyModuleNames.Add("HTML5JS");
+            }
+        }
 	}
 }
\ No newline at end of file
diff --git a/Engine/Source/Runtime/NetworkFileSystem/NetworkFileSystem.Build.cs b/Engine/Source/Runtime/NetworkFileSystem/NetworkFileSystem.Build.cs
--- a/Engine/Source/Runtime/NetworkFileSystem/NetworkFileSystem.Build.cs
+++ b/Engine/Source/Runtime/NetworkFileSystem/NetworkFileSystem.Build.cs
@@ -6,39 +6,45 @@ namespace UnrealBuildTool.Rules
 	{
 		public NetworkFileSystem(TargetInfo Target)
 		{
-			PrivateIncludePaths.AddRange(
-				new string[] {
-					"Runtime/NetworkFileSystem/Private",
-					"Runtime/NetworkFileSystem/Private/Simple",
-					"Runtime/NetworkFileSystem/Private/Streaming",
-				}
-				);
+            PrivateIncludePaths.AddRange(
+	            new string[] {
+		            "Runtime/NetworkFileSystem/Private",
+		            "Runtime/NetworkFileSystem/Private/Simple",
+		            "Runtime/NetworkFileSystem/Private/Streaming",
+	            }
+	            );
 
-		PrivateDependencyModuleNames.AddRange(
-			new string[]
-			{
-                "CoreUObject",
-				"Projects",
-				"SandboxFile",
-				"TargetPlatform",
-			}
-		);
+            PrivateDependencyModuleNames.AddRange(
+	            new string[]
+	            {
+                    "CoreUObject",
+		            "Projects",
+		            "SandboxFile",
+		            "TargetPlatform",
+	            }
+            );
 
-		PublicIncludePaths.AddRange(
-			new string[] {
-				"Runtime/NetworkFileSystem/Public",
-				"Runtime/NetworkFileSystem/Public/Interfaces",
-                "Runtime/CoreUObject/Public/Interfaces",
-			}
-		);
+            PublicIncludePaths.AddRange(
+	            new string[] {
+		            "Runtime/NetworkFileSystem/Public",
+		            "Runtime/NetworkFileSystem/Public/Interfaces",
+                    "Runtime/CoreUObject/Public/Interfaces",
+	            }
+            );
 
-		PublicDependencyModuleNames.AddRange(
-			new string[]
-				{
-					"Core",
-					"Sockets",
-				}
-				);
-		}
+            PublicDependencyModuleNames.AddRange(
+	            new string[]
+		            {
+			            "Core",
+			            "Sockets",
+		            }
+		            );
+
+            if (Target.Platform == UnrealTargetPlatform.Win64 || Target.Platform == UnrealTargetPlatform.Mac)
+            {
+                AddThirdPartyPrivateStaticDependencies(Target,"WebSockets");
+                Definitions.Add("ENABLE_HTTP_FOR_NFS=1");
+            }
+        }
 	}
 }
