Only build enabled plugins in modular builds. It wastes time to building things locally that aren't enabled. Specifying the -precompile option still builds everything, which should be enough for artist builds and Rocket.

diff --git a/Engine/Source/Programs/UnrealBuildTool/Configuration/UEBuildTarget.cs b/Engine/Source/Programs/UnrealBuildTool/Configuration/UEBuildTarget.cs
--- a/Engine/Source/Programs/UnrealBuildTool/Configuration/UEBuildTarget.cs
+++ b/Engine/Source/Programs/UnrealBuildTool/Configuration/UEBuildTarget.cs
@@ -2487,13 +2487,9 @@ namespace UnrealBuildTool
 			{
 				BuildPlugins.AddRange(ValidPlugins);
 			}
-			else if (ShouldCompileMonolithic() || TargetType == TargetRules.TargetType.Program)
-			{
-				BuildPlugins.AddRange(EnabledPlugins);
-			}
 			else
 			{
-				BuildPlugins.AddRange(ValidPlugins);
+				BuildPlugins.AddRange(EnabledPlugins);
 			}
 		}
 
