Add function to enumerate plugins available for a specific project rather than using the global array, and add a function to determine if a plugin is enabled.

diff --git a/Engine/Source/Programs/AutomationTool/ProjectUtils.cs b/Engine/Source/Programs/AutomationTool/ProjectUtils.cs
--- a/Engine/Source/Programs/AutomationTool/ProjectUtils.cs
+++ b/Engine/Source/Programs/AutomationTool/ProjectUtils.cs
@@ -186,9 +186,14 @@ namespace AutomationTool
 			{
 				EngineSourceDirectory = Path.Combine(UnrealBuildTool.Utils.GetExecutingAssemblyDirectory(), "..", "..", "..", "Engine", "Binaries");
 			}
+			Directory.SetCurrentDirectory(EngineSourceDirectory);
+
+			// Read the project descriptor, and find all the plugins available to this project
+			ProjectDescriptor Project = ProjectDescriptor.FromFile(RawProjectPath);
+			List<PluginInfo> AllPlugins = Plugins.ReadAvailablePlugins(RawProjectPath);
 
-			bool RetVal = false;
 			// check the target platforms for any differences in build settings or additional plugins
+			bool RetVal = false;
 			foreach (UnrealTargetPlatform TargetPlatformType in TargetPlatforms)
 			{
 				IUEBuildPlatform BuildPlat = UEBuildPlatform.GetBuildPlatform(TargetPlatformType, true);
@@ -198,40 +203,22 @@ namespace AutomationTool
 					break;
 				}
 
-
-				// find if there are any plugins
-				List<string> PluginList = new List<string>();
-				Directory.SetCurrentDirectory(EngineSourceDirectory);
-				// Use the project settings to update the plugin list for this target
-				PluginList = UProjectInfo.GetEnabledPlugins(RawProjectPath, PluginList, TargetPlatformType);
-				if (PluginList.Count > 0 && !RetVal)
+				// find if there are any plugins enabled or disabled which differ from the default
+				foreach(PluginInfo Plugin in AllPlugins)
 				{
-					foreach (var PluginName in PluginList)
+					if(UProjectInfo.IsPluginEnabledForProject(Plugin, Project, TargetPlatformType) != Plugin.Descriptor.bEnabledByDefault)
 					{
-						// check the plugin info for this plugin itself
-						foreach (var Plugin in Plugins.AllPlugins)
-						{
-							if (Plugin.Name == PluginName && !Plugin.Descriptor.bEnabledByDefault && !RetVal)
-							{
-								foreach (var Module in Plugin.Descriptor.Modules)
-								{
-									if(Module.IsCompiledInConfiguration(TargetPlatformType, TargetRules.TargetType.Game))
-									{
-										RetVal = true;
-										break;
-									}
-								}
-								break;
-							}
-						}
-						if (RetVal)
+						if(Plugin.Descriptor.Modules.Any(Module => Module.IsCompiledInConfiguration(TargetPlatformType, TargetRules.TargetType.Game)))
 						{
+							RetVal = true;
 							break;
 						}
 					}
 				}
-				Directory.SetCurrentDirectory(oldCWD);
 			}
+
+			// Change back to the original directory
+			Directory.SetCurrentDirectory(oldCWD);
 			return RetVal;
 		}
 
diff --git a/Engine/Source/Programs/UnrealBuildTool/System/Plugins.cs b/Engine/Source/Programs/UnrealBuildTool/System/Plugins.cs
--- a/Engine/Source/Programs/UnrealBuildTool/System/Plugins.cs
+++ b/Engine/Source/Programs/UnrealBuildTool/System/Plugins.cs
@@ -51,17 +51,45 @@ namespace UnrealBuildTool
 		/// <param name="PluginFile">The path to the plugin file to load</param>
 		/// <param name="LoadedFrom">Where the plugin was loaded from</param>
 		/// <returns>New PluginInfo for the loaded descriptor.</returns>
-		private static PluginInfo LoadPluginDescriptor( FileInfo PluginFileInfo, PluginInfo.LoadedFromType LoadedFrom )
+		public static PluginInfo ReadPlugin(string PluginFileName, PluginInfo.LoadedFromType LoadedFrom)
 		{
 			// Create the plugin info object
 			PluginInfo Info = new PluginInfo();
 			Info.LoadedFrom = LoadedFrom;
-			Info.Directory = PluginFileInfo.Directory.FullName;
+			Info.Directory = Path.GetDirectoryName(Path.GetFullPath(PluginFileName));
 			Info.Name = Path.GetFileName(Info.Directory);
-			Info.Descriptor = PluginDescriptor.FromFile(PluginFileInfo.FullName);
+			Info.Descriptor = PluginDescriptor.FromFile(PluginFileName);
 			return Info;
 		}
 
+		/// <summary>
+		/// Read all the plugins available to a given project
+		/// </summary>
+		/// <param name="ProjectFileName">Path to the project file</param>
+		/// <returns>Sequence of PluginInfo objects, one for each discovered plugin</returns>
+		public static List<PluginInfo> ReadAvailablePlugins(string ProjectFileName)
+		{
+			List<PluginInfo> Plugins = new List<PluginInfo>();
+
+			// Read all the engine plugins
+			string EnginePluginsDir = Path.Combine(BuildConfiguration.RelativeEnginePath, "Plugins");
+			foreach(string PluginFileName in EnumeratePlugins(EnginePluginsDir))
+			{
+				PluginInfo Plugin = ReadPlugin(PluginFileName, PluginInfo.LoadedFromType.Engine);
+				Plugins.Add(Plugin);
+			}
+
+			// Read all the project plugins
+			string ProjectPluginsDir = Path.Combine(Path.GetDirectoryName(ProjectFileName), "Plugins");
+			foreach(string PluginFileName in EnumeratePlugins(ProjectPluginsDir))
+			{
+				PluginInfo Plugin = ReadPlugin(PluginFileName, PluginInfo.LoadedFromType.GameProject);
+				Plugins.Add(Plugin);
+			}
+
+			return Plugins;
+		}
+
 		/// <summary>
 		/// Find paths to all the plugins under a given parent directory (recursively)
 		/// </summary>
@@ -87,7 +115,7 @@ namespace UnrealBuildTool
 			foreach(DirectoryInfo ChildDirectory in ParentDirectory.EnumerateDirectories())
 			{
 				int InitialFileNamesCount = FileNames.Count;
-				foreach(FileInfo PluginFileInfo in ParentDirectory.EnumerateFiles("*.uplugin", SearchOption.TopDirectoryOnly))
+				foreach(FileInfo PluginFileInfo in ChildDirectory.EnumerateFiles("*.uplugin", SearchOption.TopDirectoryOnly))
 				{
 					FileNames.Add(PluginFileInfo.FullName);
 				}
@@ -124,7 +152,7 @@ namespace UnrealBuildTool
 					{
 						// Load the plugin info and keep track of it
 						var PluginDescriptorFile = new FileInfo(PluginDescriptorFileName);
-						var PluginInfo = LoadPluginDescriptor(PluginDescriptorFile, LoadedFrom);
+						var PluginInfo = ReadPlugin(PluginDescriptorFile.FullName, LoadedFrom);
 
 						Plugins.Add(PluginInfo);
 						bFoundPlugin = true;
diff --git a/Engine/Source/Programs/UnrealBuildTool/System/ProjectDescriptor.cs b/Engine/Source/Programs/UnrealBuildTool/System/ProjectDescriptor.cs
--- a/Engine/Source/Programs/UnrealBuildTool/System/ProjectDescriptor.cs
+++ b/Engine/Source/Programs/UnrealBuildTool/System/ProjectDescriptor.cs
@@ -18,7 +18,7 @@ namespace UnrealBuildTool
 		Latest = LatestPlusOne - 1
 	}
 
-	class ProjectDescriptor
+	public class ProjectDescriptor
 	{
 		// Descriptor version number. */
 		public int FileVersion;
diff --git a/Engine/Source/Programs/UnrealBuildTool/System/UProjectInfo.cs b/Engine/Source/Programs/UnrealBuildTool/System/UProjectInfo.cs
--- a/Engine/Source/Programs/UnrealBuildTool/System/UProjectInfo.cs
+++ b/Engine/Source/Programs/UnrealBuildTool/System/UProjectInfo.cs
@@ -327,5 +327,28 @@ namespace UnrealBuildTool
 			}
 			return EnabledPlugins;
 		}
+
+		/// <summary>
+		/// Determine if a plugin is enabled for a given project
+		/// </summary>
+		/// <param name="Project">The project to check</param>
+		/// <param name="Plugin">Information about the plugin</param>
+		/// <param name="Platform">The target platform</param>
+		/// <returns>True if the plugin should be enabled for this project</returns>
+		public static bool IsPluginEnabledForProject(PluginInfo Plugin, ProjectDescriptor Project, UnrealTargetPlatform Platform)
+		{
+			bool bEnabled = Plugin.Descriptor.bEnabledByDefault || Plugin.LoadedFrom == PluginInfo.LoadedFromType.GameProject;
+			if(Project.Plugins != null)
+			{
+				foreach(PluginReferenceDescriptor PluginReference in Project.Plugins)
+				{
+					if(String.Compare(PluginReference.Name, Plugin.Name, true) == 0)
+					{
+						bEnabled = PluginReference.IsEnabledForPlatform(Platform);
+					}
+				}
+			}
+			return bEnabled;
+		}
 	}
 }
