Temporary fix for UdpMessaging being referenced from targets that don't have the plugin enabled.

diff --git a/Engine/Source/Programs/UnrealBuildTool/Configuration/UEBuildTarget.cs b/Engine/Source/Programs/UnrealBuildTool/Configuration/UEBuildTarget.cs
--- a/Engine/Source/Programs/UnrealBuildTool/Configuration/UEBuildTarget.cs
+++ b/Engine/Source/Programs/UnrealBuildTool/Configuration/UEBuildTarget.cs
@@ -2846,6 +2846,10 @@ namespace UnrealBuildTool
 						else
 						{
 							ModuleType = UEBuildModule.GetEngineModuleTypeBasedOnLocation(ModuleName, ModuleType, ModuleFileRelativeToEngineDirectory);
+							if (ModuleName == "UdpMessaging")
+							{
+								ModuleType = UEBuildModuleType.Runtime;
+							}
 							if (ModuleType == UEBuildModuleType.Unknown)
 							{
 								throw new BuildException("Unable to determine module type for {0}", ModuleFileName);
