Better user control of AndroidManifest.xml generation
#ue4
#android

diff --git a/Engine/Source/Programs/UnrealBuildTool/Android/UEDeployAndroid.cs b/Engine/Source/Programs/UnrealBuildTool/Android/UEDeployAndroid.cs
--- a/Engine/Source/Programs/UnrealBuildTool/Android/UEDeployAndroid.cs
+++ b/Engine/Source/Programs/UnrealBuildTool/Android/UEDeployAndroid.cs
@@ -619,7 +619,7 @@ namespace UnrealBuildTool.Android
 			}
 		}
 
-		private string GenerateManifest(string ProjectName, bool bIsForDistribution, bool bPackageDataInsideApk)
+		private string GenerateManifest(string ProjectName, bool bIsForDistribution, bool bPackageDataInsideApk, string GameBuildFilesPath)
 		{
 			// ini file to get settings from
 			ConfigCacheIni Ini = new ConfigCacheIni(UnrealTargetPlatform.Android, "Engine", UnrealBuildTool.GetUProjectPath());
@@ -698,6 +698,14 @@ namespace UnrealBuildTool.Android
 					Text.AppendLine("\t\t\t" + Line);
 				}
 			}
+			string ActivityAdditionsFile = Path.Combine(GameBuildFilesPath, "ManifestActivityAdditions.txt");
+			if (File.Exists(ActivityAdditionsFile))
+			{
+				foreach (string Line in File.ReadAllLines(ActivityAdditionsFile))
+				{
+					Text.AppendLine("\t\t\t" + Line);
+				}
+			}
 			Text.AppendLine("\t\t</activity>");
 			Text.AppendLine(string.Format("\t\t<meta-data android:name=\"com.epicgames.ue4.GameActivity.DepthBufferPreference\" android:value=\"{0}\"/>", ConvertDepthBufferIniValue(DepthBufferPreference)));
 			Text.AppendLine(string.Format("\t\t<meta-data android:name=\"com.epicgames.ue4.GameActivity.bPackageDataInsideApk\" android:value=\"{0}\"/>", bPackageDataInsideApk ? "true" : "false"));
@@ -721,38 +729,67 @@ namespace UnrealBuildTool.Android
 					Text.AppendLine("\t\t" + Line);
 				}
 			}
+			string ApplicationAdditionsFile = Path.Combine(GameBuildFilesPath, "ManifestApplicationAdditions.txt");
+			if (File.Exists(ApplicationAdditionsFile))
+			{
+				foreach (string Line in File.ReadAllLines(ApplicationAdditionsFile))
+				{
+					Text.AppendLine("\t\t" + Line);
+				}
+			}
 			Text.AppendLine("\t</application>");
 
 			Text.AppendLine("");
 			Text.AppendLine("\t<!-- Requirements -->");
-			// need just the number part of the sdk
-			Text.AppendLine(string.Format("\t<uses-sdk android:minSdkVersion=\"{0}\"/>", MinSDKVersion));
-			Text.AppendLine("\t<uses-feature android:glEsVersion=\"0x00020000\" android:required=\"true\" />");
-			Text.AppendLine("\t<uses-permission android:name=\"android.permission.INTERNET\"/>");
-			Text.AppendLine("\t<uses-permission android:name=\"android.permission.WRITE_EXTERNAL_STORAGE\"/>");
-			Text.AppendLine("\t<uses-permission android:name=\"android.permission.ACCESS_NETWORK_STATE\"/>");
-			Text.AppendLine("\t<uses-permission android:name=\"android.permission.WAKE_LOCK\"/>");
-			Text.AppendLine("\t<uses-permission android:name=\"android.permission.READ_PHONE_STATE\"/>");
-			Text.AppendLine("\t<uses-permission android:name=\"com.android.vending.CHECK_LICENSE\"/>");
-			Text.AppendLine("\t<uses-permission android:name=\"android.permission.ACCESS_WIFI_STATE\"/>");
-			Text.AppendLine("\t<uses-permission android:name=\"android.permission.MODIFY_AUDIO_SETTINGS\"/>");
-			Text.AppendLine("\t<uses-permission android:name=\"android.permission.GET_ACCOUNTS\"/>");
-            Text.AppendLine("\t<uses-permission android:name=\"android.permission.VIBRATE\"/>");
-//			Text.AppendLine("\t<uses-permission android:name=\"android.permission.DISABLE_KEYGUARD\"/>");
-
-			if (bEnableIAP)
-			{
-				Text.AppendLine("\t<uses-permission android:name=\"com.android.vending.BILLING\"/>");
-			}
-			if (bPackageForGearVR)
+
+			// check for an override for the requirements section of the manifest
+			string RequirementsOverrideFile = Path.Combine(GameBuildFilesPath, "ManifestRequirementsOverride.txt");
+			if (File.Exists(RequirementsOverrideFile))
 			{
-				Text.AppendLine("\t<uses-feature android:name=\"android.hardware.usb.host\"/>");
+				foreach (string Line in File.ReadAllLines(RequirementsOverrideFile))
+				{
+					Text.AppendLine("\t" + Line);
+				}
 			}
-			if (ExtraPermissions != null)
+			else
 			{
-				foreach (string Permission in ExtraPermissions)
+				// need just the number part of the sdk
+				Text.AppendLine(string.Format("\t<uses-sdk android:minSdkVersion=\"{0}\"/>", MinSDKVersion));
+				Text.AppendLine("\t<uses-feature android:glEsVersion=\"0x00020000\" android:required=\"true\" />");
+				Text.AppendLine("\t<uses-permission android:name=\"android.permission.INTERNET\"/>");
+				Text.AppendLine("\t<uses-permission android:name=\"android.permission.WRITE_EXTERNAL_STORAGE\"/>");
+				Text.AppendLine("\t<uses-permission android:name=\"android.permission.ACCESS_NETWORK_STATE\"/>");
+				Text.AppendLine("\t<uses-permission android:name=\"android.permission.WAKE_LOCK\"/>");
+				Text.AppendLine("\t<uses-permission android:name=\"android.permission.READ_PHONE_STATE\"/>");
+				Text.AppendLine("\t<uses-permission android:name=\"com.android.vending.CHECK_LICENSE\"/>");
+				Text.AppendLine("\t<uses-permission android:name=\"android.permission.ACCESS_WIFI_STATE\"/>");
+				Text.AppendLine("\t<uses-permission android:name=\"android.permission.MODIFY_AUDIO_SETTINGS\"/>");
+				Text.AppendLine("\t<uses-permission android:name=\"android.permission.GET_ACCOUNTS\"/>");
+				Text.AppendLine("\t<uses-permission android:name=\"android.permission.VIBRATE\"/>");
+				//			Text.AppendLine("\t<uses-permission android:name=\"android.permission.DISABLE_KEYGUARD\"/>");
+
+				if (bEnableIAP)
+				{
+					Text.AppendLine("\t<uses-permission android:name=\"com.android.vending.BILLING\"/>");
+				}
+				if (bPackageForGearVR)
+				{
+					Text.AppendLine("\t<uses-feature android:name=\"android.hardware.usb.host\"/>");
+				}
+				if (ExtraPermissions != null)
 				{
-					Text.AppendLine(string.Format("\t<uses-permission android:name=\"{0}\"/>", Permission));
+					foreach (string Permission in ExtraPermissions)
+					{
+						Text.AppendLine(string.Format("\t<uses-permission android:name=\"{0}\"/>", Permission));
+					}
+				}
+				string RequirementsAdditionsFile = Path.Combine(GameBuildFilesPath, "ManifestRequirementsAdditions.txt");
+				if (File.Exists(RequirementsAdditionsFile))
+				{
+					foreach (string Line in File.ReadAllLines(RequirementsAdditionsFile))
+					{
+						Text.AppendLine("\t" + Line);
+					}
 				}
 			}
 			Text.AppendLine("</manifest>");
@@ -787,7 +824,7 @@ namespace UnrealBuildTool.Android
 			bool bBuildSettingsMatch = false;
 
 			string ManifestFile = Path.Combine(UE4BuildPath, "AndroidManifest.xml");
-			string NewManifest = GenerateManifest(ProjectName, bForDistribution, bPackageDataInsideApk);
+			string NewManifest = GenerateManifest(ProjectName, bForDistribution, bPackageDataInsideApk, GameBuildFilesPath);
 			string OldManifest = File.Exists(ManifestFile) ? File.ReadAllText(ManifestFile) : "";
 			if (NewManifest == OldManifest) 
 			{
diff --git a/Engine/Source/Runtime/Android/AndroidRuntimeSettings/Classes/AndroidRuntimeSettings.h b/Engine/Source/Runtime/Android/AndroidRuntimeSettings/Classes/AndroidRuntimeSettings.h
--- a/Engine/Source/Runtime/Android/AndroidRuntimeSettings/Classes/AndroidRuntimeSettings.h
+++ b/Engine/Source/Runtime/Android/AndroidRuntimeSettings/Classes/AndroidRuntimeSettings.h
@@ -132,15 +132,16 @@ public:
 	UPROPERTY(GlobalConfig, EditAnywhere, Category = APKPackaging, Meta = (DisplayName = "Preferred Depth Buffer format"))
 	TEnumAsByte<EAndroidDepthBufferPreference::Type> DepthBufferPreference;
 
-	// Any extra settings for the <application> section
+	// Any extra settings for the <application> section (an optional file <Project>/Build/Android/ManifestApplicationAdditions.txt will also be included)
 	UPROPERTY(GlobalConfig, EditAnywhere, Category = AdvancedAPKPackaging, Meta = (DisplayName = "Extra Settings for <application> section (\\n to separate lines)"))
 	FString ExtraApplicationSettings;
 
-	// Any extra settings for the main <activity> section
+	// Any extra settings for the main <activity> section (an optional file <Project>/Build/Android/ManifestApplicationActivtyAdditions.txt will also be included)
 	UPROPERTY(GlobalConfig, EditAnywhere, Category = AdvancedAPKPackaging, Meta = (DisplayName = "Extra Settings for <activity> section (\\n to separate lines)"))
 	FString ExtraActivitySettings;
 
-	// Any extra permissions your app needs
+	// Any extra permissions your app needs (an optional file <Project>/Build/Android/ManifestRequirementsAdditions.txt will also be included,
+	// or an optional file <Project>/Build/Android/ManifestRequirementsOverride.txt will replace the entire <!-- Requirements --> section)
 	UPROPERTY(GlobalConfig, EditAnywhere, Category = AdvancedAPKPackaging, Meta = (DisplayName = "Extra Permissions (e.g. 'android.permission.INTERNET')"))
 	TArray<FString> ExtraPermissions;
 
