Hack to allow portal to link against the MessagingDebugger module, even though it's marked as editor-only, and the Portal doesn't even use that plugin.

diff --git a/Engine/Source/Programs/UnrealBuildTool/Configuration/UEBuildTarget.cs b/Engine/Source/Programs/UnrealBuildTool/Configuration/UEBuildTarget.cs
--- a/Engine/Source/Programs/UnrealBuildTool/Configuration/UEBuildTarget.cs
+++ b/Engine/Source/Programs/UnrealBuildTool/Configuration/UEBuildTarget.cs
@@ -2846,7 +2846,7 @@ namespace UnrealBuildTool
 						else
 						{
 							ModuleType = UEBuildModule.GetEngineModuleTypeBasedOnLocation(ModuleName, ModuleType, ModuleFileRelativeToEngineDirectory);
-							if (ModuleName == "UdpMessaging")
+							if (ModuleName == "UdpMessaging" || ModuleName == "MessagingDebugger")
 							{
 								ModuleType = UEBuildModuleType.Runtime;
 							}
