Fix list of enabled plugins in programs that reference plugin modules directly. Previously relying on UBT permitting references to any known modules, even if they weren't supposed to be included in the target.

diff --git a/Engine/Source/Programs/IOS/UnrealLaunchDaemon/UnrealLaunchDaemon.Target.cs b/Engine/Source/Programs/IOS/UnrealLaunchDaemon/UnrealLaunchDaemon.Target.cs
--- a/Engine/Source/Programs/IOS/UnrealLaunchDaemon/UnrealLaunchDaemon.Target.cs
+++ b/Engine/Source/Programs/IOS/UnrealLaunchDaemon/UnrealLaunchDaemon.Target.cs
@@ -11,6 +11,7 @@ public class UnrealLaunchDaemonTarget : TargetRules
 		bUsesSlate = false;
 		//PlatformType = TargetRules.TargetPlatformType.Mobile;
 		//bRequiresUnrealHeaderGeneration = true;
+		AdditionalPlugins.Add("UdpMessaging");
 	}
 
 	//
diff --git a/Engine/Source/Programs/UnrealBuildTool/Configuration/UEBuildTarget.cs b/Engine/Source/Programs/UnrealBuildTool/Configuration/UEBuildTarget.cs
--- a/Engine/Source/Programs/UnrealBuildTool/Configuration/UEBuildTarget.cs
+++ b/Engine/Source/Programs/UnrealBuildTool/Configuration/UEBuildTarget.cs
@@ -2137,13 +2137,10 @@ namespace UnrealBuildTool
 			{
 				if (Module.IsCompiledInConfiguration(Platform, TargetType))
 				{
-					// Create the module, to ensure that it's set up with the right paths for a plugin module
-					FindOrCreateModuleByName(Module.Name, Plugin);
-
 					// Add the corresponding binary for it
 					string ModuleFileName = RulesCompiler.GetModuleFilename(Module.Name);
 					bool bHasSource = (!String.IsNullOrEmpty(ModuleFileName) && Directory.EnumerateFiles(Path.GetDirectoryName(ModuleFileName), "*.cpp", SearchOption.AllDirectories).Any());
-					AddBinaryForModule(Module.Name, BinaryType, Plugin, bAllowCompilation: bHasSource, bIsCrossTarget: false);
+					AddBinaryForModule(Module.Name, BinaryType, bAllowCompilation: bHasSource, bIsCrossTarget: false);
 				}
 			}
 		}
@@ -2234,7 +2231,7 @@ namespace UnrealBuildTool
 					if(PrecompiledModule is UEBuildModuleCPP && !BoundModuleNames.Contains(PrecompiledModule.Name))
 					{
 						UEBuildBinaryType BinaryType = bCompileMonolithic? UEBuildBinaryType.StaticLibrary : UEBuildBinaryType.DynamicLinkLibrary;
-						UEBuildBinary Binary = AddBinaryForModule(PrecompiledModule.Name, BinaryType, Plugin: null, bAllowCompilation: !bUsePrecompiled, bIsCrossTarget: false);
+						UEBuildBinary Binary = AddBinaryForModule(PrecompiledModule.Name, BinaryType, bAllowCompilation: !bUsePrecompiled, bIsCrossTarget: false);
 						PrecompiledBinaries.Add(Binary);
 						BoundModuleNames.Add(PrecompiledModule.Name);
 					}
@@ -2253,7 +2250,7 @@ namespace UnrealBuildTool
 			else
 			{
 				// Otherwise create a new module for it
-				Binary = AddBinaryForModule(ModuleName, UEBuildBinaryType.DynamicLinkLibrary, Plugin: null, bAllowCompilation: true, bIsCrossTarget: bIsCrossTarget);
+				Binary = AddBinaryForModule(ModuleName, UEBuildBinaryType.DynamicLinkLibrary, bAllowCompilation: true, bIsCrossTarget: bIsCrossTarget);
 			}
 			return Binary;
 		}
@@ -2266,8 +2263,11 @@ namespace UnrealBuildTool
 		/// <param name="bAllowCompilation">Whether this binary can be compiled. The function will check whether plugin binaries can be compiled.</param>
 		/// <param name="bIsCrossTarget">True if module is for supporting a different target-platform</param>
 		/// <returns>The new binary</returns>
-		private UEBuildBinaryCPP AddBinaryForModule(string ModuleName, UEBuildBinaryType BinaryType, PluginInfo Plugin, bool bAllowCompilation, bool bIsCrossTarget)
+		private UEBuildBinaryCPP AddBinaryForModule(string ModuleName, UEBuildBinaryType BinaryType, bool bAllowCompilation, bool bIsCrossTarget)
 		{
+			// Get the plugin info for this module
+			PluginInfo Plugin = FindPluginForModule(ModuleName);
+
 			// Get the root output directory and base name (target name/app name) for this binary
 			string BaseOutputDirectory;
 			if(Plugin != null)
@@ -2322,6 +2322,16 @@ namespace UnrealBuildTool
 			return Binary;
 		}
 
+		/// <summary>
+		/// Find the plugin which contains a given module
+		/// </summary>
+		/// <param name="ModuleName">Name of the module</param>
+		/// <returns>Matching plugin, or null if not found</returns>
+		private PluginInfo FindPluginForModule(string ModuleName)
+		{
+			return BuildPlugins.FirstOrDefault(BuildPlugin => BuildPlugin.Descriptor.Modules.Any(Module => Module.Name == ModuleName));
+		}
+
 		/**
 		 * @return true if debug information is created, false otherwise.
 		 */
@@ -2796,12 +2806,6 @@ namespace UnrealBuildTool
 
 		/** Finds a module given its name.  Throws an exception if the module couldn't be found. */
 		public UEBuildModule FindOrCreateModuleByName(string ModuleName)
-		{
-			return FindOrCreateModuleByName(ModuleName, null);
-		}
-
-		/** Finds a module given its name.  Throws an exception if the module couldn't be found. */
-		private UEBuildModule FindOrCreateModuleByName(string ModuleName, PluginInfo Plugin)
 		{
 			UEBuildModule Module;
 			if (!Modules.TryGetValue(ModuleName, out Module))
@@ -2814,6 +2818,9 @@ namespace UnrealBuildTool
 				var RulesObject = RulesCompiler.CreateModuleRules(ModuleName, TargetInfo, out ModuleFileName);
 				var ModuleDirectory = Path.GetDirectoryName(ModuleFileName);
 
+				// Get the plugin for this module
+				PluginInfo Plugin = FindPluginForModule(ModuleName);
+
 				// Making an assumption here that any project file path that contains '../../'
 				// is NOT from the engine and therefore must be an application-specific module.
 				// Get the type of module we're creating
@@ -2846,10 +2853,6 @@ namespace UnrealBuildTool
 						else
 						{
 							ModuleType = UEBuildModule.GetEngineModuleTypeBasedOnLocation(ModuleName, ModuleType, ModuleFileRelativeToEngineDirectory);
-							if (ModuleName == "UdpMessaging" || ModuleName == "MessagingDebugger")
-							{
-								ModuleType = UEBuildModuleType.Runtime;
-							}
 							if (ModuleType == UEBuildModuleType.Unknown)
 							{
 								throw new BuildException("Unable to determine module type for {0}", ModuleFileName);
diff --git a/Engine/Source/Programs/UnrealBuildTool/System/RulesCompiler.cs b/Engine/Source/Programs/UnrealBuildTool/System/RulesCompiler.cs
--- a/Engine/Source/Programs/UnrealBuildTool/System/RulesCompiler.cs
+++ b/Engine/Source/Programs/UnrealBuildTool/System/RulesCompiler.cs
@@ -462,8 +462,9 @@ namespace UnrealBuildTool
 		public UnrealTargetConfiguration UndecoratedConfiguration = UnrealTargetConfiguration.Development;
 
 		/// <summary>
-		/// A list of additional plugins which need to be built for this target. Game and editor targets can use the EnabledPlugins 
-		/// setting in their config files to control this.
+		/// A list of additional plugins which need to be included in this target. This allows referencing non-optional plugin modules
+		/// which cannot be disabled, and allows building against specific modules in program targets which do not fit the categories
+		/// in ModuleHostType.
 		/// </summary>
 		public List<string> AdditionalPlugins = new List<string>();		
 
