Setup the list of enabled plugins before setting up the platform build environment, allowing PS4 to directly check whether the target includes Morpheus.

diff --git a/Engine/Source/Programs/UnrealBuildTool/Configuration/UEBuildTarget.cs b/Engine/Source/Programs/UnrealBuildTool/Configuration/UEBuildTarget.cs
--- a/Engine/Source/Programs/UnrealBuildTool/Configuration/UEBuildTarget.cs
+++ b/Engine/Source/Programs/UnrealBuildTool/Configuration/UEBuildTarget.cs
@@ -1579,6 +1579,9 @@ namespace UnrealBuildTool
 		/// <returns>Special Rocket lib files that are build products.</returns>
 		public void PreBuildSetup(IUEToolChain TargetToolChain)
 		{
+			// Setup the target's plugins
+			SetupPlugins();
+
 			// Set up the global compile and link environment in GlobalCompileEnvironment and GlobalLinkEnvironment.
 			SetupGlobalEnvironment();
 
@@ -1588,9 +1591,6 @@ namespace UnrealBuildTool
 			// Setup the target's binaries.
 			SetupBinaries();
 
-			// Setup the target's plugins
-			SetupPlugins();
-
 			// Add the enabled plugins to the build
 			foreach (PluginInfo BuildPlugin in BuildPlugins)
 			{
@@ -2438,10 +2438,11 @@ namespace UnrealBuildTool
 			ExtraModuleNames.AddRange(PlatformExtraModules);			
 		}
 
-		public void FindValidPlugins(out List<PluginInfo> ValidPlugins, out List<string> EnabledPluginNames)
+		/** Sets up the plugins for this target */
+		protected virtual void SetupPlugins()
 		{
 			// Filter the plugins list by the current project
-			ValidPlugins = new List<PluginInfo>();
+			List<PluginInfo> ValidPlugins = new List<PluginInfo>();
 			foreach (PluginInfo Plugin in Plugins.AllPlugins)
 			{
 				if (Plugin.LoadedFrom != PluginInfo.LoadedFromType.GameProject || Plugin.Directory.StartsWith(ProjectDirectory, StringComparison.InvariantCultureIgnoreCase))
@@ -2461,7 +2462,7 @@ namespace UnrealBuildTool
 			}
 
 			// Build a list of enabled plugins
-			EnabledPluginNames = new List<string>(Rules.AdditionalPlugins);
+			List<string> EnabledPluginNames = new List<string>(Rules.AdditionalPlugins);
 
 			// Add the list of plugins enabled by default
 			if (UEBuildConfiguration.bCompileAgainstEngine)
@@ -2478,15 +2479,6 @@ namespace UnrealBuildTool
 				// Use the project settings to update the plugin list for this target
 				EnabledPluginNames = UProjectInfo.GetEnabledPlugins(UnrealBuildTool.GetUProjectFile(), EnabledPluginNames, Platform);
 			}
-		}
-		
-
-		/** Sets up the plugins for this target */
-		protected virtual void SetupPlugins()
-		{
-			List<PluginInfo> ValidPlugins;
-			List<string> EnabledPluginNames;
-			FindValidPlugins(out ValidPlugins, out EnabledPluginNames);
 
 			// Set the list of plugins we're dependent on
 			PluginInfo[] EnabledPlugins = ValidPlugins.Where(x => EnabledPluginNames.Contains(x.Name)).Distinct().ToArray();
