Remove another use of the AllPlugins array in UEBuildTarget, and use a single code path for determining whether a plugin is enabled for a given platform.

diff --git a/Engine/Source/Programs/UnrealBuildTool/Configuration/UEBuildTarget.cs b/Engine/Source/Programs/UnrealBuildTool/Configuration/UEBuildTarget.cs
--- a/Engine/Source/Programs/UnrealBuildTool/Configuration/UEBuildTarget.cs
+++ b/Engine/Source/Programs/UnrealBuildTool/Configuration/UEBuildTarget.cs
@@ -677,11 +677,11 @@ namespace UnrealBuildTool
 
 		/** All plugins which are built for this target */
 		[NonSerialized]
-		public List<PluginInfo> BuildPlugins = new List<PluginInfo>();
+		public List<PluginInfo> BuildPlugins;
 
 		/** All plugin dependencies for this target. This differs from the list of plugins that is built for Rocket, where we build everything, but link in only the enabled plugins. */
 		[NonSerialized]
-		public List<PluginInfo> DependentPlugins = new List<PluginInfo>();
+		public List<PluginInfo> EnabledPlugins;
 
 		/** All application binaries; may include binaries not built by this target. */
 		[NonSerialized]
@@ -1848,7 +1848,7 @@ namespace UnrealBuildTool
 					}
 				}
 			}
-			foreach (PluginInfo Plugin in DependentPlugins)
+			foreach (PluginInfo Plugin in EnabledPlugins)
 			{
 				foreach (ModuleDescriptor Module in Plugin.Descriptor.Modules)
 				{
@@ -2442,14 +2442,7 @@ namespace UnrealBuildTool
 		protected virtual void SetupPlugins()
 		{
 			// Filter the plugins list by the current project
-			List<PluginInfo> ValidPlugins = new List<PluginInfo>();
-			foreach (PluginInfo Plugin in Plugins.AllPlugins)
-			{
-				if (Plugin.LoadedFrom != PluginInfo.LoadedFromType.GameProject || Plugin.Directory.StartsWith(ProjectDirectory, StringComparison.InvariantCultureIgnoreCase))
-				{
-					ValidPlugins.Add(Plugin);
-				}
-			}
+			List<PluginInfo> ValidPlugins = Plugins.ReadAvailablePlugins(UnrealBuildTool.GetUProjectFile());
 
 			// Remove any plugins for platforms we don't have
 			foreach (UnrealTargetPlatform TargetPlatform in Enum.GetValues(typeof(UnrealTargetPlatform)))
@@ -2462,36 +2455,43 @@ namespace UnrealBuildTool
 			}
 
 			// Build a list of enabled plugins
-			List<string> EnabledPluginNames = new List<string>(Rules.AdditionalPlugins);
+			EnabledPlugins = new List<PluginInfo>();
 
-			// Add the list of plugins enabled by default
-			if (UEBuildConfiguration.bCompileAgainstEngine)
+			// If we're compiling against the engine, add the plugins enabled for this target
+			if(UEBuildConfiguration.bCompileAgainstEngine)
 			{
-				EnabledPluginNames.AddRange(ValidPlugins.Where(x => x.Descriptor.bEnabledByDefault).Select(x => x.Name));
+				ProjectDescriptor Project = UnrealBuildTool.HasUProjectFile()? ProjectDescriptor.FromFile(UnrealBuildTool.GetUProjectFile()) : null;
+				foreach(PluginInfo ValidPlugin in ValidPlugins)
+				{
+					if(UProjectInfo.IsPluginEnabledForProject(ValidPlugin, Project, Platform))
+					{
+						EnabledPlugins.Add(ValidPlugin);
+					}
+				}
 			}
 
-			// Update the plugin list for game targets
-			if (TargetType != TargetRules.TargetType.Program && UnrealBuildTool.HasUProjectFile())
+			// Add the plugins explicitly required by the target rules
+			foreach(string AdditionalPlugin in Rules.AdditionalPlugins)
 			{
-				// Enable all the game specific plugins by default
-				EnabledPluginNames.AddRange(ValidPlugins.Where(x => x.LoadedFrom == PluginInfo.LoadedFromType.GameProject).Select(x => x.Name));
-
-				// Use the project settings to update the plugin list for this target
-				EnabledPluginNames = UProjectInfo.GetEnabledPlugins(UnrealBuildTool.GetUProjectFile(), EnabledPluginNames, Platform);
+				PluginInfo Plugin = ValidPlugins.FirstOrDefault(ValidPlugin => ValidPlugin.Name == AdditionalPlugin);
+				if(Plugin == null)
+				{
+					throw new BuildException("Plugin '{0}' is in the list of additional plugins for {1}, but was not found.", AdditionalPlugin, TargetName);
+				}
+				if(!EnabledPlugins.Contains(Plugin))
+				{
+					EnabledPlugins.Add(Plugin);
+				}
 			}
 
-			// Set the list of plugins we're dependent on
-			PluginInfo[] EnabledPlugins = ValidPlugins.Where(x => EnabledPluginNames.Contains(x.Name)).Distinct().ToArray();
-			DependentPlugins.AddRange(EnabledPlugins);
-
 			// Set the list of plugins that should be built
 			if (bPrecompile && TargetType != TargetRules.TargetType.Program)
 			{
-				BuildPlugins.AddRange(ValidPlugins);
+				BuildPlugins = new List<PluginInfo>(ValidPlugins);
 			}
 			else
 			{
-				BuildPlugins.AddRange(EnabledPlugins);
+				BuildPlugins = new List<PluginInfo>(EnabledPlugins);
 			}
 		}
 
diff --git a/Engine/Source/Programs/UnrealBuildTool/System/Plugins.cs b/Engine/Source/Programs/UnrealBuildTool/System/Plugins.cs
--- a/Engine/Source/Programs/UnrealBuildTool/System/Plugins.cs
+++ b/Engine/Source/Programs/UnrealBuildTool/System/Plugins.cs
@@ -65,7 +65,7 @@ namespace UnrealBuildTool
 		/// <summary>
 		/// Read all the plugins available to a given project
 		/// </summary>
-		/// <param name="ProjectFileName">Path to the project file</param>
+		/// <param name="ProjectFileName">Path to the project file (or null)</param>
 		/// <returns>Sequence of PluginInfo objects, one for each discovered plugin</returns>
 		public static List<PluginInfo> ReadAvailablePlugins(string ProjectFileName)
 		{
@@ -80,11 +80,14 @@ namespace UnrealBuildTool
 			}
 
 			// Read all the project plugins
-			string ProjectPluginsDir = Path.Combine(Path.GetDirectoryName(ProjectFileName), "Plugins");
-			foreach(string PluginFileName in EnumeratePlugins(ProjectPluginsDir))
+			if(!String.IsNullOrEmpty(ProjectFileName))
 			{
-				PluginInfo Plugin = ReadPlugin(PluginFileName, PluginInfo.LoadedFromType.GameProject);
-				Plugins.Add(Plugin);
+				string ProjectPluginsDir = Path.Combine(Path.GetDirectoryName(ProjectFileName), "Plugins");
+				foreach(string PluginFileName in EnumeratePlugins(ProjectPluginsDir))
+				{
+					PluginInfo Plugin = ReadPlugin(PluginFileName, PluginInfo.LoadedFromType.GameProject);
+					Plugins.Add(Plugin);
+				}
 			}
 
 			return Plugins;
diff --git a/Engine/Source/Programs/UnrealBuildTool/System/UProjectInfo.cs b/Engine/Source/Programs/UnrealBuildTool/System/UProjectInfo.cs
--- a/Engine/Source/Programs/UnrealBuildTool/System/UProjectInfo.cs
+++ b/Engine/Source/Programs/UnrealBuildTool/System/UProjectInfo.cs
@@ -299,35 +299,6 @@ namespace UnrealBuildTool
 			return "";
 		}
 
-		/// <summary>
-		/// Get the project folder for the given project name
-		/// </summary>
-		/// <param name="InProjectName">Name of the project of interest</param>
-		/// <returns>The project filename, empty string if not found</returns>
-		public static List<string> GetEnabledPlugins(string ProjectFileName, List<string> DefaultEnabledPlugins, UnrealTargetPlatform Platform)
-		{
-			// Parse the project file
-			ProjectDescriptor Project = ProjectDescriptor.FromFile(ProjectFileName);
-
-			// Get the list of plugins
-			List<string> EnabledPlugins = new List<string>(DefaultEnabledPlugins);
-			if(Project.Plugins != null)
-			{
-				foreach(PluginReferenceDescriptor Plugin in Project.Plugins)
-				{
-					if(Plugin.IsEnabledForPlatform(Platform))
-					{
-						EnabledPlugins.Add(Plugin.Name);
-					}
-					else
-					{
-						EnabledPlugins.RemoveAll(x => x.Equals(Plugin.Name, StringComparison.CurrentCultureIgnoreCase));
-					}
-				}
-			}
-			return EnabledPlugins;
-		}
-
 		/// <summary>
 		/// Determine if a plugin is enabled for a given project
 		/// </summary>
@@ -338,7 +309,7 @@ namespace UnrealBuildTool
 		public static bool IsPluginEnabledForProject(PluginInfo Plugin, ProjectDescriptor Project, UnrealTargetPlatform Platform)
 		{
 			bool bEnabled = Plugin.Descriptor.bEnabledByDefault || Plugin.LoadedFrom == PluginInfo.LoadedFromType.GameProject;
-			if(Project.Plugins != null)
+			if(Project != null && Project.Plugins != null)
 			{
 				foreach(PluginReferenceDescriptor PluginReference in Project.Plugins)
 				{
