Remove the global plugins array. Targets can figure out which plugins to enable for themselves, and that list isn't cross-pollenated with other targets.

diff --git a/Engine/Source/Programs/UnrealBuildTool/System/Plugins.cs b/Engine/Source/Programs/UnrealBuildTool/System/Plugins.cs
--- a/Engine/Source/Programs/UnrealBuildTool/System/Plugins.cs
+++ b/Engine/Source/Programs/UnrealBuildTool/System/Plugins.cs
@@ -39,12 +39,6 @@ namespace UnrealBuildTool
 
 	public class Plugins
 	{
-		/// File extension of plugin descriptor files.  NOTE: This constant exists in UnrealBuildTool code as well.
-		/// NOTE: This constant exists in PluginManager C++ code as well.
-		private static string PluginDescriptorFileExtension = ".uplugin";
-
-
-
 		/// <summary>
 		/// Loads a plugin descriptor file and fills out a new PluginInfo structure.  Throws an exception on failure.
 		/// </summary>
@@ -128,104 +122,5 @@ namespace UnrealBuildTool
 				}
 			}
 		}
-
-		/// <summary>
-		/// Recursively locates all plugins in the specified folder, appending to the incoming list
-		/// </summary>
-		/// <param name="PluginsDirectory">Directory to search</param>
-		/// <param name="LoadedFrom">Where we're loading these plugins from</param>
-		/// <param name="Plugins">List of plugins found so far</param>
-		private static void FindPluginsRecursively(string PluginsDirectory, PluginInfo.LoadedFromType LoadedFrom, ref List<PluginInfo> Plugins)
-		{
-			// NOTE: The logic in this function generally matches that of the C++ code for FindPluginsRecursively
-			//       in the core engine code.  These routines should be kept in sync.
-
-			// Each sub-directory is possibly a plugin.  If we find that it contains a plugin, we won't recurse any
-			// further -- you can't have plugins within plugins.  If we didn't find a plugin, we'll keep recursing.
-
-			var PluginsDirectoryInfo = new DirectoryInfo(PluginsDirectory);
-			foreach( var PossiblePluginDirectory in PluginsDirectoryInfo.EnumerateDirectories() )
-			{
-				// Do we have a plugin descriptor in this directory?
-				bool bFoundPlugin = false;
-				foreach (var PluginDescriptorFileName in Directory.GetFiles(PossiblePluginDirectory.FullName, "*" + PluginDescriptorFileExtension))
-				{
-					// Found a plugin directory!  No need to recurse any further, but make sure it's unique.
-					if (!Plugins.Any(x => x.Directory == PossiblePluginDirectory.FullName))
-					{
-						// Load the plugin info and keep track of it
-						var PluginDescriptorFile = new FileInfo(PluginDescriptorFileName);
-						var PluginInfo = ReadPlugin(PluginDescriptorFile.FullName, LoadedFrom);
-
-						Plugins.Add(PluginInfo);
-						bFoundPlugin = true;
-						Log.TraceVerbose("Found plugin in: " + PluginInfo.Directory);
-					}
-
-					// No need to search for more plugins
-					break;
-				}
-
-				if (!bFoundPlugin)
-				{
-					// Didn't find a plugin in this directory.  Continue to look in subfolders.
-					FindPluginsRecursively(PossiblePluginDirectory.FullName, LoadedFrom, ref Plugins);
-				}
-			}
-		}
-
-
-		/// <summary>
-		/// Finds all plugins in the specified base directory
-		/// </summary>
-		/// <param name="PluginsDirectory">Base directory to search.  All subdirectories will be searched, except directories within other plugins.</param>
-		/// <param name="LoadedFrom">Where we're loading these plugins from</param>
-		/// <param name="Plugins">List of all of the plugins we found</param>
-		public static void FindPluginsIn(string PluginsDirectory, PluginInfo.LoadedFromType LoadedFrom, ref List<PluginInfo> Plugins)
-		{
-			if (Directory.Exists(PluginsDirectory))
-			{
-				FindPluginsRecursively(PluginsDirectory, LoadedFrom, ref Plugins);
-			}
-		}
-
-
-		/// <summary>
-		/// Discovers all plugins
-		/// </summary>
-		private static void DiscoverAllPlugins()
-		{
-			if( AllPluginsVar == null )	// Only do this search once per session
-			{
-				AllPluginsVar = new List< PluginInfo >();
-
-				// Engine plugins
-				var EnginePluginsDirectory = Path.Combine( ProjectFileGenerator.EngineRelativePath, "Plugins" );
-				Plugins.FindPluginsIn( EnginePluginsDirectory, PluginInfo.LoadedFromType.Engine, ref AllPluginsVar );
-
-				// Game plugins
-				if (RulesCompiler.AllGameFolders != null)
-				{
-					foreach (var GameProjectFolder in RulesCompiler.AllGameFolders)
-					{
-						var GamePluginsDirectory = Path.Combine(GameProjectFolder, "Plugins");
-						Plugins.FindPluginsIn(GamePluginsDirectory, PluginInfo.LoadedFromType.GameProject, ref AllPluginsVar);
-					}
-				}
-			}
-		}
-
-		/// Access the list of all plugins.  We'll scan for plugins when this is called the first time.
-		public static List< PluginInfo > AllPlugins
-		{
-			get
-			{
-				DiscoverAllPlugins();
-				return AllPluginsVar;
-			}
-		}
-
-		/// List of all plugins we've found so far in this session
-		private static List< PluginInfo > AllPluginsVar = null;
 	}
 }
diff --git a/Engine/Source/Programs/UnrealBuildTool/System/RulesCompiler.cs b/Engine/Source/Programs/UnrealBuildTool/System/RulesCompiler.cs
--- a/Engine/Source/Programs/UnrealBuildTool/System/RulesCompiler.cs
+++ b/Engine/Source/Programs/UnrealBuildTool/System/RulesCompiler.cs
@@ -1160,16 +1160,19 @@ namespace UnrealBuildTool
 
 			// @todo plugin: Disallow modules from including plugin modules as dependency modules? (except when the module is part of that plugin)
 
-			// And plugin directories
+			// Get all the root folders for plugins
+			List<string> RootFolders = new List<string>();
+			RootFolders.Add(ProjectFileGenerator.EngineRelativePath);
+			RootFolders.AddRange(AllGameFolders);
+
+			// Find all the plugin source directories
+			foreach(string RootFolder in RootFolders)
 			{
-				foreach( var Plugin in Plugins.AllPlugins )
+				string PluginsFolder = Path.Combine(RootFolder, "Plugins");
+				foreach(string PluginFile in Plugins.EnumeratePlugins(PluginsFolder))
 				{
-					// Plugin source directory
-					var PluginSourceDirectory = Path.Combine( Plugin.Directory, "Source" );
-					if( Directory.Exists( PluginSourceDirectory ) )
-					{
-						Folders.Add( PluginSourceDirectory );	
-					}
+					string PluginDirectory = Path.GetDirectoryName(PluginFile);
+					Folders.Add(Path.Combine(PluginDirectory, "Source"));
 				}
 			}
 
