Allow -Wshadow to be enabled via build configuration

diff --git a/Engine/Source/Programs/UnrealBuildTool/Configuration/BuildConfiguration.cs b/Engine/Source/Programs/UnrealBuildTool/Configuration/BuildConfiguration.cs
--- a/Engine/Source/Programs/UnrealBuildTool/Configuration/BuildConfiguration.cs
+++ b/Engine/Source/Programs/UnrealBuildTool/Configuration/BuildConfiguration.cs
@@ -77,6 +77,12 @@ namespace UnrealBuildTool
 		[XmlConfig]
 		public static bool bOmitPCDebugInfoInDevelopment;
 
+		/// <summary>
+		/// Whether shadow variable warning should be enabled
+		/// </summary>
+		[XmlConfig]
+		public static bool bEnableShadowVariableWarning;
+
 		/// <summary>
 		/// Whether PDB files should be used for Visual C++ builds.
 		/// </summary>
@@ -458,6 +464,8 @@ namespace UnrealBuildTool
 			bSupportEditAndContinue = false;
 			bUseActionHistory = true;
 
+			bEnableShadowVariableWarning = false;
+
 			// Incremental linking can yield faster iteration times when making small changes
 			// NOTE: We currently don't use incremental linking because it tends to behave a bit buggy on some computers (PDB-related compile errors)
 			bUseIncrementalLinking = false;
diff --git a/Engine/Source/Programs/UnrealBuildTool/IOS/IOSToolChain.cs b/Engine/Source/Programs/UnrealBuildTool/IOS/IOSToolChain.cs
--- a/Engine/Source/Programs/UnrealBuildTool/IOS/IOSToolChain.cs
+++ b/Engine/Source/Programs/UnrealBuildTool/IOS/IOSToolChain.cs
@@ -229,6 +229,11 @@ namespace UnrealBuildTool
 // 			}
 
 			Result += " -Wall -Werror";
+				
+			if (BuildConfiguration.bEnableShadowVariableWarning)
+			{
+				Result += " -Wshadow";
+			}
 
 			Result += " -Wno-unused-variable";
 			Result += " -Wno-unused-value";
diff --git a/Engine/Source/Programs/UnrealBuildTool/Linux/LinuxToolChain.cs b/Engine/Source/Programs/UnrealBuildTool/Linux/LinuxToolChain.cs
--- a/Engine/Source/Programs/UnrealBuildTool/Linux/LinuxToolChain.cs
+++ b/Engine/Source/Programs/UnrealBuildTool/Linux/LinuxToolChain.cs
@@ -337,6 +337,11 @@ namespace UnrealBuildTool
             Result += " -Wno-unknown-pragmas";			// Slate triggers this (with its optimize on/off pragmas)
 			Result += " -Wno-invalid-offsetof"; // needed to suppress warnings about using offsetof on non-POD types.
 
+			if (BuildConfiguration.bEnableShadowVariableWarning)
+			{
+				Result += " -Wshadow";
+			}
+
             //Result += " -DOPERATOR_NEW_INLINE=FORCENOINLINE";
 
             // shipping builds will cause this warning with "ensure", so disable only in those case
diff --git a/Engine/Source/Programs/UnrealBuildTool/Mac/MacToolChain.cs b/Engine/Source/Programs/UnrealBuildTool/Mac/MacToolChain.cs
--- a/Engine/Source/Programs/UnrealBuildTool/Mac/MacToolChain.cs
+++ b/Engine/Source/Programs/UnrealBuildTool/Mac/MacToolChain.cs
@@ -170,6 +170,11 @@ namespace UnrealBuildTool
 			Result += " -Wno-unused-private-field";
 			Result += " -Wno-invalid-offsetof"; // needed to suppress warnings about using offsetof on non-POD types.
 
+			if (BuildConfiguration.bEnableShadowVariableWarning)
+			{
+				Result += " -Wshadow";
+			}
+
 			// @todo: Remove these two when the code is fixed and they're no longer needed
 			Result += " -Wno-logical-op-parentheses";
 			Result += " -Wno-unknown-pragmas";
