Use the plugin filename to derive the plugin name, rather than using the directory name. Using the directory name is not consistent with the engine.

diff --git a/Engine/Source/Programs/UnrealBuildTool/System/Plugins.cs b/Engine/Source/Programs/UnrealBuildTool/System/Plugins.cs
--- a/Engine/Source/Programs/UnrealBuildTool/System/Plugins.cs
+++ b/Engine/Source/Programs/UnrealBuildTool/System/Plugins.cs
@@ -8,6 +8,7 @@ using System.Linq;
 
 namespace UnrealBuildTool
 {
+	[DebuggerDisplay("\\{{FileName}\\}")]
 	public class PluginInfo
 	{
 		public enum LoadedFromType
@@ -20,10 +21,13 @@ namespace UnrealBuildTool
 		};
 		
 		// Plugin name
-		public string Name;
+		public readonly string Name;
+
+		// Path to the plugin
+		public readonly string FileName;
 
 		// Path to the plugin's root directory
-		public string Directory;
+		public readonly string Directory;
 
 		// The plugin descriptor
 		public PluginDescriptor Descriptor;
@@ -31,31 +35,23 @@ namespace UnrealBuildTool
 		// Where does this plugin live?
 		public LoadedFromType LoadedFrom;
 
-		public override string ToString()
+		/// <summary>
+		/// Constructs a PluginInfo object
+		/// </summary>
+		/// <param name="InFileName"></param>
+		/// <param name="InLoadedFrom">Where this pl</param>
+		public PluginInfo(string InFileName, LoadedFromType InLoadedFrom)
 		{
-			return Path.Combine(Directory, Name + ".uplugin");
+			Name = Path.GetFileNameWithoutExtension(InFileName);
+			FileName = Path.GetFullPath(InFileName);
+			Directory = Path.GetDirectoryName(FileName);
+			Descriptor = PluginDescriptor.FromFile(FileName);
+			LoadedFrom = InLoadedFrom;
 		}
 	}
 
 	public class Plugins
 	{
-		/// <summary>
-		/// Loads a plugin descriptor file and fills out a new PluginInfo structure.  Throws an exception on failure.
-		/// </summary>
-		/// <param name="PluginFile">The path to the plugin file to load</param>
-		/// <param name="LoadedFrom">Where the plugin was loaded from</param>
-		/// <returns>New PluginInfo for the loaded descriptor.</returns>
-		public static PluginInfo ReadPlugin(string PluginFileName, PluginInfo.LoadedFromType LoadedFrom)
-		{
-			// Create the plugin info object
-			PluginInfo Info = new PluginInfo();
-			Info.LoadedFrom = LoadedFrom;
-			Info.Directory = Path.GetDirectoryName(Path.GetFullPath(PluginFileName));
-			Info.Name = Path.GetFileName(Info.Directory);
-			Info.Descriptor = PluginDescriptor.FromFile(PluginFileName);
-			return Info;
-		}
-
 		/// <summary>
 		/// Read all the plugins available to a given project
 		/// </summary>
@@ -69,7 +65,7 @@ namespace UnrealBuildTool
 			string EnginePluginsDir = Path.Combine(BuildConfiguration.RelativeEnginePath, "Plugins");
 			foreach(string PluginFileName in EnumeratePlugins(EnginePluginsDir))
 			{
-				PluginInfo Plugin = ReadPlugin(PluginFileName, PluginInfo.LoadedFromType.Engine);
+				PluginInfo Plugin = new PluginInfo(PluginFileName, PluginInfo.LoadedFromType.Engine);
 				Plugins.Add(Plugin);
 			}
 
@@ -79,7 +75,7 @@ namespace UnrealBuildTool
 				string ProjectPluginsDir = Path.Combine(Path.GetDirectoryName(ProjectFileName), "Plugins");
 				foreach(string PluginFileName in EnumeratePlugins(ProjectPluginsDir))
 				{
-					PluginInfo Plugin = ReadPlugin(PluginFileName, PluginInfo.LoadedFromType.GameProject);
+					PluginInfo Plugin = new PluginInfo(PluginFileName, PluginInfo.LoadedFromType.GameProject);
 					Plugins.Add(Plugin);
 				}
 			}
