Disabled Mac .dSYM generation as we hit the Windows build machine's MAX_PATH limit for various projects.

diff --git a/Engine/Source/Programs/UnrealBuildTool/Mac/UEBuildMac.cs b/Engine/Source/Programs/UnrealBuildTool/Mac/UEBuildMac.cs
--- a/Engine/Source/Programs/UnrealBuildTool/Mac/UEBuildMac.cs
+++ b/Engine/Source/Programs/UnrealBuildTool/Mac/UEBuildMac.cs
@@ -173,7 +173,8 @@ namespace UnrealBuildTool
 			BuildConfiguration.ProcessorCountMultiplier = MacToolChain.GetAdjustedProcessorCountMultiplier();
 			BuildConfiguration.bUseSharedPCHs = false;
 
-			BuildConfiguration.bUsePDBFiles = bCreateDebugInfo && Configuration != CPPTargetConfiguration.Debug && Platform == CPPTargetPlatform.Mac;
+			// Disabled as we hit the Windows build machine's MAX_PATH limit for various projects
+			// BuildConfiguration.bUsePDBFiles = bCreateDebugInfo && Configuration != CPPTargetConfiguration.Debug && Platform == CPPTargetPlatform.Mac;
 
 			// we always deploy - the build machines need to be able to copy the files back, which needs the full bundle
 			BuildConfiguration.bDeployAfterCompile = true;
