Rename PluginInfo.LoadedFromType to PluginLoadedFrom.

diff --git a/Engine/Source/Programs/UnrealBuildTool/System/Plugins.cs b/Engine/Source/Programs/UnrealBuildTool/System/Plugins.cs
--- a/Engine/Source/Programs/UnrealBuildTool/System/Plugins.cs
+++ b/Engine/Source/Programs/UnrealBuildTool/System/Plugins.cs
@@ -8,18 +8,18 @@ using System.Linq;
 
 namespace UnrealBuildTool
 {
+	public enum PluginLoadedFrom
+	{
+		// Plugin is built-in to the engine
+		Engine,
+
+		// Project-specific plugin, stored within a game project directory
+		GameProject
+	}
+
 	[DebuggerDisplay("\\{{FileName}\\}")]
 	public class PluginInfo
 	{
-		public enum LoadedFromType
-		{
-			// Plugin is built-in to the engine
-			Engine,
-
-			// Project-specific plugin, stored within a game project directory
-			GameProject
-		};
-		
 		// Plugin name
 		public readonly string Name;
 
@@ -33,14 +33,14 @@ namespace UnrealBuildTool
 		public PluginDescriptor Descriptor;
 
 		// Where does this plugin live?
-		public LoadedFromType LoadedFrom;
+		public PluginLoadedFrom LoadedFrom;
 
 		/// <summary>
 		/// Constructs a PluginInfo object
 		/// </summary>
 		/// <param name="InFileName"></param>
 		/// <param name="InLoadedFrom">Where this pl</param>
-		public PluginInfo(string InFileName, LoadedFromType InLoadedFrom)
+		public PluginInfo(string InFileName, PluginLoadedFrom InLoadedFrom)
 		{
 			Name = Path.GetFileNameWithoutExtension(InFileName);
 			FileName = Path.GetFullPath(InFileName);
@@ -65,7 +65,7 @@ namespace UnrealBuildTool
 			string EnginePluginsDir = Path.Combine(BuildConfiguration.RelativeEnginePath, "Plugins");
 			foreach(string PluginFileName in EnumeratePlugins(EnginePluginsDir))
 			{
-				PluginInfo Plugin = new PluginInfo(PluginFileName, PluginInfo.LoadedFromType.Engine);
+				PluginInfo Plugin = new PluginInfo(PluginFileName, PluginLoadedFrom.Engine);
 				Plugins.Add(Plugin);
 			}
 
@@ -75,7 +75,7 @@ namespace UnrealBuildTool
 				string ProjectPluginsDir = Path.Combine(Path.GetDirectoryName(ProjectFileName), "Plugins");
 				foreach(string PluginFileName in EnumeratePlugins(ProjectPluginsDir))
 				{
-					PluginInfo Plugin = new PluginInfo(PluginFileName, PluginInfo.LoadedFromType.GameProject);
+					PluginInfo Plugin = new PluginInfo(PluginFileName, PluginLoadedFrom.GameProject);
 					Plugins.Add(Plugin);
 				}
 			}
diff --git a/Engine/Source/Programs/UnrealBuildTool/System/UProjectInfo.cs b/Engine/Source/Programs/UnrealBuildTool/System/UProjectInfo.cs
--- a/Engine/Source/Programs/UnrealBuildTool/System/UProjectInfo.cs
+++ b/Engine/Source/Programs/UnrealBuildTool/System/UProjectInfo.cs
@@ -308,7 +308,7 @@ namespace UnrealBuildTool
 		/// <returns>True if the plugin should be enabled for this project</returns>
 		public static bool IsPluginEnabledForProject(PluginInfo Plugin, ProjectDescriptor Project, UnrealTargetPlatform Platform)
 		{
-			bool bEnabled = Plugin.Descriptor.bEnabledByDefault || Plugin.LoadedFrom == PluginInfo.LoadedFromType.GameProject;
+			bool bEnabled = Plugin.Descriptor.bEnabledByDefault || Plugin.LoadedFrom == PluginLoadedFrom.GameProject;
 			if(Project != null && Project.Plugins != null)
 			{
 				foreach(PluginReferenceDescriptor PluginReference in Project.Plugins)
