Linux: suppress '-W-inconsistent-missing-override' warning for 4.8

diff --git a/Engine/Source/Programs/UnrealBuildTool/Linux/LinuxToolChain.cs b/Engine/Source/Programs/UnrealBuildTool/Linux/LinuxToolChain.cs
--- a/Engine/Source/Programs/UnrealBuildTool/Linux/LinuxToolChain.cs
+++ b/Engine/Source/Programs/UnrealBuildTool/Linux/LinuxToolChain.cs
@@ -326,6 +326,7 @@ namespace UnrealBuildTool
 				if (CompilerVersionGreaterOrEqual(3, 6, 0))
 				{
 					Result += " -Wno-unused-local-typedef";	// clang is being overly strict here? PhysX headers trigger this.
+					Result += " -Wno-inconsistent-missing-override";	// these have to be suppressed for UE 4.8, should be fixed later.
 				}
             }
 
