AppxManifest and Resource files for Xbox One generated on the fly.
Settings used to generate these made settable from the editor.

diff --git a/Engine/Config/BaseGame.ini b/Engine/Config/BaseGame.ini
--- a/Engine/Config/BaseGame.ini
+++ b/Engine/Config/BaseGame.ini
@@ -21,6 +21,7 @@ bRequiresPushToTalk=true
 
 [/Script/EngineSettings.GeneralProjectSettings]
 CompanyName=
+CompanyDistinguishedName=
 CopyrightNotice=Fill out your copyright notice in the Description page of Project Settings.
 Description=
 LicensingTerms=
@@ -38,6 +39,7 @@ IncludePrerequisites=True
 IncludeCrashReporter=True
 InternationalizationPreset=English
 +CulturesToStage=en
+DefaultCulture=en
 
 [/Script/Engine.HUD]
 ConsoleMessageCount=4
diff --git a/Engine/Source/Editor/UnrealEd/Classes/Settings/ProjectPackagingSettings.h b/Engine/Source/Editor/UnrealEd/Classes/Settings/ProjectPackagingSettings.h
--- a/Engine/Source/Editor/UnrealEd/Classes/Settings/ProjectPackagingSettings.h
+++ b/Engine/Source/Editor/UnrealEd/Classes/Settings/ProjectPackagingSettings.h
@@ -130,6 +130,10 @@ public:
 	UPROPERTY(config, EditAnywhere, Category=Packaging, AdvancedDisplay, meta=(DisplayName="Localizations to Package"))
 	TArray<FString> CulturesToStage;
 
+	/** Culture to use if no matching culture is found. */
+	UPROPERTY(config, EditAnywhere, Category=Packaging, AdvancedDisplay, meta=(DisplayName="Package default localization"))
+	FString DefaultCulture;
+
 	/**
 	 * Cook all things in the project content directory
 	 */
diff --git a/Engine/Source/Programs/UnrealBuildTool/System/UEPlatformProjectGenerator.cs b/Engine/Source/Programs/UnrealBuildTool/System/UEPlatformProjectGenerator.cs
--- a/Engine/Source/Programs/UnrealBuildTool/System/UEPlatformProjectGenerator.cs
+++ b/Engine/Source/Programs/UnrealBuildTool/System/UEPlatformProjectGenerator.cs
@@ -247,7 +247,7 @@ namespace UnrealBuildTool
 		 *	
 		 *	@return	string				The custom property import lines for the project file; Empty string if it doesn't require one
 		 */
-		public virtual string GetVisualStudioLayoutDirSection(UnrealTargetPlatform InPlatform, UnrealTargetConfiguration InConfiguration, string InConditionString, TargetRules.TargetType TargetType, string TargetRulesPath, string ProjectFilePath)
+		public virtual string GetVisualStudioLayoutDirSection(UnrealTargetPlatform InPlatform, UnrealTargetConfiguration InConfiguration, string InConditionString, TargetRules.TargetType TargetType, string TargetRulesPath, string ProjectFilePath, string NMakeOutputPath)
 		{
 			return "";
 		}
@@ -259,7 +259,7 @@ namespace UnrealBuildTool
 		 *	
 		 *	@return	string				The output manifest section for the project file; Empty string if it doesn't require one
 		 */
-		public virtual string GetVisualStudioOutputManifestSection(UnrealTargetPlatform InPlatform, TargetRules.TargetType TargetType, string TargetRulesPath, string ProjectFilePath)
+        public virtual string GetVisualStudioOutputManifestSection(UnrealTargetPlatform InPlatform, TargetRules.TargetType TargetType, string TargetRulesPath, string ProjectFilePath)
 		{
 			return "";
 		}
diff --git a/Engine/Source/Programs/UnrealBuildTool/System/VCProject.cs b/Engine/Source/Programs/UnrealBuildTool/System/VCProject.cs
--- a/Engine/Source/Programs/UnrealBuildTool/System/VCProject.cs
+++ b/Engine/Source/Programs/UnrealBuildTool/System/VCProject.cs
@@ -916,21 +916,20 @@ namespace UnrealBuildTool
 					"		<Import Project=\"$(UserRootDir)\\Microsoft.Cpp.$(Platform).user.props\" Condition=\"exists('$(UserRootDir)\\Microsoft.Cpp.$(Platform).user.props')\" Label=\"LocalAppDataPlatform\" />" + ProjectFileGenerator.NewLine +
 					"	</ImportGroup>" + ProjectFileGenerator.NewLine);
 
-				VCProjectFileContent.Append(
-					"	<PropertyGroup " + ConditionString + ">" + ProjectFileGenerator.NewLine);
-
 				if (IsStubProject)
 				{
 					string ProjectRelativeUnusedDirectory = NormalizeProjectPath(Path.Combine(ProjectFileGenerator.EngineRelativePath, BuildConfiguration.BaseIntermediateFolder, "Unused"));
 
 					VCProjectFileContent.Append(
+					    "	<PropertyGroup " + ConditionString + ">" + ProjectFileGenerator.NewLine +
 						"		<OutDir>" + ProjectRelativeUnusedDirectory + Path.DirectorySeparatorChar + "</OutDir>" + ProjectFileGenerator.NewLine +
 						"		<IntDir>" + ProjectRelativeUnusedDirectory + Path.DirectorySeparatorChar + "</IntDir>" + ProjectFileGenerator.NewLine +
 						"		<NMakeBuildCommandLine>@rem Nothing to do.</NMakeBuildCommandLine>" + ProjectFileGenerator.NewLine +
 						"		<NMakeReBuildCommandLine>@rem Nothing to do.</NMakeReBuildCommandLine>" + ProjectFileGenerator.NewLine +
 						"		<NMakeCleanCommandLine>@rem Nothing to do.</NMakeCleanCommandLine>" + ProjectFileGenerator.NewLine +
-						"		<NMakeOutput/>" + ProjectFileGenerator.NewLine);
-				}
+						"		<NMakeOutput/>" + ProjectFileGenerator.NewLine +
+                        "	</PropertyGroup>" + ProjectFileGenerator.NewLine);
+                }
 				else
 				{
 					TargetRules TargetRulesObject = Combination.ProjectTarget.TargetRules;
@@ -999,6 +998,9 @@ namespace UnrealBuildTool
 					NMakePath += BuildPlatform.GetBinaryExtension(UEBuildBinaryType.Executable);
 					NMakePath = (BuildPlatform as UEBuildPlatform).ModifyNMakeOutput(NMakePath);
 
+                    VCProjectFileContent.Append(
+                        "	<PropertyGroup " + ConditionString + ">" + ProjectFileGenerator.NewLine);
+
 					string PathStrings = (ProjGenerator != null) ? ProjGenerator.GetVisualStudioPathsEntries(Platform, Configuration, TargetRulesObject.Type, TargetFilePath, ProjectFilePath, NMakePath) : "";
 					if (string.IsNullOrEmpty(PathStrings) || (PathStrings.Contains("<IntDir>") == false))
 					{
@@ -1053,8 +1055,11 @@ namespace UnrealBuildTool
 					VCProjectFileContent.Append("		<NMakeOutput>");
 					VCProjectFileContent.Append(NormalizeProjectPath(NMakePath));
 					VCProjectFileContent.Append("</NMakeOutput>" + ProjectFileGenerator.NewLine);
-				}
-				VCProjectFileContent.Append("	</PropertyGroup>" + ProjectFileGenerator.NewLine);
+				    VCProjectFileContent.Append("	</PropertyGroup>" + ProjectFileGenerator.NewLine);
+
+                    string LayoutDirString = (ProjGenerator != null) ? ProjGenerator.GetVisualStudioLayoutDirSection(Platform, Configuration, ConditionString, Combination.ProjectTarget.TargetRules.Type, Combination.ProjectTarget.TargetFilePath, ProjectFilePath, NMakePath) : "";
+                    VCProjectFileContent.Append(LayoutDirString);
+                }
 
 				if (VCUserFileContent != null && Combination.ProjectTarget != null)
 				{
@@ -1102,9 +1107,6 @@ namespace UnrealBuildTool
 					VCUserFileContent.Append(PlatformUserFileStrings);
 				}
 			}
-
-			string LayoutDirString = (ProjGenerator != null) ? ProjGenerator.GetVisualStudioLayoutDirSection(Platform, Configuration, ConditionString, Combination.ProjectTarget.TargetRules.Type, Combination.ProjectTarget.TargetFilePath, ProjectFilePath) : "";
-			VCProjectFileContent.Append(LayoutDirString);
 		}
 	}
 
diff --git a/Engine/Source/Programs/UnrealBuildTool/UnrealBuildTool.csproj b/Engine/Source/Programs/UnrealBuildTool/UnrealBuildTool.csproj
--- a/Engine/Source/Programs/UnrealBuildTool/UnrealBuildTool.csproj
+++ b/Engine/Source/Programs/UnrealBuildTool/UnrealBuildTool.csproj
@@ -91,6 +91,7 @@
     <Reference Include="System.Data" />
     <Reference Include="System.Management" />
     <Reference Include="System.Web.Extensions" />
+    <Reference Include="System.Windows.Forms" />
     <Reference Include="System.Xml" />
     <Reference Include="System.Xml.Linq" />
   </ItemGroup>
@@ -258,6 +259,8 @@
     <Compile Include="Windows\VCEnvironment.cs" />
     <Compile Include="Windows\VCToolChain.cs" />
     <Compile Include="Windows\WindowsProjectGenerator.cs" />
+    <Compile Include="XboxOne\XboxOneResourceGenerator.cs" />
+    <Compile Include="XboxOne\XboxOneManifestGenerator.cs" />
   </ItemGroup>
   <ItemGroup>
     <BootstrapperPackage Include="Microsoft.Net.Client.3.5">
diff --git a/Engine/Source/Runtime/EngineSettings/Classes/GeneralProjectSettings.h b/Engine/Source/Runtime/EngineSettings/Classes/GeneralProjectSettings.h
--- a/Engine/Source/Runtime/EngineSettings/Classes/GeneralProjectSettings.h
+++ b/Engine/Source/Runtime/EngineSettings/Classes/GeneralProjectSettings.h
@@ -15,6 +15,10 @@ class ENGINESETTINGS_API UGeneralProjectSettings
 	UPROPERTY(config, EditAnywhere, Category=Publisher)
 	FString CompanyName;
 
+	/** The distinguished name of the company (author, provider) that created the project. */
+	UPROPERTY(config, EditAnywhere, Category=Publisher)
+	FString CompanyDistinguishedName;
+
 	/** The project's copyright and/or trademark notices. */
 	UPROPERTY(config, EditAnywhere, Category=Legal)
 	FString CopyrightNotice;
