UBT: better explanation for signal exit codes.

- To reduce confusion regarding UHT failinng with "error code: 139"

diff --git a/Engine/Source/Programs/UnrealBuildTool/System/ExternalExecution.cs b/Engine/Source/Programs/UnrealBuildTool/System/ExternalExecution.cs
--- a/Engine/Source/Programs/UnrealBuildTool/System/ExternalExecution.cs
+++ b/Engine/Source/Programs/UnrealBuildTool/System/ExternalExecution.cs
@@ -642,6 +642,16 @@ namespace UnrealBuildTool
 
 					if (UHTResult != ECompilationResult.Succeeded)
 					{
+						// On Linux and Mac, the shell will return 128+signal number exit codes if UHT gets a signal (e.g. crashes or is interrupted)
+						if ((BuildHostPlatform.Current.Platform == UnrealTargetPlatform.Linux ||
+						    BuildHostPlatform.Current.Platform == UnrealTargetPlatform.Mac) &&
+						    (int)(UHTResult) >= 128
+						    )
+						{
+							// SIGINT is 2, so 128 + SIGINT is 130
+							UHTResult = ((int)(UHTResult) == 130) ? ECompilationResult.Canceled : ECompilationResult.CrashOrAssert;
+						}
+
 						Log.TraceInformation("Error: Failed to generate code for {0} - error code: {2} ({1})", ActualTargetName, (int) UHTResult, UHTResult.ToString());
 						return false;
 					}
