Fix performance issue when showing progress for large number of items; console I/O becomes a bottleneck.

diff --git a/Engine/Source/Programs/UnrealBuildTool/System/Utils.cs b/Engine/Source/Programs/UnrealBuildTool/System/Utils.cs
--- a/Engine/Source/Programs/UnrealBuildTool/System/Utils.cs
+++ b/Engine/Source/Programs/UnrealBuildTool/System/Utils.cs
@@ -1004,6 +1004,7 @@ namespace UnrealBuildTool
 		bool bWriteToConsole;
 		string Message;
 		int NumCharsToBackspaceOver;
+		string CurrentProgressString;
 
 		public ProgressWriter(string InMessage, bool bInWriteToConsole)
 		{
@@ -1028,22 +1029,25 @@ namespace UnrealBuildTool
 		{
 			float ProgressValue = Denominator > 0 ? ((float)Numerator / (float)Denominator) : 1.0f;
 			string ProgressString = String.Format("{0}%", Math.Round(ProgressValue * 100.0f));
-
-			if (bWriteMarkup)
-			{
-				Log.WriteLine(TraceEventType.Information, "@progress '{0}' {1}", Message, ProgressString);
-			}
-			else if (bWriteToConsole)
+			if (ProgressString != CurrentProgressString)
 			{
-				// Backspace over previous progress value
-				while (NumCharsToBackspaceOver-- > 0)
+				CurrentProgressString = ProgressString;
+				if (bWriteMarkup)
 				{
-					Console.Write("\b");
+					Log.WriteLine(TraceEventType.Information, "@progress '{0}' {1}", Message, ProgressString);
 				}
+				else if (bWriteToConsole)
+				{
+					// Backspace over previous progress value
+					while (NumCharsToBackspaceOver-- > 0)
+					{
+						Console.Write("\b");
+					}
 
-				// Display updated progress string and keep track of how long it was
-				NumCharsToBackspaceOver = ProgressString.Length;
-				Console.Write(ProgressString);
+					// Display updated progress string and keep track of how long it was
+					NumCharsToBackspaceOver = ProgressString.Length;
+					Console.Write(ProgressString);
+				}
 			}
 		}
 	}
