Don't include the module name in the generated PCH; it's in the directory name already, and pushes filenames closer to MAX_PATH on Windows.

Pathalogical case in the ExtendingTheEditor sample was using a PCH name of: D:\\BuildFarm\\buildmachine_++depot+UE4\\Samples\\Sandbox\\ExtendingTheEditor\\Plugins\\DemoEditorExtensions\\Intermediate\\Build\\Win64\\ExtendingTheEditor\\Debug\\DemoEditorExtensionsRuntime\\PCH.DemoEditorExtensionsRuntime.DemoEditorExtensionsRuntimePrivatePCH.h.cpp

diff --git a/Engine/Source/Programs/UnrealBuildTool/WinRT/WinRTToolChain.cs b/Engine/Source/Programs/UnrealBuildTool/WinRT/WinRTToolChain.cs
--- a/Engine/Source/Programs/UnrealBuildTool/WinRT/WinRTToolChain.cs
+++ b/Engine/Source/Programs/UnrealBuildTool/WinRT/WinRTToolChain.cs
@@ -506,7 +506,7 @@ namespace UnrealBuildTool
 				if (CompileEnvironment.Config.PrecompiledHeaderAction == PrecompiledHeaderAction.Create)
 				{
 					// Generate a CPP File that just includes the precompiled header.
-					string PCHCPPFilename = "PCH." + ModuleName + "." + Path.GetFileName(CompileEnvironment.Config.PrecompiledHeaderIncludeFilename) + ".cpp";
+					string PCHCPPFilename = "PCH." + Path.GetFileName(CompileEnvironment.Config.PrecompiledHeaderIncludeFilename) + ".cpp";
 					string PCHCPPPath = Path.Combine(CompileEnvironment.Config.OutputDirectory, PCHCPPFilename);
 					FileItem PCHCPPFile = FileItem.CreateIntermediateTextFile(
 						PCHCPPPath,
diff --git a/Engine/Source/Programs/UnrealBuildTool/Windows/VCToolChain.cs b/Engine/Source/Programs/UnrealBuildTool/Windows/VCToolChain.cs
--- a/Engine/Source/Programs/UnrealBuildTool/Windows/VCToolChain.cs
+++ b/Engine/Source/Programs/UnrealBuildTool/Windows/VCToolChain.cs
@@ -1044,7 +1044,7 @@ namespace UnrealBuildTool
 				if (CompileEnvironment.Config.PrecompiledHeaderAction == PrecompiledHeaderAction.Create)
 				{
 					// Generate a CPP File that just includes the precompiled header.
-					string PCHCPPFilename = "PCH." + ModuleName + "." + Path.GetFileName(CompileEnvironment.Config.PrecompiledHeaderIncludeFilename) + ".cpp";
+					string PCHCPPFilename = "PCH." + Path.GetFileName(CompileEnvironment.Config.PrecompiledHeaderIncludeFilename) + ".cpp";
 					string PCHCPPPath = Path.Combine(CompileEnvironment.Config.OutputDirectory, PCHCPPFilename);
 					FileItem PCHCPPFile = FileItem.CreateIntermediateTextFile(
 						PCHCPPPath,
@@ -1188,12 +1188,12 @@ namespace UnrealBuildTool
 					// All files using the same PCH are required to share the same PDB that was used when compiling the PCH
 					if (CompileEnvironment.Config.PrecompiledHeaderAction == PrecompiledHeaderAction.Include)
 					{
-						PDBFileName = "PCH." + ModuleName + "." + Path.GetFileName(CompileEnvironment.Config.PrecompiledHeaderIncludeFilename);
+						PDBFileName = "PCH." + Path.GetFileName(CompileEnvironment.Config.PrecompiledHeaderIncludeFilename);
 					}
 					// Files creating a PCH use a PDB per file.
 					else if (CompileEnvironment.Config.PrecompiledHeaderAction == PrecompiledHeaderAction.Create)
 					{
-						PDBFileName = "PCH." + ModuleName + "." + Path.GetFileName(CompileEnvironment.Config.PrecompiledHeaderIncludeFilename);
+						PDBFileName = "PCH." + Path.GetFileName(CompileEnvironment.Config.PrecompiledHeaderIncludeFilename);
 						bActionProducesPDB = true;
 					}
 					// Ungrouped C++ files use a PDB per file.
