Changed the xcode scheme generation to respect binary locations for targets not under Engine.
Merging using UE4-To-UE4-LauncherDev

diff --git a/Engine/Source/Programs/UnrealBuildTool/System/XcodeProjectFileGenerator.cs b/Engine/Source/Programs/UnrealBuildTool/System/XcodeProjectFileGenerator.cs
--- a/Engine/Source/Programs/UnrealBuildTool/System/XcodeProjectFileGenerator.cs
+++ b/Engine/Source/Programs/UnrealBuildTool/System/XcodeProjectFileGenerator.cs
@@ -1204,6 +1204,35 @@ namespace UnrealBuildTool
 			{
 				TargetBinariesFolder = Path.GetDirectoryName(Directory.GetCurrentDirectory());
 
+                if (!ExeBaseName.EndsWith("Editor") && !string.IsNullOrWhiteSpace(Target.TargetFilePath) && !Utils.IsFileUnderDirectory(Target.TargetFilePath, TargetBinariesFolder))
+                {
+					string DepotRoot = TargetBinariesFolder.TrimEnd(Path.DirectorySeparatorChar);
+					DepotRoot = DepotRoot.Substring(0, DepotRoot.LastIndexOf(Path.DirectorySeparatorChar));
+					DepotRoot = DepotRoot.TrimEnd(Path.DirectorySeparatorChar);
+
+					Log.TraceInformation(string.Format("Target.TargetFilePath={0}", Target.TargetFilePath));
+					Log.TraceInformation(string.Format("DepotRoot={0}", DepotRoot));
+
+                    if (Utils.IsFileUnderDirectory(Target.TargetFilePath, DepotRoot))
+                    {
+                        string PreviousPath = Target.TargetFilePath;
+                        string CurrentPath = Target.TargetFilePath;
+                        while (!string.IsNullOrWhiteSpace(CurrentPath) && CurrentPath != DepotRoot)
+                        {
+                            PreviousPath = CurrentPath;
+                            CurrentPath = CurrentPath.Substring(0, CurrentPath.LastIndexOf(Path.DirectorySeparatorChar));
+                        }
+
+                        if (CurrentPath == DepotRoot && !string.IsNullOrWhiteSpace(PreviousPath))
+                        {
+                            TargetBinariesFolder = PreviousPath;
+                        }
+
+						Log.TraceInformation(string.Format("TargetBinariesFolder={0}", TargetBinariesFolder));
+						Log.TraceInformation(string.Format("ExeBaseName={0}", ExeBaseName));
+                    }
+                }
+
 				if (ExeBaseName == Target.TargetName && ExeBaseName.EndsWith("Game") && ExeExtension == ".app")
 				{
 					TargetBinariesFolder = TargetBinariesFolder.Replace("/Engine", "/") + ExeBaseName + "/Binaries/";
