Fix some log output to use Environment.Newline

diff --git a/Engine/Source/Programs/UnrealBuildTool/System/RulesCompiler.cs b/Engine/Source/Programs/UnrealBuildTool/System/RulesCompiler.cs
--- a/Engine/Source/Programs/UnrealBuildTool/System/RulesCompiler.cs
+++ b/Engine/Source/Programs/UnrealBuildTool/System/RulesCompiler.cs
@@ -1808,14 +1808,14 @@ namespace UnrealBuildTool
 				ExceptionMessage += TargetName;
 				ExceptionMessage += "' in rules assembly '";
 				ExceptionMessage += RulesAssembly.FullName;
-				ExceptionMessage += "'.\n";
+                ExceptionMessage += "'." + Environment.NewLine;
 
-				ExceptionMessage += "Location: " + RulesAssembly.Location + "\n";
+                ExceptionMessage += "Location: " + RulesAssembly.Location + Environment.NewLine;
 
-				ExceptionMessage += "Target rules found:\n";
+                ExceptionMessage += "Target rules found:" + Environment.NewLine;
 				foreach (KeyValuePair<string, string> entry in TargetNameToTargetFileMap)
 				{
-					ExceptionMessage += "\t" + entry.Key + " - " + entry.Value + "\n";
+                    ExceptionMessage += "\t" + entry.Key + " - " + entry.Value + Environment.NewLine;
 				}
 
 				throw new BuildException(ExceptionMessage);
