UBT Utils.cs (New logging system)
* Allows us to use built-in Trace providers (console, file, etc) directly and still use our custom formatting.
* Fat comments explaining why Trace.WriteXXX functions should not be used directly in our system.
* Fixes thread safety by using Trace.WriteXXX under the hood after formatting, which uses a global lock (except on Mono, where a bug appears to be preventing this. Simulating the call on that platform).
* No need for TraceEvent overloads, which saves us the extra parameter cruft.
* Removed non-varargs overloads of Log functions (technically a bit slower, but these are already small messages).
* No longer needed VerbosityFilter and ConsoleListener classes.
* Avoid calling GetSource() if we aren't outputting the source.
* Avoid formatting the string if it won't pass the verbosity level.
* Consolidated all of UAT and UBT options into this class, so they could fully share the implementation.

UBT BuildConfiguration.cs
* Added LogFilename (and --log=<file> arg) that enables logging to a file.
* Added static ctor guard that asserts if someone tries to read a config before we have loaded config files and parsed config-override commandlines. It's a poor man's hack, but better than nothing!

UBT UEBuildConfiguration.cs
* Same static ctor guard as above.

UBT UnrealBuildTools.cs (initialization refactoring)
* In general I tried to de-mystify some of the rationale behind our startup code via fat comments.
* Broke main into 3 stages:
1. "early code" that should not try to read a config value.
  * Very little code here. Mostly setting the current directory.
  * Does an early init of logging to ensure logging is around, but config values won't be ready.
2. "Init Configuration code" that loads config files and parses command lines that may override them.
  * I isolated two locations in startup that parsed long sets of switches and moved ones that trivially affected BuildConfiguration and UEBuildConfiguration in here. Those two locations seemed to have mostly copies of the same switches, indicating serious param parsing issues at some point in time.
  * This allows switches to override config files more easily than the patchwork of re-parsing that was currently used (particularly for -verbose).
  * I did a cursory examination of later code that indicated this double (actually, triple) parsing was no longer necessary with the refactors above. Any insight into why things may have ended up this way would be helpful.
3. "Post Init code" that is actually the meat of UBT.
  * I left this code largely untouched.
  * Removed 2 of 3 different command line logging statements.
  * Removed two redundant parses of config overrides (ParseBuildConfigurationFlags).
* Guarded all of main in a try/catch block to ensure no exceptions can leak from UBT without returning a valid error code. It ALMOST already did this, but only covered the part surrounded by the Mutex.
* There was a perplexing bit that redundantly called XmlConfigLoader.Reset<> (line 683) that I struggled to understand. It turns out UEBuildConfiguration was sensitive to the current directory being set before files were loaded, and the old code called XmlConfigLoader.Init() super early, which required it to be called again after the current directory was set (see UEBuldConfiguration.UEThirdPartySourceDirectory for the cause). After my changes, I verified as best I could that these calls are no longer needed and removed them.

XmlConfigLoader.cs
* Add support for Properties in XmlConfigLoader.

AutomationTool Program.cs
* Guard logging shutdown code in try/finally so it can't be missed.

AutomationTool Log.cs
* Uses new logging system from UBT
* Removed unnecessary classes (VerbosityFilter, AutomationConsoleTraceListener, and AutomationFileTraceListener)
* Console trace logic is handled by UBT code now, moved UTF8Output handling to InitLogging.
* A custom TraceListener for file logging was unnecessary.
  * Logic to handle creating the log file and retry loops was move into InitLogging, and the result passed to a regular TextFileTraceListener.
  * Logic to handle copying the log on shutdown was moved to a ShutdownLogging function.

diff --git a/Engine/Source/Programs/AutomationTool/CommandUtils.cs b/Engine/Source/Programs/AutomationTool/CommandUtils.cs
--- a/Engine/Source/Programs/AutomationTool/CommandUtils.cs
+++ b/Engine/Source/Programs/AutomationTool/CommandUtils.cs
@@ -122,7 +122,7 @@ namespace AutomationTool
 		[MethodImplAttribute(MethodImplOptions.NoInlining)]
 		public static void Log(string Format, params object[] Args)
 		{
-			UnrealBuildTool.Log.WriteLine(1, TraceEventType.Information, Format, Args);
+			UnrealBuildTool.Log.WriteLine(1, null, TraceEventType.Information, Format, Args);
 		}
 
 		/// <summary>
@@ -132,7 +132,7 @@ namespace AutomationTool
 		[MethodImplAttribute(MethodImplOptions.NoInlining)]
 		public static void Log(string Message)
 		{
-			UnrealBuildTool.Log.WriteLine(1, TraceEventType.Information, Message);
+            UnrealBuildTool.Log.WriteLine(1, null, TraceEventType.Information, Message);
 		}
 
 		/// <summary>
@@ -143,7 +143,7 @@ namespace AutomationTool
 		[MethodImplAttribute(MethodImplOptions.NoInlining)]
 		public static void LogError(string Format, params object[] Args)
 		{
-			UnrealBuildTool.Log.WriteLine(1, TraceEventType.Error, Format, Args);
+            UnrealBuildTool.Log.WriteLine(1, null, TraceEventType.Error, Format, Args);
 		}
 
 		/// <summary>
@@ -153,7 +153,7 @@ namespace AutomationTool
 		[MethodImplAttribute(MethodImplOptions.NoInlining)]
 		public static void LogError(string Message)
 		{
-			UnrealBuildTool.Log.WriteLine(1, TraceEventType.Error, Message);
+            UnrealBuildTool.Log.WriteLine(1, null, TraceEventType.Error, Message);
 		}
 
 		/// <summary>
@@ -164,7 +164,7 @@ namespace AutomationTool
 		[MethodImplAttribute(MethodImplOptions.NoInlining)]
 		public static void LogWarning(string Format, params object[] Args)
 		{
-			UnrealBuildTool.Log.WriteLine(1, TraceEventType.Warning, Format, Args);
+            UnrealBuildTool.Log.WriteLine(1, null, TraceEventType.Warning, Format, Args);
 		}
 
 		/// <summary>
@@ -174,7 +174,7 @@ namespace AutomationTool
 		[MethodImplAttribute(MethodImplOptions.NoInlining)]
 		public static void LogWarning(string Message)
 		{
-			UnrealBuildTool.Log.WriteLine(1, TraceEventType.Warning, Message);
+            UnrealBuildTool.Log.WriteLine(1, null, TraceEventType.Warning, Message);
 		}
 
 		/// <summary>
@@ -185,7 +185,7 @@ namespace AutomationTool
 		[MethodImplAttribute(MethodImplOptions.NoInlining)]
 		public static void LogVerbose(string Format, params object[] Args)
 		{
-			UnrealBuildTool.Log.WriteLine(1, TraceEventType.Verbose, Format, Args);
+            UnrealBuildTool.Log.WriteLine(1, null, TraceEventType.Verbose, Format, Args);
 		}
 
 		/// <summary>
@@ -195,7 +195,7 @@ namespace AutomationTool
 		[MethodImplAttribute(MethodImplOptions.NoInlining)]
 		public static void LogVerbose(string Message)
 		{
-			UnrealBuildTool.Log.WriteLine(1, TraceEventType.Verbose, Message);
+            UnrealBuildTool.Log.WriteLine(1, null, TraceEventType.Verbose, Message);
 		}
 
 		/// <summary>
@@ -207,7 +207,7 @@ namespace AutomationTool
 		[MethodImplAttribute(MethodImplOptions.NoInlining)]
 		public static void Log(TraceEventType Verbosity, string Format, params object[] Args)
 		{
-			UnrealBuildTool.Log.WriteLine(1, Verbosity, Format, Args);
+            UnrealBuildTool.Log.WriteLine(1, null, Verbosity, Format, Args);
 		}
 
 		/// <summary>
@@ -218,7 +218,7 @@ namespace AutomationTool
 		[MethodImplAttribute(MethodImplOptions.NoInlining)]
 		public static void Log(TraceEventType Verbosity, string Message)
 		{
-			UnrealBuildTool.Log.WriteLine(1, Verbosity, Message);
+            UnrealBuildTool.Log.WriteLine(1, null, Verbosity, Message);
 		}
 
 		/// <summary>
@@ -229,7 +229,7 @@ namespace AutomationTool
 		[MethodImplAttribute(MethodImplOptions.NoInlining)]
 		public static void Log(TraceEventType Verbosity, Exception Ex)
 		{
-			UnrealBuildTool.Log.WriteLine(1, Verbosity, LogUtils.FormatException(Ex));
+            UnrealBuildTool.Log.WriteLine(1, null, Verbosity, LogUtils.FormatException(Ex));
 		}
 
 		#endregion
diff --git a/Engine/Source/Programs/AutomationTool/Log.cs b/Engine/Source/Programs/AutomationTool/Log.cs
--- a/Engine/Source/Programs/AutomationTool/Log.cs
+++ b/Engine/Source/Programs/AutomationTool/Log.cs
@@ -14,8 +14,7 @@ namespace AutomationTool
 
 	public class LogUtils
 	{
-		private static Stopwatch Timer;
-		private static AutomationFileTraceListener FileLog;
+        private static string LogFilename;
 
 		/// <summary>
 		/// Initializes trace logging.
@@ -23,75 +22,126 @@ namespace AutomationTool
 		/// <param name="CommandLine">Command line.</param>
 		public static void InitLogging(string[] CommandLine)
 		{
-			Timer = (CommandUtils.ParseParam(CommandLine, "-Timestamps"))? Stopwatch.StartNew() : null;
-			var VerbosityLevel = CommandUtils.ParseParam(CommandLine, "-Verbose") ? TraceEventType.Verbose : TraceEventType.Information;
-			var Filter = new VerbosityFilter(VerbosityLevel);
-			Trace.Listeners.Add(new AutomationConsoleTraceListener());
-			FileLog = new AutomationFileTraceListener();
-			Trace.Listeners.Add(FileLog);
-			Trace.Listeners.Add(new AutomationMemoryLogListener());
-			foreach (TraceListener Listener in Trace.Listeners)
-			{
-				Listener.Filter = Filter;
-			}
-		}
-
-		/// <summary>
-		/// Closes logging system
-		/// </summary>
-		public static void CloseFileLogging()
-		{
-			if (FileLog != null)
-			{
-				Trace.Listeners.Remove(FileLog);
-				FileLog.Close();
-				FileLog.Dispose();
-				FileLog = null;
-			}
-		}
-
-		/// <summary>
-		/// Formats message for logging.
-		/// </summary>
-		/// <param name="Source">Source of the message (usually "Class.Method")</param>
-		/// <param name="Verbosity">Message verbosity level</param>
-		/// <param name="Format">Message format string</param>
-		/// <param name="Args">Format arguments</param>
-		/// <returns>Formatted message</returns>
-		public static string FormatMessage(string Source, TraceEventType Verbosity, string Format, params object[] Args)
-		{
-			return FormatMessage(Source, Verbosity, String.Format(Format, Args));
-		}
-
-		/// <summary>
-		/// Formats message for logging.
-		/// </summary>
-		/// <param name="Source">Source of the message (usually Class.Method)</param>
-		/// <param name="Verbosity">Message verbosity level</param>
-		/// <param name="Message">Message text</param>
-		/// <returns>Formatted message</returns>
-		public static string FormatMessage(string Source, TraceEventType Verbosity, string Message)
-		{
-			string FormattedMessage = "";
-			if(Timer != null)
-			{
-				FormattedMessage += String.Format("[{0:hh\\:mm\\:ss\\.fff}] ", Timer.Elapsed);
-			}
-			FormattedMessage += Source + ": ";
-			switch (Verbosity)
-			{
-				case TraceEventType.Error:
-					FormattedMessage += "ERROR: ";
-					break;
-				case TraceEventType.Warning:
-					FormattedMessage += "WARNING: ";
-					break;
-			}
-			FormattedMessage += Message;
-			return FormattedMessage;
-		}
-
-		/// <summary>
+            UnrealBuildTool.Log.InitLogging(
+                bLogTimestamps: CommandUtils.ParseParam(CommandLine, "-Timestamps"),
+                bLogVerbose: CommandUtils.ParseParam(CommandLine, "-Verbose"),
+                bLogSeverity: true,
+                bLogSources: true,
+                bColorConsoleOutput: true,
+                TraceListeners: new TraceListener[]
+                {
+                    new ConsoleTraceListener(),
+                    // could return null, but InitLogging handles this gracefully.
+                    CreateLogFileListener(out LogFilename),
+                    //@todo - this is only used by GUBP nodes. Ideally we don't waste this 20MB if we are not running GUBP.
+                    new AutomationMemoryLogListener(),
+                });
+
+            // ensure UTF8Output flag is respected, since we are initializing logging early in the program.
+            if (CommandLine.Any(Arg => Arg.Equals("-utf8output", StringComparison.InvariantCultureIgnoreCase)))
+            {
+                Console.OutputEncoding = new System.Text.UTF8Encoding(false, false);
+            }
+        }
+
+        public static void ShutdownLogging()
+        {
+            // This closes all the output streams immediately, inside the Global Lock, so it's threadsafe.
+            Trace.Close();
+
+            // from here we can copy the log file to its final resting place
+            try
+            {
+                // Try to copy the log file to the log folder. The reason why it's done here is that
+                // at the time the log file is being initialized the env var may not yet be set (this 
+                // applies to local log folder in particular)
+                var LogFolder = Environment.GetEnvironmentVariable(EnvVarNames.LogFolder);
+                if (!String.IsNullOrEmpty(LogFolder) && Directory.Exists(LogFolder) &&
+                        !String.IsNullOrEmpty(LogFilename) && File.Exists(LogFilename))
+                {
+                    var DestFilename = CommandUtils.CombinePaths(LogFolder, "UAT_" + Path.GetFileName(LogFilename));
+                    SafeCopyLogFile(LogFilename, DestFilename);
+                }
+            }
+            catch (Exception)
+            {
+                // Silently ignore, logging is pointless because eveything is shut down at this point
+            }
+        }
+
+        /// <summary>
+        /// Copies log file to the final log folder, does multiple attempts if the destination file could not be created.
+        /// </summary>
+        /// <param name="SourceFilename"></param>
+        /// <param name="DestFilename"></param>
+        private static void SafeCopyLogFile(string SourceFilename, string DestFilename)
+        {
+            const int MaxAttempts = 10;
+            int AttemptNo = 0;
+            var DestLogFilename = DestFilename;
+            bool Result = false;
+            do
+            {
+                try
+                {
+                    File.Copy(SourceFilename, DestLogFilename, true);
+                    Result = true;
+                }
+                catch (Exception)
+                {
+                    var ModifiedFilename = String.Format("{0}_{1}{2}", Path.GetFileNameWithoutExtension(DestFilename), AttemptNo, Path.GetExtension(DestLogFilename));
+                    DestLogFilename = CommandUtils.CombinePaths(Path.GetDirectoryName(DestFilename), ModifiedFilename);
+                    AttemptNo++;
+                }
+            }
+            while (Result == false && AttemptNo <= MaxAttempts);
+        }
+
+        /// <summary>
+        /// Creates the TraceListener used for file logging.
+        /// We cannot simply use a TextWriterTraceListener because we need more flexibility when the file cannot be created.
+        /// TextWriterTraceListener lazily creates the file, silently failing when it cannot.
+        /// </summary>
+        /// <returns>The newly created TraceListener, or null if it could not be created.</returns>
+        private static TraceListener CreateLogFileListener(out string LogFilename)
+        {
+            StreamWriter LogFile = null;
+            const int MaxAttempts = 10;
+            int Attempt = 0;
+            var TempLogFolder = Path.GetTempPath();
+            do
+            {
+                if (Attempt == 0)
+                {
+                    LogFilename = CommandUtils.CombinePaths(TempLogFolder, "Log.txt");
+                }
+                else
+                {
+                    LogFilename = CommandUtils.CombinePaths(TempLogFolder, String.Format("Log_{0}.txt", Attempt));
+                }
+                try
+                {
+                    // We do not need to set AutoFlush on the StreamWriter because we set Trace.AutoFlush, which calls it for us.
+                    // Not only would this be redundant, StreamWriter AutoFlush does not flush the encoder, while a direct call to 
+                    // StreamWriter.Flush() will, which is what the Trace system with AutoFlush = true will do.
+                    // Internally, FileStream constructor opens the file with good arguments for writing to log files.
+                    return new TextWriterTraceListener(new StreamWriter(LogFilename), "AutomationFileLogListener");
+                }
+                catch (Exception Ex)
+                {
+                    if (Attempt == (MaxAttempts - 1))
+                    {
+                        // Clear out the LogFilename to indicate we were not able to write one.
+                        LogFilename = null;
+                        UnrealBuildTool.Log.TraceWarning("Unable to create log file: {0}", LogFilename);
+                        UnrealBuildTool.Log.TraceWarning(LogUtils.FormatException(Ex));
+                    }
+                }
+            } while (LogFile == null && ++Attempt < MaxAttempts);
+            return null;
+        }
+
+        /// <summary>
 		/// Dumps exception info to log.
 		/// </summary>
 		/// <param name="Verbosity">Verbosity</param>
@@ -166,99 +216,6 @@ namespace AutomationTool
 
 	#endregion
 
-
-	#region VerbosityFilter
-
-	/// <summary>
-	/// Trace verbosity filter.
-	/// </summary>
-	class VerbosityFilter : TraceFilter
-	{
-		private TraceEventType VerbosityLevel = TraceEventType.Information;
-		public VerbosityFilter(TraceEventType Level)
-		{
-			VerbosityLevel = Level;
-		}
-
-		public override bool ShouldTrace(TraceEventCache Cache, string Source, TraceEventType EventType, int Id, string FormatOrMessage, object[] Args, object Data1, object[] Data)
-		{
-			return EventType <= VerbosityLevel;
-		}
-	}
-
-	#endregion
-
-	#region AutomationConsoleTraceListener
-
-	/// <summary>
-	/// Trace console listener.
-	/// </summary>
-	class AutomationConsoleTraceListener : TraceListener
-	{
-		public AutomationConsoleTraceListener()
-		{
-			// We use the command line directly here, because GlobalCommandLine isn't initialized early enough
-			if (SharedUtils.ParseCommandLine().Any(Arg => Arg.ToLower() == "-utf8output"))
-			{
-				Console.OutputEncoding = new System.Text.UTF8Encoding(false, false);
-			}
-		}
-
-		/// <summary>
-		/// Writes a formatted line to the console.
-		/// </summary>
-		/// <param name="Source">Message source</param>
-		/// <param name="Verbosity">Message verbosity</param>
-		/// <param name="Format">Message format string</param>
-		/// <param name="Args">Format arguments</param>
-		private void WriteLine(string Source, TraceEventType Verbosity, string Format, params object[] Args)
-		{
-			string Message = LogUtils.FormatMessage(Source, Verbosity, Format, Args);
-			if (Filter == null || Filter.ShouldTrace(null, Source, Verbosity, 0, Format, Args, null, null))
-			{
-				ConsoleColor DefaultColor = Console.ForegroundColor;
-				switch (Verbosity)
-				{
-					case TraceEventType.Critical:
-					case TraceEventType.Error:
-						Console.ForegroundColor = ConsoleColor.Red;
-						break;
-					case TraceEventType.Warning:
-						Console.ForegroundColor = ConsoleColor.Yellow;
-						break;
-				}
-				Console.WriteLine(Message);
-				Console.ForegroundColor = DefaultColor;
-			}
-		}
-
-		#region TraceListener Interface
-
-		public override void Write(string message)
-		{
-			WriteLine(String.Empty, TraceEventType.Information, message);
-		}
-
-		public override void WriteLine(string message)
-		{
-			WriteLine(String.Empty, TraceEventType.Information, message);
-		}
-
-		public override void TraceEvent(TraceEventCache eventCache, string source, TraceEventType eventType, int id, string message)
-		{
-			WriteLine(source, eventType, "{0}", message);
-		}
-
-		public override void TraceEvent(TraceEventCache eventCache, string source, TraceEventType eventType, int id, string format, params object[] args)
-		{
-			WriteLine(source, eventType, format, args);
-		}
-
-		#endregion
-	}
-
-	#endregion
-
 	#region AutomationMemoryLogListener
 
 	/// <summary>
@@ -269,22 +226,16 @@ namespace AutomationTool
 		private static StringBuilder AccumulatedLog = new StringBuilder(1024 * 1024 * 20);
 		private static object SyncObject = new object();
 
-		/// <summary>
+        public override bool IsThreadSafe { get { return true; } }
+
+        /// <summary>
 		/// Writes a formatted line to the console.
 		/// </summary>
-		/// <param name="Source">Message source</param>
-		/// <param name="Verbosity">Message verbosity</param>
-		/// <param name="Format">Message format string</param>
-		/// <param name="Args">Format arguments</param>
-		private void WriteLine(string Source, TraceEventType Verbosity, string Format, params object[] Args)
+		private void WriteLinePrivate(string Message)
 		{
 			lock (SyncObject)
 			{
-				string Message = LogUtils.FormatMessage(Source, Verbosity, Format, Args);
-				if (Filter == null || Filter.ShouldTrace(null, Source, Verbosity, 0, Format, Args, null, null))
-				{
-					AccumulatedLog.AppendLine(Message);
-				}
+				AccumulatedLog.AppendLine(Message);
 			}
 		}
 
@@ -300,207 +251,15 @@ namespace AutomationTool
 
 		public override void Write(string message)
 		{
-			WriteLine(String.Empty, TraceEventType.Information, message);
-		}
-
-		public override void WriteLine(string message)
-		{
-			WriteLine(String.Empty, TraceEventType.Information, message);
-		}
-
-		public override void TraceEvent(TraceEventCache eventCache, string source, TraceEventType eventType, int id, string message)
-		{
-			WriteLine(source, eventType, "{0}", message);
-		}
-
-		public override void TraceEvent(TraceEventCache eventCache, string source, TraceEventType eventType, int id, string format, params object[] args)
-		{
-			WriteLine(source, eventType, format, args);
-		}
-
-		#endregion
-	}
-
-	#endregion
-
-
-	#region AutomationFileTraceListener
-
-	/// <summary>
-	/// Log file trace listener.
-	/// </summary>
-	class AutomationFileTraceListener : TraceListener
-	{
-		private string LogFilename;
-		private StreamWriter LogFile;
-		private static object SyncObject = new object();
-
-		public AutomationFileTraceListener()
-		{
-			const int MaxAttempts = 10;
-			int Attempt = 0;
-			var TempLogFolder = Path.GetTempPath();
-			do
-			{
-				if (Attempt == 0)
-				{
-					LogFilename = CommandUtils.CombinePaths(TempLogFolder, "Log.txt");
-				}
-				else
-				{
-					LogFilename = CommandUtils.CombinePaths(TempLogFolder, String.Format("Log_{0}.txt", Attempt));
-				}
-				try
-				{
-					FileStream File = new FileStream(LogFilename, FileMode.Create, FileAccess.ReadWrite, FileShare.Read);
-					LogFile = new StreamWriter(File);
-					LogFile.AutoFlush = true;
-				}
-				catch (Exception Ex)
-				{
-					if (Attempt == (MaxAttempts - 1))
-					{
-						UnrealBuildTool.Log.TraceWarning("Unable to create log file: {0}", LogFilename);
-						UnrealBuildTool.Log.TraceWarning(LogUtils.FormatException(Ex));
-					}
-				}
-			} while (LogFile == null && ++Attempt < MaxAttempts);
-		}
-
-		private void WriteToFile(string Message)
-		{
-			lock (SyncObject)
-			{
-				if (LogFile != null)
-				{
-					LogFile.WriteLine(Message);
-				}
-			}
-		}
-
-		/// <summary>
-		/// Writes a formatted line to the log.
-		/// </summary>
-		/// <param name="SkipStackFrames"></param>
-		/// <param name="Verbosity"></param>
-		/// <param name="Format"></param>
-		/// <param name="Args"></param>
-		[MethodImplAttribute(MethodImplOptions.NoInlining)]
-		public void WriteLine(string Source, TraceEventType Verbosity, string Format, params object[] Args)
-		{
-			var Message = LogUtils.FormatMessage(Source, Verbosity, Format, Args);
-			WriteToFile(Message);
-		}
-
-		/// <summary>
-		/// Writes a formatted line to the log.
-		/// </summary>
-		/// <param name="SkipStackFrames"></param>
-		/// <param name="Verbosity"></param>
-		/// <param name="Format"></param>
-		/// <param name="Args"></param>
-		[MethodImplAttribute(MethodImplOptions.NoInlining)]
-		public void WriteLine(string Source, TraceEventType Verbosity, string Message)
-		{
-			Message = LogUtils.FormatMessage(Source, Verbosity, Message);
-			WriteToFile(Message);
-		}
-
-		#region TraceListener Interface
-
-		public override void Flush()
-		{
-			lock (SyncObject)
-			{
-				if (LogFile != null)
-				{
-					LogFile.Flush();
-				}
-			}
-			base.Flush();
-		}
-
-		public override void Close()
-		{
-			lock (SyncObject)
-			{
-				if (LogFile != null)
-				{
-					LogFile.Flush();
-					LogFile.Close();
-					LogFile.Dispose();
-					LogFile = null;
-				}
-			}
-			try
-			{
-				// Try to copy the log file to the log folder. The reason why it's done here is that
-				// at the time the log file is being initialized the env var may not yet be set (this 
-				// applies to local log folder in particular)
-				var LogFolder = Environment.GetEnvironmentVariable(EnvVarNames.LogFolder);
-				if (!String.IsNullOrEmpty(LogFolder) && Directory.Exists(LogFolder) &&
-						!String.IsNullOrEmpty(LogFilename) && File.Exists(LogFilename))
-				{
-					var DestFilename = CommandUtils.CombinePaths(LogFolder, "UAT_" + Path.GetFileName(LogFilename));
-					SafeCopyLogFile(LogFilename, DestFilename);
-				}
-			}
-			catch (Exception)
-			{
-				// Silently ignore, logging is pointless because eveything is shut down at this point
-			}
-			base.Close();
-		}
-
-		/// <summary>
-		/// Copies log file to the final log folder, does multiple attempts if the destination file could not be created.
-		/// </summary>
-		/// <param name="SourceFilename"></param>
-		/// <param name="DestFilename"></param>
-		private void SafeCopyLogFile(string SourceFilename, string DestFilename)
-		{
-			const int MaxAttempts = 10;
-			int AttemptNo = 0;
-			var DestLogFilename = DestFilename;
-			bool Result = false;
-			do
-			{				
-				try
-				{
-					File.Copy(SourceFilename, DestLogFilename, true);
-					Result = true;
-				}
-				catch (Exception)
-				{
-					var ModifiedFilename = String.Format("{0}_{1}{2}", Path.GetFileNameWithoutExtension(DestFilename), AttemptNo, Path.GetExtension(DestLogFilename));
-					DestLogFilename = CommandUtils.CombinePaths(Path.GetDirectoryName(DestFilename), ModifiedFilename);
-					AttemptNo++;
-				}
-			}
-			while (Result == false && AttemptNo <= MaxAttempts);
-		}
-
-		public override void Write(string message)
-		{
-			WriteLine(String.Empty, TraceEventType.Information, message);
+            WriteLinePrivate(message);
 		}
 
 		public override void WriteLine(string message)
 		{
-			WriteLine(String.Empty, TraceEventType.Information, message);
-		}
-
-		public override void TraceEvent(TraceEventCache eventCache, string source, TraceEventType eventType, int id, string message)
-		{
-			WriteLine(source, TraceEventType.Verbose, message);
-		}
-
-		public override void TraceEvent(TraceEventCache eventCache, string source, TraceEventType eventType, int id, string format, params object[] args)
-		{
-			WriteLine(source, TraceEventType.Verbose, format, args);
+            WriteLinePrivate(message);
 		}
 
-		#endregion
+        #endregion
 	}
 
 	#endregion
diff --git a/Engine/Source/Programs/AutomationTool/ProcessUtils.cs b/Engine/Source/Programs/AutomationTool/ProcessUtils.cs
--- a/Engine/Source/Programs/AutomationTool/ProcessUtils.cs
+++ b/Engine/Source/Programs/AutomationTool/ProcessUtils.cs
@@ -8,6 +8,7 @@ using System.Threading;
 using System.Diagnostics;
 using System.Management;
 using System.Runtime.InteropServices;
+using System.Runtime.CompilerServices;
 
 namespace AutomationTool
 {
@@ -250,14 +251,16 @@ namespace AutomationTool
 			ProcessManager.RemoveProcess(this);
 		}
 
-		private void LogOutput(TraceEventType Verbosity, string Message)
+        /// <summary>
+        /// Log output of a remote process at a given severity.
+        /// To pretty up the output, we use a custom source so it will say the source of the process instead of this method name.
+        /// </summary>
+        /// <param name="Verbosity"></param>
+        /// <param name="Message"></param>
+        [MethodImplAttribute(MethodImplOptions.NoInlining)]
+        private void LogOutput(TraceEventType Verbosity, string Message)
 		{
-			// Manually send message to trace listeners (skips Class.Method source formatting of Log.WriteLine)
-			var EventCache = new TraceEventCache();
-			foreach (TraceListener Listener in Trace.Listeners)
-			{
-				Listener.TraceEvent(EventCache, Source, Verbosity, (int)Verbosity, Message);
-			}
+            UnrealBuildTool.Log.WriteLine(1, Source, Verbosity, Message);
 		}
        
 		/// <summary>
diff --git a/Engine/Source/Programs/AutomationTool/Program.cs b/Engine/Source/Programs/AutomationTool/Program.cs
--- a/Engine/Source/Programs/AutomationTool/Program.cs
+++ b/Engine/Source/Programs/AutomationTool/Program.cs
@@ -88,55 +88,62 @@ namespace AutomationTool
 			HostPlatform.Initialize();
 			
 			LogUtils.InitLogging(CommandLine);
-			Log.WriteLine(TraceEventType.Information, "Running on {0}", HostPlatform.Current.GetType().Name);
-
-			XmlConfigLoader.Init();
-
-			// Log if we're running from the launcher
-			var ExecutingAssemblyLocation = CommandUtils.CombinePaths(Assembly.GetExecutingAssembly().Location);
-			if (String.Compare(ExecutingAssemblyLocation, CommandUtils.CombinePaths(InternalUtils.ExecutingAssemblyLocation), true) != 0)
-			{
-				Log.WriteLine(TraceEventType.Information, "Executed from AutomationToolLauncher ({0})", ExecutingAssemblyLocation);
-			}
-			Log.WriteLine(TraceEventType.Information, "CWD={0}", Environment.CurrentDirectory);
-
-			// Hook up exit callbacks
-			var Domain = AppDomain.CurrentDomain;
-			Domain.ProcessExit += Domain_ProcessExit;
-			Domain.DomainUnload += Domain_ProcessExit;
-			HostPlatform.Current.SetConsoleCtrlHandler(ProgramCtrlHandler);
-
-			var Version = InternalUtils.ExecutableVersion;
-			Log.WriteLine(TraceEventType.Verbose, "{0} ver. {1}", Version.ProductName, Version.ProductVersion);
-
-			try
-			{
-				// Don't allow simultaneous execution of AT (in the same branch)
-				ReturnCode = InternalUtils.RunSingleInstance(MainProc, CommandLine);
-			}
-			catch (Exception Ex)
-			{
-				Log.WriteLine(TraceEventType.Error, "AutomationTool terminated with exception:");
-				Log.WriteLine(TraceEventType.Error, LogUtils.FormatException(Ex));
-				Log.WriteLine(TraceEventType.Error, Ex.Message);
-				if (ReturnCode == 0)
-				{
-					ReturnCode = (int)ErrorCodes.Error_Unknown;
-				}
-			}
-
-			// Make sure there's no directiories on the stack.
-			CommandUtils.ClearDirStack();
-			Environment.ExitCode = ReturnCode;
-
-			// Try to kill process before app domain exits to leave the other KillAll call to extreme edge cases
-			if (ShouldKillProcesses && !Utils.IsRunningOnMono)
-			{
-				ProcessManager.KillAll();
-			}
-
-			Log.WriteLine(TraceEventType.Information, "AutomationTool exiting with ExitCode={0}", ReturnCode);
-			LogUtils.CloseFileLogging();
+            try
+            {
+                Log.WriteLine(TraceEventType.Information, "Running on {0}", HostPlatform.Current.GetType().Name);
+
+                XmlConfigLoader.Init();
+
+                // Log if we're running from the launcher
+                var ExecutingAssemblyLocation = CommandUtils.CombinePaths(Assembly.GetExecutingAssembly().Location);
+                if (String.Compare(ExecutingAssemblyLocation, CommandUtils.CombinePaths(InternalUtils.ExecutingAssemblyLocation), true) != 0)
+                {
+                    Log.WriteLine(TraceEventType.Information, "Executed from AutomationToolLauncher ({0})", ExecutingAssemblyLocation);
+                }
+                Log.WriteLine(TraceEventType.Information, "CWD={0}", Environment.CurrentDirectory);
+
+                // Hook up exit callbacks
+                var Domain = AppDomain.CurrentDomain;
+                Domain.ProcessExit += Domain_ProcessExit;
+                Domain.DomainUnload += Domain_ProcessExit;
+                HostPlatform.Current.SetConsoleCtrlHandler(ProgramCtrlHandler);
+
+                var Version = InternalUtils.ExecutableVersion;
+                Log.WriteLine(TraceEventType.Verbose, "{0} ver. {1}", Version.ProductName, Version.ProductVersion);
+
+                try
+                {
+                    // Don't allow simultaneous execution of AT (in the same branch)
+                    ReturnCode = InternalUtils.RunSingleInstance(MainProc, CommandLine);
+                }
+                catch (Exception Ex)
+                {
+                    Log.WriteLine(TraceEventType.Error, "AutomationTool terminated with exception:");
+                    Log.WriteLine(TraceEventType.Error, LogUtils.FormatException(Ex));
+                    Log.WriteLine(TraceEventType.Error, Ex.Message);
+                    if (ReturnCode == 0)
+                    {
+                        ReturnCode = (int)ErrorCodes.Error_Unknown;
+                    }
+                }
+
+                // Make sure there's no directiories on the stack.
+                CommandUtils.ClearDirStack();
+                Environment.ExitCode = ReturnCode;
+
+                // Try to kill process before app domain exits to leave the other KillAll call to extreme edge cases
+                if (ShouldKillProcesses && !Utils.IsRunningOnMono)
+                {
+                    ProcessManager.KillAll();
+                }
+
+                Log.WriteLine(TraceEventType.Information, "AutomationTool exiting with ExitCode={0}", ReturnCode);
+            }
+            finally
+            {
+                // ensure that logging is always shut down cleanly, as this potentially copies the log to a final resting place.
+                LogUtils.ShutdownLogging();
+            }
 
 			return ReturnCode;
 		}
diff --git a/Engine/Source/Programs/UnrealBuildTool/Configuration/BuildConfiguration.cs b/Engine/Source/Programs/UnrealBuildTool/Configuration/BuildConfiguration.cs
--- a/Engine/Source/Programs/UnrealBuildTool/Configuration/BuildConfiguration.cs
+++ b/Engine/Source/Programs/UnrealBuildTool/Configuration/BuildConfiguration.cs
@@ -10,6 +10,14 @@ namespace UnrealBuildTool
 {
 	public class BuildConfiguration
 	{
+        static BuildConfiguration()
+        {
+            if (!UnrealBuildTool.bIsSafeToReferenceConfigurationValues)
+            {
+                throw new BuildException("UEBuildConfiguration was referenced before the XmlConfig files could be loaded.");
+            }
+        }
+
 		/// <summary>
 		/// Whether to unify C++ code into larger files for faster compilation.
 		/// </summary>
@@ -116,11 +124,17 @@ namespace UnrealBuildTool
 		[XmlConfig]
 		public static bool bForcePrecompiledHeaderForGameModules;
 
-		/// <summary>
+        /// <summary>
 		/// Whether debug info should be written to the console.
-		/// </summary>
-		[XmlConfig]
-		public static bool bPrintDebugInfo;
+        /// </summary>
+        [XmlConfig]
+        public static bool bPrintDebugInfo;
+
+        /// <summary>
+        /// Allows logging to a file
+        /// </summary>
+        [XmlConfig]
+        public static string LogFilename;
 
 		/// <summary>
 		/// Prints performance diagnostics about include dependencies and other bits
diff --git a/Engine/Source/Programs/UnrealBuildTool/Configuration/UEBuildConfiguration.cs b/Engine/Source/Programs/UnrealBuildTool/Configuration/UEBuildConfiguration.cs
--- a/Engine/Source/Programs/UnrealBuildTool/Configuration/UEBuildConfiguration.cs
+++ b/Engine/Source/Programs/UnrealBuildTool/Configuration/UEBuildConfiguration.cs
@@ -7,6 +7,14 @@ namespace UnrealBuildTool
 {
 	public class UEBuildConfiguration
 	{
+        static UEBuildConfiguration()
+        {
+            if (!UnrealBuildTool.bIsSafeToReferenceConfigurationValues)
+            {
+                throw new BuildException("UEBuildConfiguration was referenced before the XmlConfig files could be loaded.");
+            }
+        }
+
 		/** Whether to include PhysX support */
 		[XmlConfig]
 		public static bool bCompilePhysX;
diff --git a/Engine/Source/Programs/UnrealBuildTool/System/UnrealBuildTool.cs b/Engine/Source/Programs/UnrealBuildTool/System/UnrealBuildTool.cs
--- a/Engine/Source/Programs/UnrealBuildTool/System/UnrealBuildTool.cs
+++ b/Engine/Source/Programs/UnrealBuildTool/System/UnrealBuildTool.cs
@@ -582,66 +582,59 @@ namespace UnrealBuildTool
         }
 
         /// <summary>
-        /// Sets up UBT Trace listeners and filters.
+        /// This is a safeguard to our initialization order, which tends to be fragile.
+        /// It is not really safe to reference a configuration value until the XmlConfig system loads it (and command lines override it), but non-trivial
+        /// code needs to run before the XmlConfig can load, so it can be hard to tell if a function you are calling in these early 
+        /// stages might unsafely reference a config value.
+        /// 
+        /// This bool is set RIGHT BEFORE the XmlConfig system loads configs, so any references to the config class before that will be trapped.
+        /// If you think you need to change where this value is set, you need to reconsider how soon you are trying to reference config values.
         /// </summary>
-        private static void InitLogging()
+        public static bool bIsSafeToReferenceConfigurationValues = false;
+
+        /// <summary>
+        /// Paranoid init of the logging system in case someone tries to log super early in startup.
+        /// However, no configuration files have been loaded, and the command line has not been parsed, so 
+        /// certain features like verbose and log file cannot be initialized yet.
+        /// </summary>
+        private static void PreInitLogging()
         {
-            Trace.Listeners.Add(new ConsoleListener());
-            var Filter = new VerbosityFilter();
-            foreach (TraceListener Listener in Trace.Listeners)
-            {
-                Listener.Filter = Filter;
-            }
+            Log.InitLogging(
+                bLogTimestamps: false,
+                bLogVerbose: true,
+                bLogSeverity: false,
+                bLogSources: false,
+                bColorConsoleOutput: true,
+                TraceListeners: new[] { new ConsoleTraceListener() });
         }
 
-        private static int Main(string[] Arguments)
+        /// <summary>
+        /// UBT startup order is fairly fragile, and relies on globals that may or may not be safe to use yet.
+        /// This function is for super early startup stuff that should not access Configuration classes (anything loaded by XmlConfg).
+        /// This should be very minimal startup code.
+        /// </summary>
+        /// <param name="Arguments">Cmdline arguments</param>
+        private static void DoStartupStuffThatCannotAccessConfigurationClasses(string[] Arguments)
         {
-            InitLogging();
-
-            try
-            {
-                XmlConfigLoader.Init();
-            }
-            catch (BuildException Exception)
-            {
-                Log.TraceError("UnrealBuildTool Exception: " + Exception);
-                return (int) ECompilationResult.OtherCompilationError;
-            }
-
-			// Because XmlConfigLoader.Init() will reset all parameter by call LoadDefaults, 
-			// so need to parse verbose argument here.
-			int Verbose;
-			if (Utils.ParseCommandLineFlag(Arguments, "-verbose", out Verbose))
-			{
-				BuildConfiguration.bPrintDebugInfo = true;
-			}
-
-			Log.TraceVerbose(Environment.CommandLine);
-
-            InitialEnvironment = Environment.GetEnvironmentVariables();
-            if (InitialEnvironment.Count < 1)
-            {
-                throw new BuildException("Environment could not be read");
-            }
             // Helpers used for stats tracking.
             StartTime = DateTime.UtcNow;
 
-            Telemetry.Initialize();
+            // Do super early log init as a safeguard. We'll re-init with proper config options later.
+            PreInitLogging();
 
             // Copy off the arguments to allow checking for command-line arguments elsewhere
             CmdLine = new List<string>(Arguments);
 
-            foreach (string Arg in Arguments)
+            // Grab the environment.
+            InitialEnvironment = Environment.GetEnvironmentVariables();
+            if (InitialEnvironment.Count < 1)
             {
-                string TempGameName = null;
-                if (ParseRocketCommandlineArg(Arg, ref TempGameName) == true)
-                {
-                    // This is to allow relative paths for the project file
-                    Log.TraceVerbose("UBT Running for Rocket: " + UProjectFile);
-                }
+                throw new BuildException("Environment could not be read");
             }
 
-            // Change the working directory to be the Engine/Source folder. We are running from Engine/Binaries/DotNET
+            // Change the working directory to be the Engine/Source folder. We are likely running from Engine/Binaries/DotNET
+            // This is critical to be done early so any code that relies on the current directory being Engine/Source will work.
+            // UEBuildConfiguration.PostReset is susceptible to this, so we must do this before configs are loaded.
             string EngineSourceDirectory = Path.Combine(Utils.GetExecutingAssemblyDirectory(), "..", "..", "..", "Engine", "Source");
 
             [REDACTED_EMAIL]: This is a workaround for recompiling game code in editor
@@ -662,6 +655,152 @@ namespace UnrealBuildTool
             {
                 Directory.SetCurrentDirectory(EngineSourceDirectory);
             }
+        }
+
+        /// <summary>
+        /// Enough startup work has been done to load configuration classes (like setting the proper current directory).
+        /// It is still not safe to log. Mostly this relates to XmlConfigLoader, but any other code added here cannot log yet either.
+        /// </summary>
+        /// <param name="Arguments"></param>
+        private static void InitConfigurationClasses(string[] Arguments)
+        {
+            // Load the Config XML files first.
+            XmlConfigLoader.Init();
+            
+            // Then let the command lines override any configs necessary.
+            const string UCAUsageThresholdString = "-ucausagethreshold=";
+            const string UCAModuleToAnalyzeString = "-ucamoduletoanalyze=";
+            foreach (var LowercaseArg in Arguments.Select(arg => arg.ToLowerInvariant()))
+            {
+                if (LowercaseArg == "-verbose")
+                {
+                    BuildConfiguration.bPrintDebugInfo = true;
+                }
+                else if (LowercaseArg.StartsWith("-log="))
+                {
+                    BuildConfiguration.LogFilename = LowercaseArg.Replace("-log=","");
+                }
+                else if (LowercaseArg == "-forceheadergeneration")
+                {
+                    UEBuildConfiguration.bForceHeaderGeneration = true;
+                }
+                else if (LowercaseArg == "-nobuilduht")
+                {
+                    UEBuildConfiguration.bDoNotBuildUHT = true;
+                }
+                else if (LowercaseArg == "-failifgeneratedcodechanges")
+                {
+                    UEBuildConfiguration.bFailIfGeneratedCodeChanges = true;
+                }
+                else if (LowercaseArg == "-clean")
+                {
+                    UEBuildConfiguration.bCleanProject = true;
+                }
+                else if (LowercaseArg == "-prepfordeploy")
+                {
+                    UEBuildConfiguration.bPrepForDeployment = true;
+                }
+                else if (LowercaseArg == "-generateexternalfilelist")
+                {
+                    UEBuildConfiguration.bGenerateExternalFileList = true;
+                }
+                else if (LowercaseArg == "-mergeexternalfilelist")
+                {
+                    UEBuildConfiguration.bMergeExternalFileList = true;
+                }
+                else if (LowercaseArg == "-generatemanifest")
+                {
+                    // Generate a manifest file containing all the files required to be in Perforce
+                    UEBuildConfiguration.bGenerateManifest = true;
+                }
+                else if (LowercaseArg == "-mergemanifests")
+                {
+                    // Whether to add to the existing manifest (if it exists), or start afresh
+                    UEBuildConfiguration.bMergeManifests = true;
+                }
+                else if (LowercaseArg == "-noxge")
+                {
+                    BuildConfiguration.bAllowXGE = false;
+                }
+                else if (LowercaseArg == "-noxgemonitor")
+                {
+                    BuildConfiguration.bShowXGEMonitor = false;
+                }
+                else if (LowercaseArg == "-xgeexport")
+                {
+                    BuildConfiguration.bXGEExport = true;
+                    BuildConfiguration.bAllowXGE = true;
+                }
+                else if (LowercaseArg == "-noubtmakefiles")
+                {
+                    BuildConfiguration.bUseUBTMakefiles = false;
+                }
+                else if (LowercaseArg == "-nosimplygon")
+                {
+                    UEBuildConfiguration.bCompileSimplygon = false;
+                }
+                else if (LowercaseArg == "-nospeedtree")
+                {
+                    UEBuildConfiguration.bCompileSpeedTree = false;
+                }
+                else if (LowercaseArg == "-canskiplink")
+                {
+                    UEBuildConfiguration.bSkipLinkingWhenNothingToCompile = true;
+                }
+                else if (LowercaseArg == "-nocef3")
+                {
+                    UEBuildConfiguration.bCompileCEF3 = false;
+                }
+                else if (LowercaseArg == "-rununrealcodeanalyzer")
+                {
+                    BuildConfiguration.bRunUnrealCodeAnalyzer = true;
+                }
+                else if (LowercaseArg.StartsWith(UCAUsageThresholdString))
+                {
+                    string UCAUsageThresholdValue = LowercaseArg.Substring(UCAUsageThresholdString.Length);
+                    BuildConfiguration.UCAUsageThreshold = float.Parse(UCAUsageThresholdValue.Replace(',', '.'), System.Globalization.CultureInfo.InvariantCulture);
+                }
+                else if (LowercaseArg.StartsWith(UCAModuleToAnalyzeString))
+                {
+                    BuildConfiguration.UCAModuleToAnalyze = LowercaseArg.Substring(UCAModuleToAnalyzeString.Length).Trim();
+                }
+            }
+        }
+
+        /// <summary>
+        /// Sets up UBT Trace listeners and filters.
+        /// </summary>
+        private static void InitLogging()
+        {
+            // Initialize the shared log system using UBT specific listeners.
+            Log.InitLogging(
+                bLogTimestamps: false,
+                bLogVerbose: BuildConfiguration.bPrintDebugInfo,
+                bLogSeverity: false, 
+                bLogSources: false,
+                bColorConsoleOutput: true,
+                TraceListeners: new[] 
+                {
+                    new ConsoleTraceListener(),
+                    !string.IsNullOrEmpty(BuildConfiguration.LogFilename) ? new TextWriterTraceListener(new StreamWriter(new FileStream(BuildConfiguration.LogFilename, FileMode.Create, FileAccess.ReadWrite, FileShare.Read)) { AutoFlush = true }) : null,
+                });
+        }
+
+        private static int DoPostStartupStuffThatCanAccessConfigs(string[] Arguments)
+        {
+            // We can now do a full initialization of the logging system with proper configuration values.
+            InitLogging();
+
+            // Parse rocket-specific arguments.
+            foreach (string Arg in Arguments)
+            {
+                string TempGameName = null;
+                if (ParseRocketCommandlineArg(Arg, ref TempGameName) == true)
+                {
+                    // This is to allow relative paths for the project file
+                    Log.TraceVerbose("UBT Running for Rocket: " + UProjectFile);
+                }
+            }
 
             // Build the list of game projects that we know about. When building from the editor (for hot-reload) or for Rocket projects, we require the 
             // project file to be passed in. Otherwise we scan for projects in directories named in UE4Games.uprojectdirs.
@@ -680,18 +819,11 @@ namespace UnrealBuildTool
 
             ECompilationResult Result = ECompilationResult.Succeeded;
 
-            // Reset early so we can access BuildConfiguration even before RunUBT() is called
-            XmlConfigLoader.Reset<BuildConfiguration>();
-            XmlConfigLoader.Reset<UEBuildConfiguration>();
-
-			if (Utils.ParseCommandLineFlag(Arguments, "-verbose", out Verbose))
-			{
-				BuildConfiguration.bPrintDebugInfo = true;
-			}
-
             Log.TraceVerbose("UnrealBuildTool (DEBUG OUTPUT MODE)");
             Log.TraceVerbose("Command-line: {0}", String.Join(" ", Arguments));
 
+            Telemetry.Initialize();
+
             bool bRunCopyrightVerification = false;
             bool bDumpToFile = false;
             bool bCheckThirdPartyHeaders = false;
@@ -793,8 +925,6 @@ namespace UnrealBuildTool
                     {
                         string LowercaseArg = Arg.ToLowerInvariant();
                         const string OnlyPlatformSpecificForArg = "-onlyplatformspecificfor=";
-						const string UCAUsageThresholdString = "-ucausagethreshold=";
-						const string UCAModuleToAnalyzeString = "-ucamoduletoanalyze=";
                         if (ParseRocketCommandlineArg(Arg, ref GameName))
                         {
                             // Already handled at startup. Calling now just to properly set the game name
@@ -857,19 +987,6 @@ namespace UnrealBuildTool
                             bSpecificModulesOnly = true;
                             continue;
                         }
-                        else if (LowercaseArg == "-forceheadergeneration")
-                        {
-                            // Don't reset this as we don't recheck the args in RunUBT
-                            UEBuildConfiguration.bForceHeaderGeneration = true;
-                        }
-                        else if (LowercaseArg == "-nobuilduht")
-                        {
-                            UEBuildConfiguration.bDoNotBuildUHT = true;
-                        }
-                        else if (LowercaseArg == "-failifgeneratedcodechanges")
-                        {
-                            UEBuildConfiguration.bFailIfGeneratedCodeChanges = true;
-                        }
                         else if (LowercaseArg == "-copyrightverification")
                         {
                             bRunCopyrightVerification = true;
@@ -882,45 +999,6 @@ namespace UnrealBuildTool
                         {
                             bCheckThirdPartyHeaders = true;
                         }
-                        else if (LowercaseArg == "-clean")
-                        {
-                            UEBuildConfiguration.bCleanProject = true;
-                        }
-                        else if (LowercaseArg == "-prepfordeploy")
-                        {
-                            UEBuildConfiguration.bPrepForDeployment = true;
-                        }
-                        else if (LowercaseArg == "-generateexternalfilelist")
-                        {
-                            UEBuildConfiguration.bGenerateExternalFileList = true;
-                        }
-                        else if (LowercaseArg == "-mergeexternalfilelist")
-                        {
-                            UEBuildConfiguration.bMergeExternalFileList = true;
-                        }
-                        else if (LowercaseArg == "-generatemanifest")
-                        {
-                            // Generate a manifest file containing all the files required to be in Perforce
-                            UEBuildConfiguration.bGenerateManifest = true;
-                        }
-                        else if (LowercaseArg == "-mergemanifests")
-                        {
-                            // Whether to add to the existing manifest (if it exists), or start afresh
-                            UEBuildConfiguration.bMergeManifests = true;
-                        }
-                        else if (LowercaseArg == "-noxge")
-                        {
-                            BuildConfiguration.bAllowXGE = false;
-                        }
-                        else if (LowercaseArg == "-xgeexport")
-                        {
-                            BuildConfiguration.bXGEExport = true;
-                            BuildConfiguration.bAllowXGE = true;
-                        }
-						else if (LowercaseArg == "-noubtmakefiles")
-						{
-							BuildConfiguration.bUseUBTMakefiles = false;
-						}
 						else if (LowercaseArg == "-invalidatemakefilesonly")
 						{
 							UnrealBuildTool.bIsInvalidatingMakefilesOnly = true;
@@ -951,14 +1029,6 @@ namespace UnrealBuildTool
                             UnrealBuildTool.bIsGatheringBuild_Unsafe = false;
                             UnrealBuildTool.bIsAssemblingBuild_Unsafe = true;
                         }
-                        else if (LowercaseArg == "-nosimplygon")
-                        {
-                            UEBuildConfiguration.bCompileSimplygon = false;
-                        }
-                        else if (LowercaseArg == "-nospeedtree")
-                        {
-                            UEBuildConfiguration.bCompileSpeedTree = false;
-                        }
                         else if (LowercaseArg == "-ignorejunk")
                         {
                             bIgnoreJunk = true;
@@ -971,27 +1041,6 @@ namespace UnrealBuildTool
                         {
                             ProgressWriter.bWriteMarkup = true;
                         }
-						else if (LowercaseArg == "-canskiplink")
-						{
-							UEBuildConfiguration.bSkipLinkingWhenNothingToCompile = true;
-						}
-						else if (LowercaseArg == "-nocef3")
-						{
-							UEBuildConfiguration.bCompileCEF3 = false;
-						}
-						else if (LowercaseArg == "-rununrealcodeanalyzer")
-						{
-							BuildConfiguration.bRunUnrealCodeAnalyzer = true;
-						}
-						else if (LowercaseArg.StartsWith(UCAUsageThresholdString))
-						{
-							string UCAUsageThresholdValue = LowercaseArg.Substring(UCAUsageThresholdString.Length);
-							BuildConfiguration.UCAUsageThreshold = float.Parse(UCAUsageThresholdValue.Replace(',', '.'), System.Globalization.CultureInfo.InvariantCulture);
-						}
-						else if (LowercaseArg.StartsWith(UCAModuleToAnalyzeString))
-						{
-							BuildConfiguration.UCAModuleToAnalyze = LowercaseArg.Substring(UCAModuleToAnalyzeString.Length).Trim();
-						}
                         else if (CheckPlatform.ToString().ToLowerInvariant() == LowercaseArg)
                         {
                             // It's the platform set...
@@ -1112,7 +1161,8 @@ namespace UnrealBuildTool
                     }
                     else if (bRunCopyrightVerification)
                     {
-                        CopyrightVerify.RunCopyrightVerification(EngineSourceDirectory, GameName, bDumpToFile);
+                        // We ensure we are in the engine source directory during startup.
+                        CopyrightVerify.RunCopyrightVerification(Directory.GetCurrentDirectory(), GameName, bDumpToFile);
                         Log.TraceInformation("Completed... exiting.");
                     }
                     else if (bAutoSDKOnly)
@@ -1227,6 +1277,33 @@ namespace UnrealBuildTool
             return (int)Result; 
         }
 
+        private static int Main(string[] Arguments)
+        {
+            // make sure we catch any exceptions and return an appropriate error code.
+            // Some inner code already does this (to ensure the Mutex is released),
+            // but we need something to cover all outer code as well.
+            try
+            {
+                // Make it more explicit what startup code should not be accessing config data.
+                // The Configuration classes will also assert, but this clarifies where the line is.
+                DoStartupStuffThatCannotAccessConfigurationClasses(Arguments);
+
+                // Now we can load the configuration files.
+                InitConfigurationClasses(Arguments);
+
+                // Finally, we do execute code that relies on config classes.
+                return DoPostStartupStuffThatCanAccessConfigs(Arguments);
+            }
+            catch (Exception Exception)
+            {
+                if (Log.IsInitialized())
+                {
+                    Log.TraceError("UnrealBuildTool Exception: " + Exception.ToString());
+                }
+                return (int)ECompilationResult.OtherCompilationError;
+            }
+        }
+
         /// <summary>
         /// Generates project files.  Can also be used to generate projects "in memory", to allow for building
         /// targets without even having project files at all.
@@ -1342,8 +1419,6 @@ namespace UnrealBuildTool
             // Reset global configurations
             ActionGraph.ResetAllActions();
 
-            ParseBuildConfigurationFlags(Arguments);
-
             // We need to allow the target platform to perform the 'reset' as well...
             UnrealTargetPlatform ResetPlatform = UnrealTargetPlatform.Unknown;
             UnrealTargetConfiguration ResetConfiguration;
@@ -1957,106 +2032,6 @@ namespace UnrealBuildTool
 			return bIsRunning;
 		}
 
-        private static void ParseBuildConfigurationFlags(string[] Arguments)
-        {
-            int ArgumentIndex = 0;
-
-            // Parse optional command-line flags.
-            if (Utils.ParseCommandLineFlag(Arguments, "-verbose", out ArgumentIndex))
-            {
-                BuildConfiguration.bPrintDebugInfo = true;
-            }
-
-            if (Utils.ParseCommandLineFlag(Arguments, "-noxge", out ArgumentIndex))
-            {
-                BuildConfiguration.bAllowXGE = false;
-            }
-
-            if (Utils.ParseCommandLineFlag(Arguments, "-noxgemonitor", out ArgumentIndex))
-            {
-                BuildConfiguration.bShowXGEMonitor = false;
-            }
-
-            if (Utils.ParseCommandLineFlag(Arguments, "-stresstestunity", out ArgumentIndex))
-            {
-                BuildConfiguration.bStressTestUnity = true;
-            }
-
-            if (Utils.ParseCommandLineFlag(Arguments, "-disableunity", out ArgumentIndex))
-            {
-                BuildConfiguration.bUseUnityBuild = false;
-            }
-            if (Utils.ParseCommandLineFlag(Arguments, "-forceunity", out ArgumentIndex))
-            {
-                BuildConfiguration.bForceUnityBuild = true;
-            }
-
-            // New Monolithic Graphics drivers have optional "fast calls" replacing various D3d functions
-            if (Utils.ParseCommandLineFlag(Arguments, "-fastmonocalls", out ArgumentIndex))
-            {
-                BuildConfiguration.bUseFastMonoCalls = true;
-            }
-            if (Utils.ParseCommandLineFlag(Arguments, "-nofastmonocalls", out ArgumentIndex))
-            {
-                BuildConfiguration.bUseFastMonoCalls = false;
-            }
-            
-            if (Utils.ParseCommandLineFlag(Arguments, "-uniqueintermediate", out ArgumentIndex))
-            {
-                BuildConfiguration.BaseIntermediateFolder = "Intermediate/Build/" + BuildGuid + "/";
-            }
-
-            if (Utils.ParseCommandLineFlag(Arguments, "-nopch", out ArgumentIndex))
-            {
-                BuildConfiguration.bUsePCHFiles = false;
-            }
-
-            if (Utils.ParseCommandLineFlag(Arguments, "-nosharedpch", out ArgumentIndex))
-            {
-                BuildConfiguration.bUseSharedPCHs = false;
-            }
-
-            if (Utils.ParseCommandLineFlag(Arguments, "-skipActionHistory", out ArgumentIndex))
-            {
-                BuildConfiguration.bUseActionHistory = false;
-            }
-
-            if (Utils.ParseCommandLineFlag(Arguments, "-noLTCG", out ArgumentIndex))
-            {
-                BuildConfiguration.bAllowLTCG = false;
-            }
-
-            if (Utils.ParseCommandLineFlag(Arguments, "-nopdb", out ArgumentIndex))
-            {
-                BuildConfiguration.bUsePDBFiles = false;
-            }
-
-            if (Utils.ParseCommandLineFlag(Arguments, "-deploy", out ArgumentIndex))
-            {
-                BuildConfiguration.bDeployAfterCompile = true;
-            }
-
-            if (Utils.ParseCommandLineFlag(Arguments, "-CopyAppBundleBackToDevice", out ArgumentIndex))
-            {
-                BuildConfiguration.bCopyAppBundleBackToDevice = true;
-            }
-
-            if (Utils.ParseCommandLineFlag(Arguments, "-flushmac", out ArgumentIndex))
-            {
-                BuildConfiguration.bFlushBuildDirOnRemoteMac = true;
-            }
-
-            if (Utils.ParseCommandLineFlag(Arguments, "-nodebuginfo", out ArgumentIndex))
-            {
-                BuildConfiguration.bDisableDebugInfo = true;
-            }
-            else if (Utils.ParseCommandLineFlag(Arguments, "-forcedebuginfo", out ArgumentIndex))
-            {
-                BuildConfiguration.bDisableDebugInfo = false;
-                BuildConfiguration.bOmitPCDebugInfoInDevelopment = false;
-            }
-        }
-
         /**
          * Parses the passed in command line for build configuration overrides.
          * 
@@ -2068,19 +2043,6 @@ namespace UnrealBuildTool
             var TargetSettings = new List<string[]>();
             int ArgumentIndex = 0;
 
-            // Log command-line arguments.
-            if (BuildConfiguration.bPrintDebugInfo)
-            {
-                Console.Write("Command-line arguments: ");
-                foreach (string Argument in Arguments)
-                {
-                    Console.Write("{0} ", Argument);
-                }
-                Log.TraceInformation("");
-            }
-
-            ParseBuildConfigurationFlags(Arguments);
-
             if (Utils.ParseCommandLineFlag(Arguments, "-targets", out ArgumentIndex))
             {
                 if (ArgumentIndex + 1 >= Arguments.Length)
diff --git a/Engine/Source/Programs/UnrealBuildTool/System/Utils.cs b/Engine/Source/Programs/UnrealBuildTool/System/Utils.cs
--- a/Engine/Source/Programs/UnrealBuildTool/System/Utils.cs
+++ b/Engine/Source/Programs/UnrealBuildTool/System/Utils.cs
@@ -1052,263 +1052,317 @@ namespace UnrealBuildTool
 		}
 	}
 
-	/// <summary>
-	/// Verbosity filter
-	/// </summary>
-	class VerbosityFilter : TraceFilter
-	{
-		public override bool ShouldTrace(TraceEventCache Cache, string Source, TraceEventType EventType, int Id, string FormatOrMessage, object[] Args, object Data1, object[] Data)
-		{
-			return EventType < TraceEventType.Verbose || BuildConfiguration.bPrintDebugInfo;
-		}
-	}
-
-	/// <summary>
-	/// Console Trace listener for UBT
-	/// </summary>
-	class ConsoleListener : TextWriterTraceListener
-	{
-		public override void Write(string Message)
-		{
-			Console.Write(Message);
-		}
-
-		public override void WriteLine(string Message)
-		{
-			Console.WriteLine(Message);
-		}
-
-		public override void TraceEvent(TraceEventCache EventCache, string Source, TraceEventType EventType, int Id, string Message)
-		{
-			if (Filter == null || Filter.ShouldTrace(EventCache, Source, EventType, Id, Message, null, null, null))
-			{
-				WriteLine(Message);
-			}
-		}
-
-		public override void TraceEvent(TraceEventCache EventCache, string Source, TraceEventType EventType, int Id, string Format, params object[] Args)
-		{
-			if (Filter == null || Filter.ShouldTrace(EventCache, Source, EventType, Id, Format, Args, null, null))
-			{
-				WriteLine(String.Format(Format, Args));
-			}
-		}
-	}
-
-	/// <summary>
-	/// UnrealBuiltTool console logging system.
-	/// </summary>
-	public sealed class Log
-	{
-		[MethodImplAttribute(MethodImplOptions.NoInlining)]
-		private static string GetSource(int StackFramesToSkip)
-		{
-			StackFrame Frame = new StackFrame(2 + StackFramesToSkip);
-			System.Reflection.MethodBase Method = Frame.GetMethod();
-			return String.Format("{0}.{1}", Method.DeclaringType.Name, Method.Name);
-		}
+    /// <summary>
+    /// UAT/UBT Custom log system.
+    /// 
+    /// This lets you use any TraceListeners you want, but you should only call the static 
+    /// methods below, not call Trace.XXX directly, as the static methods
+    /// This allows the system to enforce the formatting and filtering conventions we desire.
+    ///
+    /// For posterity, we cannot use the Trace or TraceSource class directly because of our special log requirements:
+    ///   1. We possibly capture the method name of the logging event. This cannot be done as a macro, so must be done at the top level so we know how many layers of the stack to peel off to get the real function.
+    ///   2. We have a verbose filter we would like to apply to all logs without having to have each listener filter individually, which would require our string formatting code to run every time.
+    ///   3. We possibly want to ensure severity prefixes are logged, but Trace.WriteXXX does not allow any severity info to be passed down.
+    /// </summary>
+    static public class Log
+    {
+        /// <summary>
+        /// Guard our initialization. Mainly used by top level exception handlers to ensure its safe to call a logging function.
+        /// In general user code should not concern itself checking for this.
+        /// </summary>
+        private static bool bIsInitialized = false;
+        /// <summary>
+        /// When true, verbose loggin is enabled.
+        /// </summary>
+        private static bool bLogVerbose = false;
+        /// <summary>
+        /// When true, warnings and errors will have a WARNING: or ERROR: prexifx, respectively.
+        /// </summary>
+        private static bool bLogSeverity = false;
+        /// <summary>
+        /// When true, logs will have the calling mehod prepended to the output as MethodName:
+        /// </summary>
+        private static bool bLogSources = false;
+        /// <summary>
+        /// When true, will detect warnings and errors and set the console output color to yellow and red.
+        /// </summary>
+        private static bool bColorConsoleOutput = false;
+        /// <summary>
+        /// When configured, this tracks time since initialization to prepend a timestamp to each log.
+        /// </summary>
+        private static Stopwatch Timer;
 
-		/// <summary>
-		/// Writes a formatted message to the console.
-		/// </summary>
-		/// <param name="Verbosity">Message verbosity level.</param>
-		/// <param name="Format">Message format string.</param>
-		/// <param name="Args">Optional arguments</param>
-		[MethodImplAttribute(MethodImplOptions.NoInlining)]
-		public static void WriteLine(int StackFramesToSkip, TraceEventType Verbosity, string Format, params object[] Args)
-		{
-			var Source = GetSource(StackFramesToSkip);
-			var EventCache = new TraceEventCache();
-			foreach (TraceListener Listener in Trace.Listeners)
-			{
-				Listener.TraceEvent(EventCache, Source, Verbosity, (int)Verbosity, Format, Args);
-			}
-		}
+        /// <summary>
+        /// Expose the log level. This is a hack for ProcessResult.LogOutput, which wants to bypass our normal formatting scheme.
+        /// </summary>
+        public static bool bIsVerbose { get { return bLogVerbose; } }
 
-		/// <summary>
-		/// Writes a message to the console.
-		/// </summary>
-		/// <param name="Verbosity">Message verbosity level.</param>
-		/// <param name="Message">Message text.</param>
-		[MethodImplAttribute(MethodImplOptions.NoInlining)]
-		public static void WriteLine(int StackFramesToSkip, TraceEventType Verbosity, string Message)
-		{
-			var Source = GetSource(StackFramesToSkip);
-			var EventCache = new TraceEventCache();
-			foreach (TraceListener Listener in Trace.Listeners)
-			{
-				Listener.TraceEvent(EventCache, Source, Verbosity, (int)Verbosity, Message);
-			}
-		}
+        /// <summary>
+        /// Allows code to check if the log system is ready yet.
+        /// End users should NOT need to use this. It pretty much exists
+        /// to work around startup issues since this is a global singleton.
+        /// </summary>
+        /// <returns></returns>
+        public static bool IsInitialized() 
+        { 
+            return bIsInitialized; 
+        }
 
-		/// <summary>
-		/// Writes a formatted message to the console.
-		/// </summary>
-		/// <param name="Verbosity">Message verbosity level.</param>
-		/// <param name="Format">Message format string.</param>
-		/// <param name="Args">Optional arguments</param>
-		[MethodImplAttribute(MethodImplOptions.NoInlining)]
-		public static void WriteLine(TraceEventType Verbosity, string Format, params object[] Args)
-		{
-			var Source = GetSource(0);
-			var EventCache = new TraceEventCache();
-			foreach (TraceListener Listener in Trace.Listeners)
-			{
-				Listener.TraceEvent(EventCache, Source, Verbosity, (int)Verbosity, Format, Args);
-			}
-		}
+        /// <summary>
+        /// Allows us to change verbosity after initializing. This can happen since we initialize logging early, 
+        /// but then read the config and command line later, which could change this value.
+        /// </summary>
+        /// <param name="bLogVerbose">Whether to log verbose logs.</param>
+        public static void SetVerboseLogging(bool bLogVerbose)
+        {
+            Log.bLogVerbose = bLogVerbose;
+        }
 
-		/// <summary>
-		/// Writes a message to the console.
-		/// </summary>
-		/// <param name="Verbosity">Message verbosity level.</param>
-		/// <param name="Message">Message text.</param>
-		[MethodImplAttribute(MethodImplOptions.NoInlining)]
-		public static void WriteLine(TraceEventType Verbosity, string Message)
-		{
-			var Source = GetSource(0);
-			var EventCache = new TraceEventCache();
-			foreach (TraceListener Listener in Trace.Listeners)
+        /// <summary>
+        /// This class allows InitLogging to be called more than once to work around chicken and eggs issues with logging and parsing command lines (see UBT startup code).
+        /// </summary>
+        /// <param name="bLogTimestamps">If true, the timestamp from Log init time will be prepended to all logs.</param>
+        /// <param name="bLogVerbose">If true, any Verbose log method method will not be ignored.</param>
+        /// <param name="bLogSeverity">If true, warnings and errors will have a WARNING: and ERROR: prefix to them. </param>
+        /// <param name="bLogSources">If true, logs will have the originating method name prepended to them.</param>
+        /// <param name="TraceListeners">Collection of trace listeners to attach to the Trace.Listeners, in addition to the Default listener. The existing listeners (except the Default listener) are cleared first.</param>
+        public static void InitLogging(bool bLogTimestamps, bool bLogVerbose, bool bLogSeverity, bool bLogSources, bool bColorConsoleOutput, IEnumerable<TraceListener> TraceListeners)
+        {
+            bIsInitialized = true;
+            Timer = (bLogTimestamps && Timer == null) ? Stopwatch.StartNew() : null;
+            Log.bLogVerbose = bLogVerbose;
+            Log.bLogSeverity = bLogSeverity;
+            Log.bLogSources = bLogSources;
+            Log.bColorConsoleOutput = bColorConsoleOutput;
+
+            // ensure that if InitLogging is called more than once we don't stack listeners.
+            // but always leave the default listener around.
+            for (int ListenerNdx = 0; ListenerNdx < Trace.Listeners.Count; )
 			{
-				Listener.TraceEvent(EventCache, Source, Verbosity, (int)Verbosity, Message);
+                if (Trace.Listeners[ListenerNdx].GetType() != typeof(DefaultTraceListener))
+				{
+                    Trace.Listeners.RemoveAt(ListenerNdx);
+				}
+				else
+				{
+                    ++ListenerNdx;
+				}
 			}
-		}
+            // don't add any null listeners
+            Trace.Listeners.AddRange(TraceListeners.Where(l => l != null).ToArray());
+            Trace.AutoFlush = true;
+        }
 
-		/// <summary>
-		/// Writes a formatted message to the console if the condition is met.
-		/// </summary>
-		/// <param name="Condition">Condition</param>
-		/// <param name="Verbosity">Message verbosity level</param>
-		/// <param name="Format">Message format string</param>
-		/// <param name="Args">Optional arguments</param>
-		[MethodImplAttribute(MethodImplOptions.NoInlining)]
-		public static void WriteLineIf(bool Condition, TraceEventType Verbosity, string Format, params object[] Args)
-		{
-			if (Condition)
-			{
-				WriteLine(1, Verbosity, Format, Args);
-			}
-		}
+        /// <summary>
+        /// Gets the name of the Method N levels deep in the stack frame. Used to trap what method actually made the logging call.
+        /// Only used when bLogSources is true.
+        /// </summary>
+        /// <param name="StackFramesToSkip"></param>
+        /// <returns>ClassName.MethodName</returns>
+        [MethodImplAttribute(MethodImplOptions.NoInlining)]
+        private static string GetSource(int StackFramesToSkip)
+        {
+            StackFrame Frame = new StackFrame(2 + StackFramesToSkip);
+            System.Reflection.MethodBase Method = Frame.GetMethod();
+            return String.Format("{0}.{1}", Method.DeclaringType.Name, Method.Name);
+        }
 
-		/// <summary>
-		/// Writes a message to the console if the condition is met.
-		/// </summary>
-		/// <param name="Condition">Condition.</param>
-		/// <param name="Verbosity">Message verbosity level</param>
-		/// <param name="Message">Message text</param>
-		[MethodImplAttribute(MethodImplOptions.NoInlining)]
-		public static void WriteLineIf(bool Condition, TraceEventType Verbosity, string Message)
-		{
-			if (Condition)
-			{
-				WriteLine(1, Verbosity, Message);
-			}
-		}
+        /// <summary>
+        /// Converts a TraceEventType into a log prefix. Only used when bLogSeverity is true.
+        /// </summary>
+        /// <param name="EventType"></param>
+        /// <returns></returns>
+        private static string GetSeverityPrefix(TraceEventType Severity)
+        {
+            return Severity <= TraceEventType.Error ? "ERROR: " : Severity == TraceEventType.Warning ? "WARNING: " : "";
+        }
 
-		/// <summary>
-		/// Writes a verbose message to the console.
-		/// </summary>
-		/// <param name="Format">Message format string</param>
-		/// <param name="Args">Optional arguments</param>
-		[Conditional("TRACE")]
-		[MethodImplAttribute(MethodImplOptions.NoInlining)]
-		public static void TraceVerbose(string Format, params object[] Args)
-		{
-			WriteLine(1, TraceEventType.Verbose, Format, Args);
-		}
+        /// <summary>
+        /// Formats message for logging. Enforces the configured options.
+        /// </summary>
+        /// <param name="StackFramesToSkip">Number of frames to skip to get to the originator of the log request.</param>
+        /// <param name="CustomSource">Custom source string to use. Use the Class.Method string if null. Only used if bLogSources = true.</param>
+        /// <param name="Verbosity">Message verbosity level</param>
+        /// <param name="Format">Message text format string</param>
+        /// <param name="Args">Message text parameters</param>
+        /// <returns>Formatted message</returns>
+        [MethodImplAttribute(MethodImplOptions.NoInlining)]
+        private static string FormatMessage(int StackFramesToSkip, string CustomSource, TraceEventType Verbosity, string Format, params object[] Args)
+        {
+            return string.Format("{0}{1}{2}{3}",
+                    Timer != null ? String.Format("[{0:hh\\:mm\\:ss\\.fff}] ", Timer.Elapsed) : "",
+                    bLogSources ? string.Format("{0}: ", string.IsNullOrEmpty(CustomSource) ? GetSource(StackFramesToSkip) : CustomSource) : "",
+                    bLogSeverity ? GetSeverityPrefix(Verbosity) : "",
+                    string.Format(Format, Args));
+        }
 
-		/// <summary>
-		/// Writes a verbose message to the console.
-		/// </summary>
-		/// <param name="Message">Message text</param>
-		[Conditional("TRACE")]
-		[MethodImplAttribute(MethodImplOptions.NoInlining)]
-		public static void TraceVerbose(string Message)
-		{
-			WriteLine(1, TraceEventType.Verbose, Message);
-		}
+        /// <summary>
+        /// Writes a formatted message to the console. All other functions should boil down to calling this method.
+        /// </summary>
+        /// <param name="StackFramesToSkip">Number of frames to skip to get to the originator of the log request.</param>
+        /// <param name="CustomSource">Custom source string to use. Use the default if null.</param>
+        /// <param name="Verbosity">Message verbosity level. We only meaningfully use values up to Verbose</param>
+        /// <param name="Format">Message format string.</param>
+        /// <param name="Args">Optional arguments</param>
+        [MethodImplAttribute(MethodImplOptions.NoInlining)]
+        private static void WriteLinePrivate(int StackFramesToSkip, string CustomSource, TraceEventType Verbosity, string Format, params object[] Args)
+        {
+            if (!bIsInitialized)
+            {
+                throw new BuildException("Tried to using Logging system before it was ready");
+            }
+            if (Verbosity < TraceEventType.Verbose || bLogVerbose)
+            {
+                // Do console color highlighting here.
+                ConsoleColor DefaultColor = ConsoleColor.Gray;
+				bool bIsWarning = false;
+				bool bIsError = false;
+                // don't try to touch the console unless we are told to color the output.
+                if (bColorConsoleOutput)
+                {
+                    DefaultColor = Console.ForegroundColor;
+                    bIsWarning = Verbosity == TraceEventType.Warning;
+                    bIsError = Verbosity <= TraceEventType.Error;
+					// @todo mono - mono doesn't seem to initialize the ForegroundColor properly, so we can't restore it properly.
+					// Avoid touching the console color unless we really need to.
+					if (bIsWarning || bIsError)
+					{
+                    	Console.ForegroundColor = bIsWarning ? ConsoleColor.Yellow : ConsoleColor.Red;
+					}
+                }
+                try
+                {
+					// @todo mono: mono has some kind of bug where calling mono recursively by spawning
+					// a new process causes Trace.WriteLine to stop functioning (it returns, but does nothing for some reason).
+					// work around this by simulating Trace.WriteLine on mono.
+					// We use UAT to spawn UBT instances recursively a lot, so this bug can effectively
+					// make all build output disappear outside of the top level UAT process.
+					#if MONO
+					    lock (((System.Collections.ICollection)Trace.Listeners).SyncRoot)
+						{
+						    foreach (TraceListener l in Trace.Listeners) 
+						    {
+                                l.WriteLine(FormatMessage(StackFramesToSkip + 1, CustomSource, Verbosity, Format, Args));
+							    l.Flush();
+						    }
+						}
+					#else
+                    	// Call Trace directly here. Trace ensures that our logging is threadsafe using the GlobalLock.
+                    	Trace.WriteLine(FormatMessage(StackFramesToSkip + 1, CustomSource, Verbosity, Format, Args));
+					#endif
+                }
+                finally
+                {
+                    // make sure we always put the console color back.
+                    if (bColorConsoleOutput && (bIsWarning || bIsError))
+                    {
+                        Console.ForegroundColor = DefaultColor;
+                    }
+                }
+            }
+        }
 
-		/// <summary>
-		/// Writes a message to the console.
-		/// </summary>
-		/// <param name="Format">Message format string</param>
-		/// <param name="Args">Optional arguments</param>
-		[MethodImplAttribute(MethodImplOptions.NoInlining)]
-		public static void TraceInformation(string Format, params object[] Args)
-		{
-			WriteLine(1, TraceEventType.Information, Format, Args);
-		}
+        /// <summary>
+        /// Similar to Trace.WriteLineIf
+        /// </summary>
+        /// <param name="Condition"></param>
+        /// <param name="Verbosity"></param>
+        /// <param name="Format"></param>
+        /// <param name="Args"></param>
+        [MethodImplAttribute(MethodImplOptions.NoInlining)]
+        public static void WriteLineIf(bool Condition, TraceEventType Verbosity, string Format, params object[] Args)
+        {
+            if (Condition)
+            {
+                WriteLinePrivate(1, null, Verbosity, Format, Args);
+            }
+        }
 
-		/// <summary>
-		/// Writes a message to the console.
-		/// </summary>
-		/// <param name="Message">Message text</param>
-		[MethodImplAttribute(MethodImplOptions.NoInlining)]
-		public static void TraceInformation(string Message)
-		{
-			WriteLine(1, TraceEventType.Information, Message);
-		}
+        /// <summary>
+        /// Mostly an internal function, but expose StackFramesToSkip to allow UAT to use existing wrapper functions and still get proper formatting.
+        /// </summary>
+        /// <param name="StackFramesToSkip"></param>
+        /// <param name="CustomSource">Custom source string to use. Use the default if null.</param>
+        /// <param name="Verbosity"></param>
+        /// <param name="Format"></param>
+        /// <param name="Args"></param>
+        [MethodImplAttribute(MethodImplOptions.NoInlining)]
+        public static void WriteLine(int StackFramesToSkip, TraceEventType Verbosity, string Format, params object[] Args)
+        {
+            WriteLinePrivate(StackFramesToSkip + 1, null, Verbosity, Format, Args);
+        }
 
-		/// <summary>
-		/// Writes a warning message to the console.
-		/// </summary>
-		/// <param name="Format">Message format string</param>
-		/// <param name="Args">Optional arguments</param>
-		[MethodImplAttribute(MethodImplOptions.NoInlining)]
-		public static void TraceWarning(string Format, params object[] Args)
-		{
-			WriteLine(1, TraceEventType.Warning, Format, Args);
-		}
+        /// <summary>
+        /// Mostly an internal function, but expose StackFramesToSkip and a custom Source to alow UAT to use existing wrapper functions and still get proper formatting.
+        /// </summary>
+        /// <param name="StackFramesToSkip"></param>
+        /// <param name="CustomSource">Custom source string to use. Use the default if null.</param>
+        /// <param name="Verbosity"></param>
+        /// <param name="Format"></param>
+        /// <param name="Args"></param>
+        [MethodImplAttribute(MethodImplOptions.NoInlining)]
+        public static void WriteLine(int StackFramesToSkip, string CustomSource, TraceEventType Verbosity, string Format, params object[] Args)
+        {
+            WriteLinePrivate(StackFramesToSkip + 1, CustomSource, Verbosity, Format, Args);
+        }
 
-		/// <summary>
-		/// Writes a warning message to the console.
-		/// </summary>
-		/// <param name="Message">Message text</param>
-		[MethodImplAttribute(MethodImplOptions.NoInlining)]
-		public static void TraceWarning(string Message)
-		{
-			WriteLine(1, TraceEventType.Warning, Message);
-		}
+        /// <summary>
+        /// Similar to Trace.WriteLin
+        /// </summary>
+        /// <param name="Verbosity"></param>
+        /// <param name="Format"></param>
+        /// <param name="Args"></param>
+        [MethodImplAttribute(MethodImplOptions.NoInlining)]
+        public static void WriteLine(TraceEventType Verbosity, string Format, params object[] Args)
+        {
+            WriteLinePrivate(1, null, Verbosity, Format, Args);
+        }
 
-		/// <summary>
-		/// Writes an error message to the console.
-		/// </summary>
-		/// <param name="Format">Message format string</param>
-		/// <param name="Args">Optional arguments</param>
-		[MethodImplAttribute(MethodImplOptions.NoInlining)]
-		public static void TraceError(string Format, params object[] Args)
-		{
-			WriteLine(1, TraceEventType.Error, Format, Args);
-		}
+        /// <summary>
+        /// Writes an error message to the console.
+        /// </summary>
+        /// <param name="Format">Message format string</param>
+        /// <param name="Args">Optional arguments</param>
+        [MethodImplAttribute(MethodImplOptions.NoInlining)]
+        public static void TraceError(string Format, params object[] Args)
+        {
+            WriteLinePrivate(1, null, TraceEventType.Error, Format, Args);
+        }
 
-		/// <summary>
-		/// Writes an error message to the console.
-		/// </summary>
-		/// <param name="Message">Message text</param>
-		[MethodImplAttribute(MethodImplOptions.NoInlining)]
-		public static void TraceError(string Message)
-		{
-			WriteLine(1, TraceEventType.Error, Message);
-		}
+        /// <summary>
+        /// Writes a verbose message to the console.
+        /// </summary>
+        /// <param name="Format">Message format string</param>
+        /// <param name="Args">Optional arguments</param>
+        [Conditional("TRACE")]
+        [MethodImplAttribute(MethodImplOptions.NoInlining)]
+        public static void TraceVerbose(string Format, params object[] Args)
+        {
+            WriteLinePrivate(1, null, TraceEventType.Verbose, Format, Args);
+        }
 
-		/// <summary>
-		/// Writes a message with the specified verbosity to the console.
-		/// </summary>
-		[MethodImplAttribute(MethodImplOptions.NoInlining)]
-		public static void TraceEvent(TraceEventType Verbosity, string Message)
-		{
-			WriteLine(1, Verbosity, Message);
-		}
+        /// <summary>
+        /// Writes a message to the console.
+        /// </summary>
+        /// <param name="Format">Message format string</param>
+        /// <param name="Args">Optional arguments</param>
+        [MethodImplAttribute(MethodImplOptions.NoInlining)]
+        public static void TraceInformation(string Format, params object[] Args)
+        {
+            WriteLinePrivate(1, null, TraceEventType.Information, Format, Args);
+        }
 
-		/// <summary>
-		/// Writes a formatted message with the specified verbosity to the console.
-		/// </summary>
-		[MethodImplAttribute(MethodImplOptions.NoInlining)]
-		public static void TraceEvent(TraceEventType Verbosity, string Format, params object[] Args)
-		{
-			WriteLine(1, Verbosity, Format, Args);
-		}
-	}
+        /// <summary>
+        /// Writes a warning message to the console.
+        /// </summary>
+        /// <param name="Format">Message format string</param>
+        /// <param name="Args">Optional arguments</param>
+        [MethodImplAttribute(MethodImplOptions.NoInlining)]
+        public static void TraceWarning(string Format, params object[] Args)
+        {
+            WriteLinePrivate(1, null, TraceEventType.Warning, Format, Args);
+        }
+    }
 
     #region StreamUtils
     public static class StreamUtils
@@ -1533,7 +1587,7 @@ namespace UnrealBuildTool
 			var IndentText = new StringBuilder(LogIndent * 2);
 			IndentText.Append(' ', LogIndent * 2);
 			
-			Log.TraceEvent(Verbosity, "{0}{1} took {2}s", IndentText.ToString(), TimerName, TotalSeconds);
+			Log.WriteLine(Verbosity, "{0}{1} took {2}s", IndentText.ToString(), TimerName, TotalSeconds);
 		}
 	}
 
@@ -1603,7 +1657,7 @@ namespace UnrealBuildTool
 
 		static private void LogAccumulator(string Name, Accumulator Counter)
 		{
-			Log.TraceEvent(Counter.Verbosity, "{0} took {1}s", Name, Counter.Time);
+			Log.WriteLine(Counter.Verbosity, "{0} took {1}s", Name, Counter.Time);
 		}
 
 		/// <summary>
diff --git a/Engine/Source/Programs/UnrealBuildTool/Utilities/XmlConfigLoader.cs b/Engine/Source/Programs/UnrealBuildTool/Utilities/XmlConfigLoader.cs
--- a/Engine/Source/Programs/UnrealBuildTool/Utilities/XmlConfigLoader.cs
+++ b/Engine/Source/Programs/UnrealBuildTool/Utilities/XmlConfigLoader.cs
@@ -18,7 +18,7 @@ namespace UnrealBuildTool
 	/// <summary>
 	/// Attribute to annotate fields in type that can be set using XML configuration system.
 	/// </summary>
-	[AttributeUsage(AttributeTargets.Field)]
+	[AttributeUsage(AttributeTargets.Field | AttributeTargets.Property)]
 	public class XmlConfigAttribute : Attribute
 	{
 
@@ -198,7 +198,13 @@ namespace UnrealBuildTool
 		/// </summary>
 		public static void Init()
 		{
-			OverwriteIfDifferent(GetXSDPath(), BuildXSD());
+            // No one should try to refereence configuration values until the config files are loaded.
+            // The XmlConfig system itself will SET config values below, then anyone can read them.
+            // But our static constructor checks can't differentiate this, so we just allow reads starting now,
+            // right before the XML files are loaded.
+            UnrealBuildTool.bIsSafeToReferenceConfigurationValues = true;
+
+            OverwriteIfDifferent(GetXSDPath(), BuildXSD());
 
 			LoadData();
 
@@ -297,7 +303,12 @@ namespace UnrealBuildTool
 					DataPair.Key.SetValue(null, DataPair.Value);
 				}
 
-				bDoneLoading = true;
+                foreach (var PropertyPair in PropertyMap)
+                {
+                    PropertyPair.Key.SetValue(null, PropertyPair.Value, null);
+                }
+
+                bDoneLoading = true;
 			}
 
 			/// <summary>
@@ -317,7 +328,24 @@ namespace UnrealBuildTool
 				}
 			}
 
-			// A variable to indicate if loading was done during invoking of
+            /// <summary>
+            /// Adds or overrides value in the cache.
+            /// </summary>
+            /// <param name="Property">The property info of the class.</param>
+            /// <param name="Value">The value to store.</param>
+            public void SetValue(PropertyInfo Property, object Value)
+            {
+                if (PropertyMap.ContainsKey(Property))
+                {
+                    PropertyMap[Property] = Value;
+                }
+                else
+                {
+                    PropertyMap.Add(Property, Value);
+                }
+            }
+
+            // A variable to indicate if loading was done during invoking of
 			// default values loader.
 			bool bDoneLoading = false;
 
@@ -326,7 +354,10 @@ namespace UnrealBuildTool
 
 			// Loaded data map.
 			Dictionary<FieldInfo, object> DataMap = new Dictionary<FieldInfo, object>();
-		}
+        
+            // Loaded data map.
+            Dictionary<PropertyInfo, object> PropertyMap = new Dictionary<PropertyInfo, object>();
+        }
 
 		/// <summary>
 		/// Class that stores information about possible BuildConfiguration.xml
@@ -615,7 +646,18 @@ namespace UnrealBuildTool
 					// allow settings in the .xml that don't exist, as another branch may have it, and can share this file from Documents
 					if (Field == null)
 					{
-						continue;
+                        PropertyInfo Property = ClassType.GetProperty(XmlField.Name);
+                        if (Property != null)
+                        {
+                            if (!IsConfigurable(Property))
+                            {
+                                throw new BuildException("BuildConfiguration Loading: property '{0}' is either non-public, non-static or not-xml-configurable.", XmlField.Name);
+                            }
+
+                            ClassData.SetValue(Property, ParseFieldData(Property.PropertyType, XmlField.InnerText));
+                        }
+
+                        continue;
 					}
 					
 					if (!IsConfigurableField(Field))
@@ -747,10 +789,23 @@ namespace UnrealBuildTool
 		/// <returns>True if the class is configurable using XML system. Otherwise false.</returns>
 		private static bool IsConfigurableClass(Type Class)
 		{
-			return Class.GetFields().Any((Field) => IsConfigurableField(Field));
+            return
+                Class.GetFields().Any((Field) => IsConfigurableField(Field)) ||
+                Class.GetProperties().Any(prop => IsConfigurable(prop));
+
 		}
 
-		/// <summary>
+        /// <summary>
+        /// Tells if given property is XML configurable.
+        /// </summary>
+        /// <param name="Property">Property to check.</param>
+        /// <returns>True if the property is configurable using XML system. Otherwise false.</returns>
+        private static bool IsConfigurable(PropertyInfo Property)
+        {
+            return Property.GetCustomAttributes(typeof(XmlConfigAttribute), false).Length > 0 && Property.CanWrite && Property.GetSetMethod().IsStatic && Property.GetSetMethod().IsPublic;
+        }
+
+        /// <summary>
 		/// Tells if given field is XML configurable.
 		/// </summary>
 		/// <param name="Field">field to check.</param>
