Handle Logging that should not have string.Format called on it. If the Arg count is zero, treat it is a raw string and simply use it.

diff --git a/Engine/Source/Programs/UnrealBuildTool/System/Utils.cs b/Engine/Source/Programs/UnrealBuildTool/System/Utils.cs
--- a/Engine/Source/Programs/UnrealBuildTool/System/Utils.cs
+++ b/Engine/Source/Programs/UnrealBuildTool/System/Utils.cs
@@ -1193,7 +1193,8 @@ namespace UnrealBuildTool
                     Timer != null ? String.Format("[{0:hh\\:mm\\:ss\\.fff}] ", Timer.Elapsed) : "",
                     bLogSources ? string.Format("{0}: ", string.IsNullOrEmpty(CustomSource) ? GetSource(StackFramesToSkip) : CustomSource) : "",
                     bLogSeverity ? GetSeverityPrefix(Verbosity) : "",
-                    string.Format(Format, Args));
+                    // If there are no extra args, don't try to format the string, in case it has any format control characters in it (our LOCTEXT strings tend to).
+                    Args.Length > 0 ? string.Format(Format, Args) : Format);
         }
 
         /// <summary>
