Fixed XSD IntelliSense in and added validation for UBT XML config files.

diff --git a/Engine/Source/Programs/UnrealBuildTool/Utilities/XmlConfigLoader.cs b/Engine/Source/Programs/UnrealBuildTool/Utilities/XmlConfigLoader.cs
--- a/Engine/Source/Programs/UnrealBuildTool/Utilities/XmlConfigLoader.cs
+++ b/Engine/Source/Programs/UnrealBuildTool/Utilities/XmlConfigLoader.cs
@@ -12,6 +12,7 @@ using System.Text.RegularExpressions;
 using System.Threading;
 using System.Xml;
 using System.Xml.Linq;
+using System.Xml.Schema;
 
 namespace UnrealBuildTool
 {
@@ -238,7 +239,7 @@ namespace UnrealBuildTool
 				}
 
 				Directory.CreateDirectory(Path.GetDirectoryName(FilePath));
-				File.WriteAllText(FilePath, Content);
+				File.WriteAllText(FilePath, Content, Encoding.UTF8);
 
 				if(bReadOnlyFile)
 				{
@@ -260,7 +261,7 @@ namespace UnrealBuildTool
 				return true;
 			}
 
-			return !File.ReadAllText(FilePath).Equals(Content, StringComparison.InvariantCulture);
+			return !File.ReadAllText(FilePath, Encoding.UTF8).Equals(Content, StringComparison.InvariantCulture);
 		}
 
 		/// <summary>
@@ -577,6 +578,44 @@ namespace UnrealBuildTool
 			}
 		}
 
+		/// <summary>
+		/// Reads config schema from XSD file.
+		/// </summary>
+		private static XmlSchema ReadConfigSchema()
+		{
+			var Settings = new XmlReaderSettings();
+			Settings.ValidationType = ValidationType.DTD;
+
+			using (var SR = new StringReader(File.ReadAllText(GetXSDPath(), Encoding.UTF8)))
+			{
+				using (var XR = XmlReader.Create(SR, Settings))
+				{
+					return XmlSchema.Read(XR, (object Sender, System.Xml.Schema.ValidationEventArgs EventArgs) =>
+					{
+						throw new BuildException("XmlConfigLoader: Reading config XSD failed:\n{0}({1}): {2}",
+							new Uri(EventArgs.Exception.SourceUri).LocalPath, EventArgs.Exception.LineNumber,
+							EventArgs.Message);
+					});
+				}
+			}
+		}
+
+		// Stores read XSD schema for config files.
+		private static XmlSchema ConfigSchemaCache = null;
+
+		/// <summary>
+		/// Gets config XSD schema.
+		/// </summary>
+		private static XmlSchema GetConfigSchema()
+		{
+			if(ConfigSchemaCache == null)
+			{
+				ConfigSchemaCache = ReadConfigSchema();
+			}
+
+			return ConfigSchemaCache;
+		}
+
 		/// <summary>
 		/// Sets values of this class with values from given XML file.
 		/// </summary>
@@ -586,7 +625,20 @@ namespace UnrealBuildTool
 			var ConfigDocument = new XmlDocument();
 			var NS = new XmlNamespaceManager(ConfigDocument.NameTable);
 			NS.AddNamespace("ns", "https://www.unrealengine.com/BuildConfiguration");
-			ConfigDocument.Load(ConfigurationXmlPath);
+			var ReaderSettings = new XmlReaderSettings();
+
+			ReaderSettings.ValidationEventHandler += (object sender, System.Xml.Schema.ValidationEventArgs e) =>
+			{
+				Console.Out.WriteLine(e.Message);
+			};
+			ReaderSettings.ValidationType = ValidationType.Schema;
+			ReaderSettings.Schemas.Add(GetConfigSchema());
+
+			using (var SR = new StringReader(File.ReadAllText(ConfigurationXmlPath, Encoding.UTF8)))
+			{
+				var Reader = XmlReader.Create(SR, ReaderSettings);
+				ConfigDocument.Load(Reader);
+			}
 
 			var XmlClasses = ConfigDocument.DocumentElement.SelectNodes("/ns:Configuration/*", NS);
 
@@ -862,7 +914,9 @@ namespace UnrealBuildTool
 				{ typeof(bool), "boolean" },
 				{ typeof(int), "int" },
 				{ typeof(long), "long" },
-				{ typeof(DateTime), "dateTime" }
+				{ typeof(DateTime), "dateTime" },
+				{ typeof(float), "float" },
+				{ typeof(double), "double" }
 			};
 
 			return TypeMap.ContainsKey(FieldType)
@@ -895,7 +949,7 @@ namespace UnrealBuildTool
 		/// </summary>
 		/// <param name="Field">Field to represent value.</param>
 		/// <returns>Field value as XML.</returns>
-		private static string GetFieldXMLValue(FieldInfo Field)
+		private static object GetFieldXMLValue(FieldInfo Field)
 		{
 			return GetObjectXMLValue(Field.FieldType, Field.GetValue(null));
 		}
@@ -906,7 +960,7 @@ namespace UnrealBuildTool
 		/// <param name="FieldType">The type of the field.</param>
 		/// <param name="Obj">The value.</param>
 		/// <returns>String representation.</returns>
-		private static string GetObjectXMLValue(Type FieldType, object Obj)
+		private static object GetObjectXMLValue(Type FieldType, object Obj)
 		{
 			if (Obj == null && FieldType == typeof(string))
 			{
@@ -918,7 +972,7 @@ namespace UnrealBuildTool
 				return (bool) Obj ? "true" : "false";
 			}
 
-			return Obj.ToString();
+			return Obj;
 		}
 
 		/// <summary>
@@ -941,11 +995,12 @@ namespace UnrealBuildTool
 					new XAttribute("minOccurs", "0"),
 					new XAttribute("maxOccurs", "1"),
 					new XElement(NS + "complexType",
-						new XElement(NS + "all",
+						new XElement(NS + "sequence",
 							new XElement(NS + "element",
 								new XAttribute("name", "Item"),
 								new XAttribute("type", GetFieldXSDType(Field.FieldType.GetElementType())),
-								new XAttribute("minOccurs", "0")
+								new XAttribute("minOccurs", "0"),
+								new XAttribute("maxOccurs", "unbounded")
 							)
 						)
 					)
