UBT: Changed XSD validation messages to be actually errors.

diff --git a/Engine/Source/Programs/UnrealBuildTool/Utilities/XmlConfigLoader.cs b/Engine/Source/Programs/UnrealBuildTool/Utilities/XmlConfigLoader.cs
--- a/Engine/Source/Programs/UnrealBuildTool/Utilities/XmlConfigLoader.cs
+++ b/Engine/Source/Programs/UnrealBuildTool/Utilities/XmlConfigLoader.cs
@@ -627,9 +627,11 @@ namespace UnrealBuildTool
 			NS.AddNamespace("ns", "https://www.unrealengine.com/BuildConfiguration");
 			var ReaderSettings = new XmlReaderSettings();
 
-			ReaderSettings.ValidationEventHandler += (object sender, System.Xml.Schema.ValidationEventArgs e) =>
+			ReaderSettings.ValidationEventHandler += (object Sender, System.Xml.Schema.ValidationEventArgs EventArgs) =>
 			{
-				Console.Out.WriteLine(e.Message);
+				throw new BuildException("XmlConfigLoader: Reading config XML failed:\n{0}({1}): {2}",
+					ConfigurationXmlPath, EventArgs.Exception.LineNumber,
+					EventArgs.Message);
 			};
 			ReaderSettings.ValidationType = ValidationType.Schema;
 			ReaderSettings.Schemas.Add(GetConfigSchema());
