UE-13654 - Package still contains data if "Package game data inside .apk?" option is once selected.
- Added logic to delete the OBB file if it already exists in the APK data directory

diff --git a/Engine/Source/Programs/UnrealBuildTool/Android/UEDeployAndroid.cs b/Engine/Source/Programs/UnrealBuildTool/Android/UEDeployAndroid.cs
--- a/Engine/Source/Programs/UnrealBuildTool/Android/UEDeployAndroid.cs
+++ b/Engine/Source/Programs/UnrealBuildTool/Android/UEDeployAndroid.cs
@@ -1165,6 +1165,15 @@ namespace UnrealBuildTool.Android
 					}
 				}
 			}
+            else // try to remove the file it we aren't packaing inside the APK
+            {
+                string ObbFileDestination = UE4BuildPath + "/assets";
+                var DestFileName = Path.Combine(ObbFileDestination, "main.obb.png");
+                if(File.Exists(DestFileName))
+                {
+                    File.Delete(DestFileName);
+                }
+            }
 
 			//Copy build files to the intermediate folder in this order (later overrides earlier):
 			//	- Shared Engine
