Add symbol stripping at the IUEToolChain level. Implemented for Windows, Linux, Android arm7.

diff --git a/Engine/Source/Programs/AutomationTool/Scripts/RocketBuild.Automation.cs b/Engine/Source/Programs/AutomationTool/Scripts/RocketBuild.Automation.cs
--- a/Engine/Source/Programs/AutomationTool/Scripts/RocketBuild.Automation.cs
+++ b/Engine/Source/Programs/AutomationTool/Scripts/RocketBuild.Automation.cs
@@ -375,22 +375,13 @@ namespace Rocket
 
 		public static void StripSymbols(UnrealTargetPlatform TargetPlatform, string[] SourceFileNames, string[] TargetFileNames)
 		{
-			if(TargetPlatform == UnrealTargetPlatform.Win32 || TargetPlatform == UnrealTargetPlatform.Win64)
+			IUEBuildPlatform Platform = UEBuildPlatform.GetBuildPlatform(TargetPlatform);
+			IUEToolChain ToolChain = UEToolChain.GetPlatformToolChain(Platform.GetCPPTargetPlatform(TargetPlatform));
+			for (int Idx = 0; Idx < SourceFileNames.Length; Idx++)
 			{
-				string PDBCopyPath = Path.Combine(Environment.GetFolderPath(Environment.SpecialFolder.ProgramFilesX86), "MSBuild", "Microsoft", "VisualStudio", "v12.0", "AppxPackage", "PDBCopy.exe");
-				for(int Idx = 0; Idx < SourceFileNames.Length; Idx++)
-				{
-					CommandUtils.CreateDirectory(Path.GetDirectoryName(TargetFileNames[Idx]));
-					CommandUtils.Run(PDBCopyPath, String.Format("{0} {1} -p", CommandUtils.MakePathSafeToUseWithCommandLine(SourceFileNames[Idx]), CommandUtils.MakePathSafeToUseWithCommandLine(TargetFileNames[Idx])));
-				}
-			}
-			else
-			{
-				for(int Idx = 0; Idx < SourceFileNames.Length; Idx++)
-				{
-					CommandUtils.LogWarning("Stripping has not yet been implemented for {0}. Skipping strip for {1}.", TargetPlatform, SourceFileNames[Idx]);
-					CommandUtils.CopyFile(SourceFileNames[Idx], TargetFileNames[Idx]);
-				}
+				CommandUtils.CreateDirectory(Path.GetDirectoryName(TargetFileNames[Idx]));
+				CommandUtils.Log("Stripping symbols: {0} -> {1}", SourceFileNames[Idx], TargetFileNames[Idx]);
+				ToolChain.StripSymbols(SourceFileNames[Idx], TargetFileNames[Idx]);
 			}
 		}
 	}
diff --git a/Engine/Source/Programs/UnrealBuildTool/Android/AndroidToolChain.cs b/Engine/Source/Programs/UnrealBuildTool/Android/AndroidToolChain.cs
--- a/Engine/Source/Programs/UnrealBuildTool/Android/AndroidToolChain.cs
+++ b/Engine/Source/Programs/UnrealBuildTool/Android/AndroidToolChain.cs
@@ -967,5 +967,24 @@ namespace UnrealBuildTool
 		{
 			return UnrealTargetPlatform.Android;
 		}
+
+		public override void StripSymbols(string SourceFileName, string TargetFileName)
+		{
+			File.Copy(SourceFileName, TargetFileName, true);
+
+			ProcessStartInfo StartInfo = new ProcessStartInfo();
+			if(SourceFileName.Contains("-armv7"))
+			{
+				StartInfo.FileName = ArPathArm.Replace("-ar.exe", "-strip.exe");
+			}
+			else
+			{
+				throw new BuildException("Couldn't determine Android architecture to strip symbols from {0}", SourceFileName);
+			}
+			StartInfo.Arguments = TargetFileName;
+			StartInfo.UseShellExecute = false;
+			StartInfo.CreateNoWindow = true;
+			Utils.RunLocalProcessAndLogOutput(StartInfo);
+		}
 	};
 }
diff --git a/Engine/Source/Programs/UnrealBuildTool/Linux/LinuxToolChain.cs b/Engine/Source/Programs/UnrealBuildTool/Linux/LinuxToolChain.cs
--- a/Engine/Source/Programs/UnrealBuildTool/Linux/LinuxToolChain.cs
+++ b/Engine/Source/Programs/UnrealBuildTool/Linux/LinuxToolChain.cs
@@ -147,6 +147,7 @@ namespace UnrealBuildTool
                 GCCPath = Which("g++");
                 ArPath = Which("ar");
                 RanlibPath = Which("ranlib");
+				StripPath = Which("strip");
 
 				// if clang is available, zero out gcc (@todo: support runtime switching?)
 				if (!String.IsNullOrEmpty(ClangPath))
@@ -171,6 +172,7 @@ namespace UnrealBuildTool
 				// ar and ranlib will be switched later to match the architecture
 				ArPath = "ar.exe";
 				RanlibPath = "ranlib.exe";
+				StripPath = "strip.exe";
 			}
 
 			if (!DetermineCompilerVersion())
@@ -262,6 +264,17 @@ namespace UnrealBuildTool
 			return RanlibPath;
 		}
 
+		/** Gets architecture-specific strip path */
+		private static string GetStripPath(string Architecture)
+		{
+			if (CrossCompiling())
+			{
+				return Path.Combine(Path.Combine(BaseLinuxPath, String.Format("bin/{0}-{1}", Architecture, StripPath)));
+			}
+
+			return StripPath;
+		}
+
 		static string GetCLArguments_Global(CPPEnvironment CompileEnvironment)
         {
             string Result = "";
@@ -553,6 +566,7 @@ namespace UnrealBuildTool
         static string GCCPath;
         static string ArPath;
         static string RanlibPath;
+		static string StripPath;
 
 		/** Version string of the current compiler, whether clang or gcc or whatever */
 		static string CompilerVersionString;
@@ -1117,5 +1131,17 @@ namespace UnrealBuildTool
 		{
 			return UnrealTargetPlatform.Linux;
 		}
+
+		public override void StripSymbols(string SourceFileName, string TargetFileName)
+		{
+			File.Copy(SourceFileName, TargetFileName, true);
+
+			ProcessStartInfo StartInfo = new ProcessStartInfo();
+			StartInfo.FileName = GetStripPath(UEBuildPlatform.GetBuildPlatform(UnrealTargetPlatform.Linux).GetActiveArchitecture());
+			StartInfo.Arguments = TargetFileName;
+			StartInfo.UseShellExecute = false;
+			StartInfo.CreateNoWindow = true;
+			Utils.RunLocalProcessAndLogOutput(StartInfo);
+		}
 	}
 }
diff --git a/Engine/Source/Programs/UnrealBuildTool/ToolChain/UEToolChain.cs b/Engine/Source/Programs/UnrealBuildTool/ToolChain/UEToolChain.cs
--- a/Engine/Source/Programs/UnrealBuildTool/ToolChain/UEToolChain.cs
+++ b/Engine/Source/Programs/UnrealBuildTool/ToolChain/UEToolChain.cs
@@ -45,6 +45,8 @@ namespace UnrealBuildTool
 		void FixBundleBinariesPaths(UEBuildTarget Target, List<UEBuildBinary> Binaries);
 
 		UnrealTargetPlatform GetPlatform();
+
+		void StripSymbols(string SourceFileName, string TargetFileName);
 	}
 
 	public abstract class UEToolChain : IUEToolChain
@@ -222,5 +224,11 @@ namespace UnrealBuildTool
 		{
 			return UnrealTargetPlatform.Unknown;
 		}
+
+		public virtual void StripSymbols(string SourceFileName, string TargetFileName)
+		{
+			Log.TraceWarning("StripSymbols() has not been implemented for {0}; copying files", GetPlatform().ToString());
+			File.Copy(SourceFileName, TargetFileName, true);
+		}
 	};
 }
\ No newline at end of file
diff --git a/Engine/Source/Programs/UnrealBuildTool/Windows/VCToolChain.cs b/Engine/Source/Programs/UnrealBuildTool/Windows/VCToolChain.cs
--- a/Engine/Source/Programs/UnrealBuildTool/Windows/VCToolChain.cs
+++ b/Engine/Source/Programs/UnrealBuildTool/Windows/VCToolChain.cs
@@ -1700,5 +1700,15 @@ namespace UnrealBuildTool
 			// @todo clang: Convert relative includes to absolute files so they'll be clickable
 			Log.TraceInformation(Output);
 		}
+
+		public override void StripSymbols(string SourceFileName, string TargetFileName)
+		{
+			ProcessStartInfo StartInfo = new ProcessStartInfo();
+			StartInfo.FileName = Path.Combine(Environment.GetFolderPath(Environment.SpecialFolder.ProgramFilesX86), "MSBuild", "Microsoft", "VisualStudio", "v12.0", "AppxPackage", "PDBCopy.exe");
+			StartInfo.Arguments = String.Format("\"{0}\" \"{1}\" -p", SourceFileName, TargetFileName);
+			StartInfo.UseShellExecute = false;
+			StartInfo.CreateNoWindow = true;
+			Utils.RunLocalProcessAndLogOutput(StartInfo);
+		}
 	};
 }
