Strip symbols from iOS and Mac libraries in Rocket builds.

diff --git a/Engine/Build/InstalledEngineFilters.ini b/Engine/Build/InstalledEngineFilters.ini
--- a/Engine/Build/InstalledEngineFilters.ini
+++ b/Engine/Build/InstalledEngineFilters.ini
@@ -186,6 +186,14 @@
 /Engine/Plugins/.../Binaries/Android/....a
 /Engine/Plugins/.../Binaries/Android/....so
 
+[StripSymbols.iOS]
+/Engine/Binaries/iOS/....a
+/Engine/Plugins/.../Binaries/iOS/....a
+
+[StripSymbols.Mac]
+/Engine/Binaries/Mac/....a
+/Engine/Plugins/.../Binaries/Mac/....a
+
 [StripSymbols.Win32]
 /Engine/Binaries/Win32/....pdb
 /Engine/Plugins/.../Binaries/Win32/....pdb
diff --git a/Engine/Source/Programs/UnrealBuildTool/IOS/IOSToolChain.cs b/Engine/Source/Programs/UnrealBuildTool/IOS/IOSToolChain.cs
--- a/Engine/Source/Programs/UnrealBuildTool/IOS/IOSToolChain.cs
+++ b/Engine/Source/Programs/UnrealBuildTool/IOS/IOSToolChain.cs
@@ -1359,5 +1359,17 @@ namespace UnrealBuildTool
 			NewProcess.WaitForExit();
 			return NewProcess.ExitCode;
 		}
+
+		public override void StripSymbols(string SourceFileName, string TargetFileName)
+		{
+			File.Copy(SourceFileName, TargetFileName, true);
+
+			ProcessStartInfo StartInfo = new ProcessStartInfo();
+			StartInfo.FileName = Path.Combine(XcodeDeveloperDir, "usr/bin/xcrun");
+			StartInfo.Arguments = String.Format("strip \"{0}\" -S", TargetFileName);
+			StartInfo.UseShellExecute = false;
+			StartInfo.CreateNoWindow = true;
+			Utils.RunLocalProcessAndLogOutput(StartInfo);
+		}
 	};
 }
diff --git a/Engine/Source/Programs/UnrealBuildTool/Mac/MacToolChain.cs b/Engine/Source/Programs/UnrealBuildTool/Mac/MacToolChain.cs
--- a/Engine/Source/Programs/UnrealBuildTool/Mac/MacToolChain.cs
+++ b/Engine/Source/Programs/UnrealBuildTool/Mac/MacToolChain.cs
@@ -1509,5 +1509,17 @@ namespace UnrealBuildTool
 		{
 			return UnrealTargetPlatform.Mac;
 		}
+
+		public override void StripSymbols(string SourceFileName, string TargetFileName)
+		{
+			File.Copy(SourceFileName, TargetFileName, true);
+
+			ProcessStartInfo StartInfo = new ProcessStartInfo();
+			StartInfo.FileName = Path.Combine(XcodeDeveloperDir, "usr/bin/xcrun");
+			StartInfo.Arguments = String.Format("strip \"{0}\" -S", TargetFileName);
+			StartInfo.UseShellExecute = false;
+			StartInfo.CreateNoWindow = true;
+			Utils.RunLocalProcessAndLogOutput(StartInfo);
+		}
 	};
 }
