UE-11479 - fix for generating project files when project has a platform name in the directory structure

diff --git a/Engine/Source/Programs/UnrealBuildTool/System/ProjectFileGenerator.cs b/Engine/Source/Programs/UnrealBuildTool/System/ProjectFileGenerator.cs
--- a/Engine/Source/Programs/UnrealBuildTool/System/ProjectFileGenerator.cs
+++ b/Engine/Source/Programs/UnrealBuildTool/System/ProjectFileGenerator.cs
@@ -1009,6 +1009,22 @@ namespace UnrealBuildTool
 					}
 				}
 
+				if (UnrealBuildTool.HasUProjectFile())
+				{
+					string ProjectRoot = UnrealBuildTool.GetUProjectPath();
+					if (Search.StartsWith(ProjectRoot, StringComparison.InvariantCultureIgnoreCase))
+					{
+						if (ProjectRoot.EndsWith("\\") || ProjectRoot.EndsWith("/"))
+						{
+							Search = Search.Substring(ProjectRoot.Length - 1);
+						}
+						else
+						{
+							Search = Search.Substring(ProjectRoot.Length);
+						}
+					}
+				}
+
 				// Skip targets in unsupported platform directories
 				bool IncludeThisTarget = true;
 				foreach( var CurPlatformName in UnsupportedPlatformNameStrings )
