UE-14324 - fixed issue with projects outside of Engine directory when using new SSH version of building from windows
#ios

diff --git a/Engine/Source/Programs/UnrealBuildTool/ToolChain/RemoteToolChain.cs b/Engine/Source/Programs/UnrealBuildTool/ToolChain/RemoteToolChain.cs
--- a/Engine/Source/Programs/UnrealBuildTool/ToolChain/RemoteToolChain.cs
+++ b/Engine/Source/Programs/UnrealBuildTool/ToolChain/RemoteToolChain.cs
@@ -671,7 +671,7 @@ namespace UnrealBuildTool
 				List<string> RelativeRsyncDirs = new List<string>();
 				foreach (string Dir in RsyncDirs)
 				{
-					RelativeRsyncDirs.Add(Utils.CleanDirectorySeparators(Utils.MakePathRelativeTo(Dir, "../.."), '/') + "/");
+					RelativeRsyncDirs.Add(Utils.CleanDirectorySeparators(Dir.Replace(":", ""), '/') + "/");
 				}
 
 				// write out directories to copy
@@ -681,8 +681,10 @@ namespace UnrealBuildTool
 				File.WriteAllLines(IncludeFromFile, RsyncExtensions);
 
 				// source and destination paths in the format rsync wants
-				string CygRootPath = ConvertPathToCygwin(Path.GetFullPath("../../"));
-				string RemotePath = ConvertPath(Path.GetFullPath("../../")).Replace(" ", "\\ ");
+				string CygRootPath = "/cygdrive";// ConvertPathToCygwin(Path.GetFullPath(""));
+				string RemotePath = string.Format("{0}{1}",
+					UserDevRootMac,
+					Environment.MachineName);
 
 				// get the executable dir for SSH, so Rsync can call it easily
 				string ExeDir = Path.GetDirectoryName(ResolvedSSHExe);
