Allow UBT to find a matching PluginInfo for a module, even if the target hasn't specified that it's using that plugin. Fixes issue where UFE was adding VisualStudioSourceCodeAccess as a dependent module directly, which was causing it to be output to the Binaries/Win64 folder rather than the plugin directory.

diff --git a/Engine/Build/JunkManifest.txt b/Engine/Build/JunkManifest.txt
--- a/Engine/Build/JunkManifest.txt
+++ b/Engine/Build/JunkManifest.txt
@@ -19,6 +19,7 @@ Engine\Binaries\Win64\UE4Editor-AndroidTarget*.*
 Engine\Binaries\Win64\UE4Editor-IOSPlatformEditor*.*
 Engine\Binaries\Win64\UE4Editor-TaskBrowser*.*
 Engine\Binaries\Win64\UE4Editor-Levels*.*
+Engine\Binaries\Win64\UnrealFrontend-VisualStudioSourceCodeAccess*.*
 Engine\Binaries\Win64\*-ProjectWorkflows.*
 Engine\Binaries\Win64\*-WorkflowEditor.*
 Engine\Binaries\Win64\*-Workflows.*
@@ -32,6 +33,7 @@ Engine\Binaries\Win64\*-CameraDebugger*.*
 Engine\Binaries\Win64\*-SessionLauncher*.*
 Engine\Binaries\Mac\UnrealFrontend.app\Contents\Frameworks\*-CoreSlateStyle.*
 Engine\Binaries\Mac\UnrealFrontend.app\Contents\Frameworks\*-SlateStyle.*
+Engine\Binaries\Mac\UnrealFrontend.app\Contents\MacOS\*-XCodeSourceCodeAccess.*
 Engine\Binaries\Mac\UE4Editor.app\Contents\Frameworks\*-CoreSlateStyle.*
 Engine\Binaries\Mac\UE4Editor.app\Contents\Frameworks\*-SlateStyle.*
 Engine\Binaries\Mac\UE4Editor.app\Contents\MacOS\*-CameraDebugger.*
diff --git a/Engine/Source/Programs/UnrealBuildTool/Configuration/UEBuildTarget.cs b/Engine/Source/Programs/UnrealBuildTool/Configuration/UEBuildTarget.cs
--- a/Engine/Source/Programs/UnrealBuildTool/Configuration/UEBuildTarget.cs
+++ b/Engine/Source/Programs/UnrealBuildTool/Configuration/UEBuildTarget.cs
@@ -675,6 +675,10 @@ namespace UnrealBuildTool
 		[NonSerialized]
 		public LinkEnvironment GlobalLinkEnvironment = new LinkEnvironment();
 
+		/** All plugins which are valid for this target */
+		[NonSerialized]
+		public List<PluginInfo> ValidPlugins;
+
 		/** All plugins which are built for this target */
 		[NonSerialized]
 		public List<PluginInfo> BuildPlugins;
@@ -2329,7 +2333,7 @@ namespace UnrealBuildTool
 		/// <returns>Matching plugin, or null if not found</returns>
 		private PluginInfo FindPluginForModule(string ModuleName)
 		{
-			return BuildPlugins.FirstOrDefault(BuildPlugin => BuildPlugin.Descriptor.Modules.Any(Module => Module.Name == ModuleName));
+			return ValidPlugins.FirstOrDefault(ValidPlugin => ValidPlugin.Descriptor.Modules.Any(Module => Module.Name == ModuleName));
 		}
 
 		/**
@@ -2442,7 +2446,7 @@ namespace UnrealBuildTool
 		protected virtual void SetupPlugins()
 		{
 			// Filter the plugins list by the current project
-			List<PluginInfo> ValidPlugins = Plugins.ReadAvailablePlugins(UnrealBuildTool.GetUProjectFile());
+			ValidPlugins = Plugins.ReadAvailablePlugins(UnrealBuildTool.GetUProjectFile());
 
 			// Remove any plugins for platforms we don't have
 			foreach (UnrealTargetPlatform TargetPlatform in Enum.GetValues(typeof(UnrealTargetPlatform)))
