UEPLAT-368: Don't copy .so libs when only Java files are changed.
#android

diff --git a/Engine/Source/Programs/UnrealBuildTool/Android/UEDeployAndroid.cs b/Engine/Source/Programs/UnrealBuildTool/Android/UEDeployAndroid.cs
--- a/Engine/Source/Programs/UnrealBuildTool/Android/UEDeployAndroid.cs
+++ b/Engine/Source/Programs/UnrealBuildTool/Android/UEDeployAndroid.cs
@@ -513,7 +513,7 @@ namespace UnrealBuildTool.Android
 			string SourceSTLSOName = Environment.ExpandEnvironmentVariables("%NDKROOT%/sources/cxx-stl/gnu-libstdc++/") + GccVersion + "/libs/" + Arch + "/libgnustl_shared.so";
 			string FinalSTLSOName = UE4BuildPath + "/libs/" + Arch + "/libgnustl_shared.so";
 			Directory.CreateDirectory(Path.GetDirectoryName(FinalSTLSOName));
-			File.Copy(SourceSTLSOName, FinalSTLSOName);
+			File.Copy(SourceSTLSOName, FinalSTLSOName, true);
 		}
 
 		//@TODO: To enable the NVIDIA Gfx Debugger
@@ -1232,11 +1232,17 @@ namespace UnrealBuildTool.Android
 					if (File.Exists(NDKBuildPath))
 					{
 						string LibDir = UE4BuildPath + "/jni/" + GetNDKArch(Arch);
-						Directory.CreateDirectory(LibDir);
-
-						// copy the binary to the standard .so location
 						FinalSOName = LibDir + "/libUE4.so";
-						File.Copy(SourceSOName, FinalSOName, true);
+
+						// check to see if libUE4.so is newer than last time we copied
+						TimeSpan Diff = File.GetLastWriteTimeUtc(FinalSOName) - File.GetLastWriteTimeUtc(SourceSOName);
+						if (!File.Exists(FinalSOName) || Diff.TotalSeconds < -1 || Diff.TotalSeconds > 1)
+						{
+							Log.TraceInformation("\nCopying new .so file to jni folder...");
+							Directory.CreateDirectory(LibDir);
+							// copy the binary to the standard .so location
+							File.Copy(SourceSOName, FinalSOName, true);
+						}
 
 						FinalNdkBuildABICommand += GetNDKArch(Arch) + " ";
 					}
@@ -1244,38 +1250,52 @@ namespace UnrealBuildTool.Android
 					{
 						// if no NDK, we don't need any of the debugger stuff, so we just copy the .so to where it will end up
 						FinalSOName = UE4BuildPath + "/libs/" + GetNDKArch(Arch) + "/libUE4.so";
-						Directory.CreateDirectory(Path.GetDirectoryName(FinalSOName));
-						File.Copy(SourceSOName, FinalSOName);
+
+						// check to see if libUE4.so is newer than last time we copied
+						TimeSpan Diff = File.GetLastWriteTimeUtc(FinalSOName) - File.GetLastWriteTimeUtc(SourceSOName);
+						if (!File.Exists(FinalSOName) || Diff.TotalSeconds < -1 || Diff.TotalSeconds > 1)
+						{
+							Directory.CreateDirectory(Path.GetDirectoryName(FinalSOName));
+							File.Copy(SourceSOName, FinalSOName, true);
+						}
 					}
 
+					// remove any read only flags
+					FileInfo DestFileInfo = new FileInfo(FinalSOName);
+					DestFileInfo.Attributes = DestFileInfo.Attributes & ~FileAttributes.ReadOnly;
+					File.SetLastWriteTimeUtc(FinalSOName, File.GetLastWriteTimeUtc(SourceSOName));
+
 					// now do final stuff per apk (or after all .so's for a shared .apk)
 					if (bMakeSeparateApks || (ArchIndex == Arches.Length - 1 && GPUArchIndex == GPUArchitectures.Length - 1))
 					{
-						// always delete libs up to this point so fat binaries and incremental builds work together (otherwise we might end up with multiple
-						// so files in an apk that doesn't want them)
-						// note that we don't want to delete all libs, just the ones we copied (
-                        if (Directory.Exists(UE4BuildPath + "/libs"))
-                        {
-                            foreach (string Lib in Directory.EnumerateFiles(UE4BuildPath + "/libs", "libUE4*.so", SearchOption.AllDirectories))
-                            {
-                                File.Delete(Lib);
-                            }
-                        }
-
 						// if we need to run ndk-build, do it now (if making a shared .apk, we need to wait until all .libs exist)
 						if (!string.IsNullOrEmpty(FinalNdkBuildABICommand))
 						{
-							string CommandLine = "APP_ABI=\"" + FinalNdkBuildABICommand + "\"";
-							if (!bForDistribution)
+							string LibSOName = UE4BuildPath + "/libs/" + GetNDKArch(Arch) + "/libUE4.so";
+							// always delete libs up to this point so fat binaries and incremental builds work together (otherwise we might end up with multiple
+							// so files in an apk that doesn't want them)
+							// note that we don't want to delete all libs, just the ones we copied
+							TimeSpan Diff = File.GetLastWriteTimeUtc(LibSOName) - File.GetLastWriteTimeUtc(FinalSOName);
+							if (!File.Exists(LibSOName) || Diff.TotalSeconds < -1 || Diff.TotalSeconds > 1)
 							{
-								CommandLine += " NDK_DEBUG=1";
+								foreach (string Lib in Directory.EnumerateFiles(UE4BuildPath + "/libs", "libUE4*.so", SearchOption.AllDirectories))
+								{
+									File.Delete(Lib);
+								}
+
+								string CommandLine = "APP_ABI=\"" + FinalNdkBuildABICommand + "\"";
+								if (!bForDistribution)
+								{
+									CommandLine += " NDK_DEBUG=1";
+								}
+								RunCommandLineProgramAndThrowOnError(UE4BuildPath, NDKBuildPath, CommandLine, "Preparing native code for debugging...", true);
+
+								File.SetLastWriteTimeUtc(LibSOName, File.GetLastWriteTimeUtc(FinalSOName));
 							}
-							RunCommandLineProgramAndThrowOnError(UE4BuildPath, NDKBuildPath, CommandLine, "Preparing native code for debugging...", true);
-
 							// next loop we don't want to redo these ABIs
 							FinalNdkBuildABICommand = "";
 						}
-	
+
 						// after ndk-build is called, we can now copy in the stl .so (ndk-build deletes old files)
 						// copy libgnustl_shared.so to library (use 4.8 if possible, otherwise 4.6)
 						if (bMakeSeparateApks)
@@ -1292,11 +1312,6 @@ namespace UnrealBuildTool.Android
 							}
 						}
 
-
-						// remove any read only flags
-						FileInfo DestFileInfo = new FileInfo(FinalSOName);
-						DestFileInfo.Attributes = DestFileInfo.Attributes & ~FileAttributes.ReadOnly;
-
 						Log.TraceInformation("\n===={0}====PERFORMING FINAL APK PACKAGE OPERATION================================================", DateTime.Now.ToString());
 
 						string AntBuildType = "debug";
