Enable /Bt+ and /time+ in VC local builds when UEBuildWindows.bLogDetailedCompilerTimingInfo is true.

diff --git a/Engine/Source/Programs/UnrealBuildTool/Windows/UEBuildWindows.cs b/Engine/Source/Programs/UnrealBuildTool/Windows/UEBuildWindows.cs
--- a/Engine/Source/Programs/UnrealBuildTool/Windows/UEBuildWindows.cs
+++ b/Engine/Source/Programs/UnrealBuildTool/Windows/UEBuildWindows.cs
@@ -54,6 +54,12 @@ namespace UnrealBuildTool
             }
         }
 
+        /// <summary>
+        /// When true, throws some CL and link flags (/Bt+ and /link) to output detailed timing info.
+        /// </summary>
+        [XmlConfig]
+        public static bool bLogDetailedCompilerTimingInfo = false;
+
         /// True if we should use Clang/LLVM instead of MSVC to compile code on Windows platform
         public static readonly bool bCompileWithClang = false;
 
diff --git a/Engine/Source/Programs/UnrealBuildTool/Windows/VCToolChain.cs b/Engine/Source/Programs/UnrealBuildTool/Windows/VCToolChain.cs
--- a/Engine/Source/Programs/UnrealBuildTool/Windows/VCToolChain.cs
+++ b/Engine/Source/Programs/UnrealBuildTool/Windows/VCToolChain.cs
@@ -1001,6 +1001,11 @@ namespace UnrealBuildTool
 				}
 			}
 
+            if (!WindowsPlatform.bCompileWithClang && WindowsPlatform.bLogDetailedCompilerTimingInfo)
+            {
+                // Force MSVC
+                SharedArguments.Append(" /Bt+");
+            }
 
 			// Add preprocessor definitions to the argument list.
 			foreach (string Definition in CompileEnvironment.Config.Definitions)
@@ -1018,6 +1023,11 @@ namespace UnrealBuildTool
 			{
 				Action CompileAction = new Action(ActionType.Compile);
 				CompileAction.CommandDescription = "Compile";
+                // ensure compiler timings are captured when we execute the action.
+                if (!WindowsPlatform.bCompileWithClang && WindowsPlatform.bLogDetailedCompilerTimingInfo)
+                {
+                    CompileAction.bPrintDebugInfo = true;
+                }
 
 				StringBuilder FileArguments = new StringBuilder();
 				bool bIsPlainCFile = Path.GetExtension(SourceFile.AbsolutePath).ToUpperInvariant() == ".C";
@@ -1408,7 +1418,10 @@ namespace UnrealBuildTool
 				AppendLinkArguments(LinkEnvironment, Arguments);
 			}
 
-
+            if (!WindowsPlatform.bCompileWithClang && WindowsPlatform.bLogDetailedCompilerTimingInfo)
+            {
+                Arguments.Append(" /time+");
+            }
 
 			// If we're only building an import library, add the '/DEF' option that tells the LIB utility
 			// to simply create a .LIB file and .EXP file, and don't bother validating imports
@@ -1580,6 +1593,11 @@ namespace UnrealBuildTool
 			LinkAction.ProducedItems    .AddRange(ProducedItems);
 			LinkAction.PrerequisiteItems.AddRange(PrerequisiteItems);
 			LinkAction.StatusDescription  = Path.GetFileName(OutputFile.AbsolutePath);
+            // ensure compiler timings are captured when we execute the action.
+            if (!WindowsPlatform.bCompileWithClang && WindowsPlatform.bLogDetailedCompilerTimingInfo)
+            {
+                LinkAction.bPrintDebugInfo = true;
+            }
 
 			if( WindowsPlatform.bCompileWithClang )
 			{ 
