Copy libgnustl_shared.so only when it's updated.

diff --git a/Engine/Source/Programs/UnrealBuildTool/Android/UEDeployAndroid.cs b/Engine/Source/Programs/UnrealBuildTool/Android/UEDeployAndroid.cs
--- a/Engine/Source/Programs/UnrealBuildTool/Android/UEDeployAndroid.cs
+++ b/Engine/Source/Programs/UnrealBuildTool/Android/UEDeployAndroid.cs
@@ -512,8 +512,19 @@ namespace UnrealBuildTool.Android
 			// copy it in!
 			string SourceSTLSOName = Environment.ExpandEnvironmentVariables("%NDKROOT%/sources/cxx-stl/gnu-libstdc++/") + GccVersion + "/libs/" + Arch + "/libgnustl_shared.so";
 			string FinalSTLSOName = UE4BuildPath + "/libs/" + Arch + "/libgnustl_shared.so";
-			Directory.CreateDirectory(Path.GetDirectoryName(FinalSTLSOName));
-			File.Copy(SourceSTLSOName, FinalSTLSOName, true);
+
+			// check to see if libgnustl_shared.so is newer than last time we copied
+			bool bFileExists = File.Exists(FinalSTLSOName);
+			TimeSpan Diff = File.GetLastWriteTimeUtc(FinalSTLSOName) - File.GetLastWriteTimeUtc(SourceSTLSOName);
+			if (!bFileExists || Diff.TotalSeconds < -1 || Diff.TotalSeconds > 1)
+			{
+				if (bFileExists)
+				{
+					File.Delete(FinalSTLSOName);
+				}
+				Directory.CreateDirectory(Path.GetDirectoryName(FinalSTLSOName));
+				File.Copy(SourceSTLSOName, FinalSTLSOName, true);
+			}
 		}
 
 		//@TODO: To enable the NVIDIA Gfx Debugger
