Facebook added to the IOS settings UI

diff --git a/Engine/Source/Programs/UnrealBuildTool/IOS/UEDeployIOS.cs b/Engine/Source/Programs/UnrealBuildTool/IOS/UEDeployIOS.cs
--- a/Engine/Source/Programs/UnrealBuildTool/IOS/UEDeployIOS.cs
+++ b/Engine/Source/Programs/UnrealBuildTool/IOS/UEDeployIOS.cs
@@ -243,6 +243,15 @@ namespace UnrealBuildTool.IOS
 				MinVersion = "6.1";
 			}
 
+			// Get Facebook Support details
+			bool bEnableFacebookSupport = true;
+			Ini.GetBool("/Script/IOSRuntimeSettings.IOSRuntimeSettings", "bEnableFacebookSupport", out bEnableFacebookSupport);
+
+			// Write the Facebook App ID if we need it.
+			string FacebookAppID = "";
+			Ini.GetString("/Script/IOSRuntimeSettings.IOSRuntimeSettings", "FacebookAppID", out FacebookAppID);
+			bEnableFacebookSupport = bEnableFacebookSupport && !string.IsNullOrWhiteSpace(FacebookAppID);
+
 			// extra plist data
 			string ExtraData = "";
 			Ini.GetString("/Script/IOSRuntimeSettings.IOSRuntimeSettings", "AdditionalPlistData", out ExtraData);
@@ -261,6 +270,11 @@ namespace UnrealBuildTool.IOS
 			Text.AppendLine("\t\t\t<key>CFBundleURLSchemes</key>");
 			Text.AppendLine("\t\t\t<array>");
 			Text.AppendLine(string.Format("\t\t\t\t<string>{0}</string>", bIsUE4Game ? "UE4Game" : GameName));
+			if (bEnableFacebookSupport)
+			{
+				// This is needed for facebook login to redirect back to the app after completion.
+				Text.AppendLine(string.Format("\t\t\t\t<string>fb{0}</string>", FacebookAppID));
+			}
 			Text.AppendLine("\t\t\t</array>");
 			Text.AppendLine("\t\t</dict>");
 			Text.AppendLine("\t</array>");
@@ -418,6 +432,12 @@ namespace UnrealBuildTool.IOS
 			Text.AppendLine("\t</array>");
 			Text.AppendLine("\t<key>MinimumOSVersion</key>");
 			Text.AppendLine(string.Format("\t<string>{0}</string>", MinVersion));
+			if (bEnableFacebookSupport)
+			{
+				Text.AppendLine("\t<key>FacebookAppID</key>");
+				Text.AppendLine(string.Format("\t<string>{0}</string>", FacebookAppID));
+				
+			}
 			if (!string.IsNullOrEmpty(ExtraData))
 			{
 				ExtraData = ExtraData.Replace("\\n", "\n");
diff --git a/Engine/Source/Runtime/IOS/IOSRuntimeSettings/Classes/IOSRuntimeSettings.h b/Engine/Source/Runtime/IOS/IOSRuntimeSettings/Classes/IOSRuntimeSettings.h
--- a/Engine/Source/Runtime/IOS/IOSRuntimeSettings/Classes/IOSRuntimeSettings.h
+++ b/Engine/Source/Runtime/IOS/IOSRuntimeSettings/Classes/IOSRuntimeSettings.h
@@ -204,6 +204,14 @@ public:
 	UPROPERTY(GlobalConfig, EditAnywhere, Category = ExtraData)
 	FString AdditionalPlistData;
 
+	// Whether the app supports Facebook
+	UPROPERTY(GlobalConfig, EditAnywhere, Category = Online)
+	bool bEnableFacebookSupport;
+
+	// Facebook App ID obtained from Facebook's Developer Centre
+	UPROPERTY(GlobalConfig, EditAnywhere, Category = Online, meta = (EditCondition = "bEnableFacebookSupport"))
+	FString FacebookAppID;
+
 #if WITH_EDITOR
 	// UObject interface
 	virtual void PostEditChangeProperty(struct FPropertyChangedEvent& PropertyChangedEvent) override;
diff --git a/Engine/Source/Runtime/Online/OnlineSubsystemFacebook/Private/IOS/OnlineSubsystemFacebook.cpp b/Engine/Source/Runtime/Online/OnlineSubsystemFacebook/Private/IOS/OnlineSubsystemFacebook.cpp
--- a/Engine/Source/Runtime/Online/OnlineSubsystemFacebook/Private/IOS/OnlineSubsystemFacebook.cpp
+++ b/Engine/Source/Runtime/Online/OnlineSubsystemFacebook/Private/IOS/OnlineSubsystemFacebook.cpp
@@ -195,13 +195,21 @@ bool FOnlineSubsystemFacebook::Tick(float DeltaTime)
 bool FOnlineSubsystemFacebook::IsEnabled()
 {
 	// Check the ini for disabling Facebook
-	bool bEnableFacebook = false;
-	GConfig->GetBool(TEXT("OnlineSubsystemFacebook"), TEXT("bEnabled"), bEnableFacebook, GEngineIni);
+	bool bEnableFacebookSupport = false;
+	// Whilst we deprecate the old warning, let's still check if it's available
+	if (!GConfig->GetBool(TEXT("OnlineSubsystemFacebook"), TEXT("bEnabled"), bEnableFacebookSupport, GEngineIni))
+	{
+		GConfig->GetBool(TEXT("/Script/IOSRuntimeSettings.IOSRuntimeSettings"), TEXT("bEnableFacebookSupport"), bEnableFacebookSupport, GEngineIni);
+	}
+	else
+	{
+		UE_LOG(LogOnline, Warning, TEXT("The [OnlineSubsystemFacebook]:bEnabled flag has been moved to the IOS Project Settings, please update your preferences and App ID there."));
+	}
 
 #if !UE_BUILD_SHIPPING
 	// Check the commandline for disabling Facebook
-	bEnableFacebook = bEnableFacebook && !FParse::Param(FCommandLine::Get(),TEXT("NOFACEBOOK"));
+	bEnableFacebookSupport = bEnableFacebookSupport && !FParse::Param(FCommandLine::Get(), TEXT("NOFACEBOOK"));
 #endif
 
-	return bEnableFacebook;
+	return bEnableFacebookSupport;
 }
\ No newline at end of file
