UE-14803 - Android Java files constantly being rewritten triggering rebuilds
- corrected bug in code which was detecting if a re-write of the file was required or not to prevent this problem happening

diff --git a/Engine/Source/Programs/UnrealBuildTool/Android/UEBuildAndroid.cs b/Engine/Source/Programs/UnrealBuildTool/Android/UEBuildAndroid.cs
--- a/Engine/Source/Programs/UnrealBuildTool/Android/UEBuildAndroid.cs
+++ b/Engine/Source/Programs/UnrealBuildTool/Android/UEBuildAndroid.cs
@@ -57,12 +57,6 @@ namespace UnrealBuildTool
         private bool HasAnySDK()
         {
             string NDKPath = Environment.GetEnvironmentVariable("NDKROOT");
-/* Don't check for existence of env vars, always set them from the .ini values if they exist
-//			bool bNeedsNDKPath = string.IsNullOrEmpty(NDKPath);
-//			bool bNeedsAndroidHome = string.IsNullOrEmpty(Environment.GetEnvironmentVariable("ANDROID_HOME"));
-//			bool bNeedsAntHome = string.IsNullOrEmpty(Environment.GetEnvironmentVariable("ANT_HOME"));
- //           if((bNeedsNDKPath || bNeedsAndroidHome || bNeedsAntHome))
- */
             {
                 var configCacheIni = new ConfigCacheIni("Engine", null);
                 var AndroidEnv = new Dictionary<string, string>();
@@ -79,7 +73,6 @@ namespace UnrealBuildTool
                 {
                     if (configCacheIni.GetPath("/Script/AndroidPlatformEditor.AndroidSDKSettings", kvp.Value, out path) && !string.IsNullOrEmpty(path))
                     {
-//                        Log.TraceWarning("Adding {0} from ini as {1} to {2}", kvp.Value, path, kvp.Key);
                         AndroidEnv.Add(kvp.Key, path);
                     }
                     else
@@ -87,7 +80,6 @@ namespace UnrealBuildTool
                         var envValue = Environment.GetEnvironmentVariable(kvp.Key);
                         if(!String.IsNullOrEmpty(envValue))
                         {
-//                            Log.TraceWarning("Adding {0} from env as {1}", kvp.Key, envValue);
                             AndroidEnv.Add(kvp.Key, envValue);
                         }
                     }
diff --git a/Engine/Source/Programs/UnrealBuildTool/Android/UEDeployAndroid.cs b/Engine/Source/Programs/UnrealBuildTool/Android/UEDeployAndroid.cs
--- a/Engine/Source/Programs/UnrealBuildTool/Android/UEDeployAndroid.cs
+++ b/Engine/Source/Programs/UnrealBuildTool/Android/UEDeployAndroid.cs
@@ -397,7 +397,7 @@ namespace UnrealBuildTool.Android
             //
             obbData.Append("};\n"); // close class definition off
 
-            if (obbDataFile == null || !obbDataFile.ToString().SequenceEqual(obbData.ToString()))
+            if (obbDataFile == null || !obbDataFile.SequenceEqual(obbData.ToString().Split('\n')))
             {
                 MakeDirectoryIfRequired(FileName);
                 File.WriteAllText(FileName, obbData.ToString());
@@ -422,7 +422,7 @@ namespace UnrealBuildTool.Android
             Log.TraceInformation("\n==== Writing to shim file {0} ====", ShimFileName);
 
             // If they aren't the same then dump out the settings
-            if (DestFileContent == null || DestFileContent.ToString() != ShimFileContent.ToString())
+            if (DestFileContent == null || !DestFileContent.SequenceEqual((ShimFileContent.ToString()).Split('\n')))
             {
                 MakeDirectoryIfRequired(ShimFileName);
                 File.WriteAllText(ShimFileName, ShimFileContent.ToString());
