UBT: fix unused variable warning in KDevelop project gen.

diff --git a/Engine/Source/Programs/UnrealBuildTool/System/KDevelopGenerator.cs b/Engine/Source/Programs/UnrealBuildTool/System/KDevelopGenerator.cs
--- a/Engine/Source/Programs/UnrealBuildTool/System/KDevelopGenerator.cs
+++ b/Engine/Source/Programs/UnrealBuildTool/System/KDevelopGenerator.cs
@@ -402,9 +402,6 @@ namespace UnrealBuildTool
 		{
 			/// RAKE! Take one KDevelopProjectFileContent and pass
 			/// it through each function that writes out the sections.
-
-			int ProjectIndex = 0;
-
 			var KDevelopFileContent = new StringBuilder();
 			var KDevelopMasterFileContent = new StringBuilder ();
 
