OnlyModules list created during deserialization to prevent later crash.

diff --git a/Engine/Source/Programs/UnrealBuildTool/Configuration/UEBuildTarget.cs b/Engine/Source/Programs/UnrealBuildTool/Configuration/UEBuildTarget.cs
--- a/Engine/Source/Programs/UnrealBuildTool/Configuration/UEBuildTarget.cs
+++ b/Engine/Source/Programs/UnrealBuildTool/Configuration/UEBuildTarget.cs
@@ -7,6 +7,7 @@ using System.Text;
 using System.Diagnostics;
 using System.IO;
 using System.Xml;
+using System.Runtime.Serialization;
 
 namespace UnrealBuildTool
 {
@@ -177,8 +178,16 @@ namespace UnrealBuildTool
 	 * A target that can be built
 	 */
 	[Serializable]
-	public class UEBuildTarget
+	public class UEBuildTarget : IDeserializationCallback
 	{
+		void IDeserializationCallback.OnDeserialization(object sender)
+		{
+			if (OnlyModules == null)
+			{
+				OnlyModules = new List<OnlyModule>();
+			}
+		}
+
 		public string GetAppName()
 		{
 			return AppName;
