Use a single flat .dSYM on OS X for debug symbols on Josh's advice, this can reside next to the corresponding binary file  without breaking the MAX_PATH limit.

diff --git a/Engine/Source/Programs/UnrealBuildTool/Mac/MacToolChain.cs b/Engine/Source/Programs/UnrealBuildTool/Mac/MacToolChain.cs
--- a/Engine/Source/Programs/UnrealBuildTool/Mac/MacToolChain.cs
+++ b/Engine/Source/Programs/UnrealBuildTool/Mac/MacToolChain.cs
@@ -1176,12 +1176,18 @@ namespace UnrealBuildTool
 		public FileItem GenerateDebugInfo(FileItem MachOBinary)
 		{
 			// Make a file item for the source and destination files
-			string FullDestPath = MachOBinary.AbsolutePath + ".dSYM";
+			string FullDestPath = Path.ChangeExtension(MachOBinary.AbsolutePath, ".dSYM");
 
 			FileItem OutputFile = FileItem.GetItemByPath(FullDestPath);
 			FileItem DestFile = LocalToRemoteFileItem(OutputFile, false);
 			FileItem InputFile = LocalToRemoteFileItem(MachOBinary, false);
 
+			// Delete on the local machine
+			if(Directory.Exists(OutputFile.AbsolutePath))
+			{
+				Directory.Delete(OutputFile.AbsolutePath, true);
+			}
+
 			// Make the compile action
 			Action GenDebugAction = new Action(ActionType.GenerateDebugInfo);
 			if (BuildHostPlatform.Current.Platform != UnrealTargetPlatform.Mac)
@@ -1192,8 +1198,9 @@ namespace UnrealBuildTool
 			GenDebugAction.WorkingDirectory = Path.GetFullPath(".");
 			GenDebugAction.CommandPath = "sh";
 
+			// Deletes ay existing file on the building machine,
 			// note that the source and dest are switched from a copy command
-			GenDebugAction.CommandArguments = string.Format("-c '\"{0}\"usr/bin/xcrun dsymutil \"{1}\" -o \"{2}\"; \"{0}\"usr/bin/xcrun strip -S -X -x \"{1}\"'",
+			GenDebugAction.CommandArguments = string.Format("-c 'rm -rf \"{2}\"; \"{0}\"usr/bin/xcrun dsymutil -f \"{1}\" -o \"{2}\"'",
 				XcodeDeveloperDir,
 				InputFile.AbsolutePath,
 				DestFile.AbsolutePath);
@@ -1331,31 +1338,6 @@ namespace UnrealBuildTool
 
         public override void AddFilesToReceipt(BuildReceipt Receipt, UEBuildBinary Binary)
 		{
-			string DebugExtension = UEBuildPlatform.GetBuildPlatform(Binary.Target.Platform).GetDebugInfoExtension(Binary.Config.Type);
-			if(DebugExtension == ".dsym")
-			{
-				for (int i = 0; i < Receipt.BuildProducts.Count; i++)
-				{
-					if(Receipt.BuildProducts[i].Type == BuildProductType.Executable || Receipt.BuildProducts[i].Type == BuildProductType.DynamicLibrary)
-					{
-						string OutputFilePath = Receipt.BuildProducts[i].Path;
-						string DsymInfo = OutputFilePath + ".dSYM/Contents/Info.plist";
-						Receipt.AddBuildProduct(DsymInfo, BuildProductType.SymbolFile);
-
-						string DsymDylib = OutputFilePath + ".dSYM/Contents/Resources/DWARF/" + Path.GetFileName(OutputFilePath);
-						Receipt.AddBuildProduct(DsymDylib, BuildProductType.SymbolFile);
-					}
-				}
-
-				for (int i = 0; i < Receipt.BuildProducts.Count; i++)
-				{
-					if(Path.GetExtension(Receipt.BuildProducts[i].Path) == DebugExtension)
-					{
-						Receipt.BuildProducts.RemoveAt(i--);
-					}
-				}
-			}
-
 			if (Binary.Target.GlobalLinkEnvironment.Config.bIsBuildingConsoleApplication)
 			{
 				return;
