[INTEGRATE] Change 2534938 by [REDACTED_PERSON][REDACTED_HANDLE] on 2015/05/02 14:39:44

	Fix crash with OnlyModules being null; not sure how UBT got into this state.

diff --git a/Engine/Source/Programs/UnrealBuildTool/Configuration/UEBuildTarget.cs b/Engine/Source/Programs/UnrealBuildTool/Configuration/UEBuildTarget.cs
--- a/Engine/Source/Programs/UnrealBuildTool/Configuration/UEBuildTarget.cs
+++ b/Engine/Source/Programs/UnrealBuildTool/Configuration/UEBuildTarget.cs
@@ -1394,7 +1394,7 @@ namespace UnrealBuildTool
 			if(Receipt != null)
 			{
 				IUEBuildPlatform BuildPlatform = UEBuildPlatform.GetBuildPlatform(Platform);
-				if(OnlyModules.Count == 0)
+				if(OnlyModules == null || OnlyModules.Count == 0)
 				{
 					string ReceiptFileName = BuildReceipt.GetDefaultPath(ProjectDirectory, TargetName, Platform, Configuration, BuildPlatform.GetActiveArchitecture());
 					Directory.CreateDirectory(Path.GetDirectoryName(ReceiptFileName));
