Don't remove plugins for the 'Desktop' target platform if we don't have the UEBuildPlatform registered. It breaks a LOT of things if you download the engine to your desktop.

diff --git a/Engine/Source/Programs/UnrealBuildTool/Configuration/UEBuildTarget.cs b/Engine/Source/Programs/UnrealBuildTool/Configuration/UEBuildTarget.cs
--- a/Engine/Source/Programs/UnrealBuildTool/Configuration/UEBuildTarget.cs
+++ b/Engine/Source/Programs/UnrealBuildTool/Configuration/UEBuildTarget.cs
@@ -2525,7 +2525,7 @@ namespace UnrealBuildTool
 			// Remove any plugins for platforms we don't have
 			foreach (UnrealTargetPlatform TargetPlatform in Enum.GetValues(typeof(UnrealTargetPlatform)))
 			{
-				if (UEBuildPlatform.GetBuildPlatform(TargetPlatform, true) == null)
+				if (TargetPlatform != UnrealTargetPlatform.Desktop && UEBuildPlatform.GetBuildPlatform(TargetPlatform, true) == null)
 				{
 					string DirectoryFragment = String.Format("/{0}/", TargetPlatform.ToString());
 					ValidPlugins.RemoveAll(x => x.Directory.Replace('\\', '/').Contains(DirectoryFragment));
