Allow using both Win32 and Win64 in Rocket builds.

diff --git a/Engine/Source/Editor/MainFrame/Private/Frame/MainFrameActions.cpp b/Engine/Source/Editor/MainFrame/Private/Frame/MainFrameActions.cpp
--- a/Engine/Source/Editor/MainFrame/Private/Frame/MainFrameActions.cpp
+++ b/Engine/Source/Editor/MainFrame/Private/Frame/MainFrameActions.cpp
@@ -622,24 +622,6 @@ void FMainFrameActionCallbacks::PackageProject( const FName InPlatformInfoName )
 
 bool FMainFrameActionCallbacks::PackageProjectCanExecute( const FName PlatformInfoName )
 {
-	// For a binary Rocket build, Development is ALWAYS Win64, and Shipping is ALWAYS Win32.
-	if(FRocketSupport::IsRocket())
-	{
-		const EProjectPackagingBuildConfigurations BuildConfiguration = GetDefault<UProjectPackagingSettings>()->BuildConfiguration;
-		switch(BuildConfiguration)
-		{
-		case PPBC_DebugGame: // DebugGame uses the same libs as Development, so fall through and apply the same validation logic
-		case PPBC_Development:
-			return PlatformInfoName != "WindowsNoEditor_Win32";
-
-		case PPBC_Shipping:
-			return PlatformInfoName != "WindowsNoEditor_Win64";
-
-		default:
-			break;
-		}
-	}
-
 	return true;
 }
 
diff --git a/Engine/Source/Programs/UnrealBuildTool/System/VCProject.cs b/Engine/Source/Programs/UnrealBuildTool/System/VCProject.cs
--- a/Engine/Source/Programs/UnrealBuildTool/System/VCProject.cs
+++ b/Engine/Source/Programs/UnrealBuildTool/System/VCProject.cs
@@ -242,34 +242,6 @@ namespace UnrealBuildTool
 				return false;
 			}
 
-			// For Rocket projects, we currently only allow 64-bit versions of Debug and Development binaries to be compiled
-			// and only 32-bit versions of Shipping binaries to be compiled.  This has to do with our choice of which build
-			// configurations we want to ship with that product.  So to simplify things, we're merging Win32 and Win64
-			// configurations together in the IDE under "Windows".  In order for this to work properly, we need to make
-			// sure that Rocket projects only report the respective Windows platform that is valid for each configuration
-			if( UnrealBuildTool.RunningRocket() && ProjectTarget.TargetRules != null )
-			{
-				if( Platform == UnrealTargetPlatform.Win32 || Platform == UnrealTargetPlatform.Win64 )
-				{
-					// In Rocket, shipping game targets are always 32-bit
-					if( Configuration == UnrealTargetConfiguration.Shipping && ProjectTarget.TargetRules.Type == TargetRules.TargetType.Game )
-					{
-						if( Platform != UnrealTargetPlatform.Win32 )
-						{
-							return false;
-						}
-					}
-					else
-					{
-						// Everything else is expecting to be 64-bit
-						if( Platform != UnrealTargetPlatform.Win64 )
-						{
-							return false;
-						}
-					}
-				}
-			}
-
 			return true;
 		}
 
diff --git a/Engine/Source/Programs/UnrealBuildTool/System/VCProjectFileGenerator.cs b/Engine/Source/Programs/UnrealBuildTool/System/VCProjectFileGenerator.cs
--- a/Engine/Source/Programs/UnrealBuildTool/System/VCProjectFileGenerator.cs
+++ b/Engine/Source/Programs/UnrealBuildTool/System/VCProjectFileGenerator.cs
@@ -459,28 +459,6 @@ namespace UnrealBuildTool
 
 								var SolutionPlatformName = CurPlatform.ToString();
 
-								// For Rocket, there are currently no targets that are valid to build both for Win32 and Win64.  So we simply things by
-								// only displaying a "Windows" platform and building for the appropriate Windows platform automatically based on whichever
-								// configuration they have selected.
-								if( UnrealBuildTool.RunningRocket() && ( CurPlatform == UnrealTargetPlatform.Win32 || CurPlatform == UnrealTargetPlatform.Win64 ) )
-								{
-									SolutionPlatformName = "Windows";
-									if( Configuration == UnrealTargetConfiguration.Shipping )
-									{
-										if(CurPlatform != UnrealTargetPlatform.Win32)
-										{
-											continue;
-										}
-									}
-									else
-									{
-										if(CurPlatform != UnrealTargetPlatform.Win64)
-										{
-											continue;
-										}
-									}
-								}
-
 								var SolutionConfigAndPlatformPair = SolutionConfigName + "|" + SolutionPlatformName;
 								SolutionConfigCombinations.Add(
 										new VCSolutionConfigCombination
