Allow adding to the manifest node in generated AndroidManifest.xml
#ue4
#android

diff --git a/Engine/Source/Programs/UnrealBuildTool/Android/UEDeployAndroid.cs b/Engine/Source/Programs/UnrealBuildTool/Android/UEDeployAndroid.cs
--- a/Engine/Source/Programs/UnrealBuildTool/Android/UEDeployAndroid.cs
+++ b/Engine/Source/Programs/UnrealBuildTool/Android/UEDeployAndroid.cs
@@ -866,6 +866,12 @@ namespace UnrealBuildTool.Android
 			Ini.GetString("/Script/AndroidRuntimeSettings.AndroidRuntimeSettings", "VersionDisplayName", out VersionDisplayName);
 			string Orientation;
 			Ini.GetString("/Script/AndroidRuntimeSettings.AndroidRuntimeSettings", "Orientation", out Orientation);
+			List<string> ExtraManifestNodeTags;
+			Ini.GetArray("/Script/AndroidRuntimeSettings.AndroidRuntimeSettings", "ExtraManifestNodeTags", out ExtraManifestNodeTags);
+			List<string> ExtraApplicationNodeTags;
+			Ini.GetArray("/Script/AndroidRuntimeSettings.AndroidRuntimeSettings", "ExtraApplicationNodeTags", out ExtraApplicationNodeTags);
+			List<string> ExtraActivityNodeTags;
+			Ini.GetArray("/Script/AndroidRuntimeSettings.AndroidRuntimeSettings", "ExtraActivityNodeTags", out ExtraActivityNodeTags);
 			string ExtraActivitySettings;
 			Ini.GetString("/Script/AndroidRuntimeSettings.AndroidRuntimeSettings", "ExtraActivitySettings", out ExtraActivitySettings);
 			string ExtraApplicationSettings;
@@ -881,13 +887,28 @@ namespace UnrealBuildTool.Android
 			Text.AppendLine("<?xml version=\"1.0\" encoding=\"utf-8\"?>");
 			Text.AppendLine("<manifest xmlns:android=\"http://schemas.android.com/apk/res/android\"");
 			Text.AppendLine(string.Format("          package=\"{0}\"", PackageName));
+			if (ExtraManifestNodeTags != null)
+			{
+				foreach (string Line in ExtraManifestNodeTags)
+				{
+					Text.AppendLine("          " + Line);
+				}
+			}
 			Text.AppendLine(string.Format("          android:versionCode=\"{0}\"", StoreVersion));
 			Text.AppendLine(string.Format("          android:versionName=\"{0}\">", VersionDisplayName));
+			
 			Text.AppendLine("");
 
 			Text.AppendLine("\t<!-- Application Definition -->");
 			Text.AppendLine("\t<application android:label=\"@string/app_name\"");
 			Text.AppendLine("\t             android:icon=\"@drawable/icon\"");
+			if (ExtraApplicationNodeTags != null)
+			{
+				foreach (string Line in ExtraApplicationNodeTags)
+				{
+					Text.AppendLine("\t             " + Line);
+				}
+			}
 			Text.AppendLine("\t             android:hasCode=\"true\">");
 			Text.AppendLine("\t\t<activity android:name=\"com.epicgames.ue4.GameActivity\"");
 			Text.AppendLine("\t\t          android:label=\"@string/app_name\"");
@@ -909,6 +930,13 @@ namespace UnrealBuildTool.Android
 			}
 			Text.AppendLine("\t\t          android:launchMode=\"singleTask\"");
 			Text.AppendLine(string.Format("\t\t          android:screenOrientation=\"{0}\"", ConvertOrientationIniValue(Orientation)));
+			if (ExtraActivityNodeTags != null)
+			{
+				foreach (string Line in ExtraActivityNodeTags)
+				{
+					Text.AppendLine("\t\t          " + Line);
+				}
+			}
 			Text.AppendLine(string.Format("\t\t          android:debuggable=\"{0}\">", bIsForDistribution ? "false" : "true"));
 			Text.AppendLine("\t\t\t<meta-data android:name=\"android.app.lib_name\" android:value=\"UE4\"/>");
 			Text.AppendLine("\t\t\t<intent-filter>");
diff --git a/Engine/Source/Runtime/Android/AndroidRuntimeSettings/Classes/AndroidRuntimeSettings.h b/Engine/Source/Runtime/Android/AndroidRuntimeSettings/Classes/AndroidRuntimeSettings.h
--- a/Engine/Source/Runtime/Android/AndroidRuntimeSettings/Classes/AndroidRuntimeSettings.h
+++ b/Engine/Source/Runtime/Android/AndroidRuntimeSettings/Classes/AndroidRuntimeSettings.h
@@ -136,10 +136,22 @@ public:
 	UPROPERTY(GlobalConfig, EditAnywhere, Category = APKPackaging, Meta = (DisplayName = "Preferred Depth Buffer format"))
 	TEnumAsByte<EAndroidDepthBufferPreference::Type> DepthBufferPreference;
 
+	// Any extra tags for the <manifest> node
+	UPROPERTY(GlobalConfig, EditAnywhere, Category = AdvancedAPKPackaging, Meta = (DisplayName = "Extra Tags for <manifest> node"))
+	TArray<FString> ExtraManifestNodeTags;
+
+	// Any extra tags for the <application> node
+	UPROPERTY(GlobalConfig, EditAnywhere, Category = AdvancedAPKPackaging, Meta = (DisplayName = "Extra Tags for <application> node"))
+	TArray<FString> ExtraApplicationNodeTags;
+
+	// Any extra tags for the com.epicgames.UE4.GameActivity <activity> node
 	// Any extra settings for the <application> section (an optional file <Project>/Build/Android/ManifestApplicationAdditions.txt will also be included)
 	UPROPERTY(GlobalConfig, EditAnywhere, Category = AdvancedAPKPackaging, Meta = (DisplayName = "Extra Settings for <application> section (\\n to separate lines)"))
 	FString ExtraApplicationSettings;
 
+	UPROPERTY(GlobalConfig, EditAnywhere, Category = AdvancedAPKPackaging, Meta = (DisplayName = "Extra Tags for UE4.GameActivity <activity> node"))
+	TArray<FString> ExtraActivityNodeTags;
+
 	// Any extra settings for the main <activity> section (an optional file <Project>/Build/Android/ManifestApplicationActivtyAdditions.txt will also be included)
 	UPROPERTY(GlobalConfig, EditAnywhere, Category = AdvancedAPKPackaging, Meta = (DisplayName = "Extra Settings for <activity> section (\\n to separate lines)"))
 	FString ExtraActivitySettings;
