Fix for Mac issue where it is using rsync when locally generating project files on the mac.

diff --git a/Engine/Source/Programs/UnrealBuildTool/ToolChain/RemoteToolChain.cs b/Engine/Source/Programs/UnrealBuildTool/ToolChain/RemoteToolChain.cs
--- a/Engine/Source/Programs/UnrealBuildTool/ToolChain/RemoteToolChain.cs
+++ b/Engine/Source/Programs/UnrealBuildTool/ToolChain/RemoteToolChain.cs
@@ -233,7 +233,8 @@ namespace UnrealBuildTool
 		// Gather a users root path from the remote server. Should only be called once.
 		public static void SetUserDevRootFromServer()
 		{
-			if (!bUseRPCUtil)
+
+			if (!bUseRPCUtil && BuildHostPlatform.Current.Platform != UnrealTargetPlatform.Mac)
 			{
 				// Only set relative to the users root when using rsync, for now
 				Hashtable Results = RPCUtilHelper.Command("/", "echo $HOME", null);
