Initial Unreal Build Tool support for Visual Studio 2015 and UAP (disabled by default)
- Merged in Microsoft's Unreal Build Tool changes for VS 2015 support, and Universal App Platform support
- VS 2015 support is disabled by default (the engine is not yet compiling with VS 2015)
- Use the new '-2015' option when generating project files to enable VS 2015 support
- Windows SDK 8.1 is used by default.  To use Windows SDK 10, enable WindowsPlatform.bUseWindowsSDK10
- UAP support is disabled (not supported yet, work in progress.)  Use WinUAPPlatform.bEnableUAPSupport to enable it.
- Various loose ends still remain (search for "[REDACTED_HANDLE] UAP" in Unreal Build Tool code)

diff --git a/Engine/Source/Programs/UnrealBuildTool/Android/AndroidProjectGenerator.cs b/Engine/Source/Programs/UnrealBuildTool/Android/AndroidProjectGenerator.cs
--- a/Engine/Source/Programs/UnrealBuildTool/Android/AndroidProjectGenerator.cs
+++ b/Engine/Source/Programs/UnrealBuildTool/Android/AndroidProjectGenerator.cs
@@ -34,16 +34,8 @@ namespace UnrealBuildTool
 
 			string ProgramFilesPath = Environment.GetFolderPath(Environment.SpecialFolder.ProgramFilesX86);
 
-			string PlatformToolsetVersion;
-			if (VCProjectFileGenerator.ProjectFileFormat == VCProjectFileGenerator.VCProjectFileFormat.VisualStudio2013)
-			{
-				PlatformToolsetVersion = "v120";
-			}
-			else if (VCProjectFileGenerator.ProjectFileFormat == VCProjectFileGenerator.VCProjectFileFormat.VisualStudio2012)
-			{
-				PlatformToolsetVersion = "v110";
-			}
-			else
+			string PlatformToolsetVersion = VCProjectFileGenerator.ProjectFilePlatformToolsetVersionString;
+			if( String.IsNullOrEmpty( PlatformToolsetVersion ) )
 			{
 				// future maintainer: add toolset version and verify that the rest of the msbuild path, version, and location in ProgramFiles(x86) is still valid
 				Log.TraceInformation("Android project generation needs to be updated for this version of Visual Studio.");
diff --git a/Engine/Source/Programs/UnrealBuildTool/Configuration/EngineConfiguration.cs b/Engine/Source/Programs/UnrealBuildTool/Configuration/EngineConfiguration.cs
--- a/Engine/Source/Programs/UnrealBuildTool/Configuration/EngineConfiguration.cs
+++ b/Engine/Source/Programs/UnrealBuildTool/Configuration/EngineConfiguration.cs
@@ -577,7 +577,7 @@ namespace UnrealBuildTool
 		/// </summary>
 		private static string GetIniPlatformName(UnrealTargetPlatform TargetPlatform)
 		{
-			if (TargetPlatform == UnrealTargetPlatform.Win32 || TargetPlatform == UnrealTargetPlatform.Win64)
+			if (TargetPlatform == UnrealTargetPlatform.Win32 || TargetPlatform == UnrealTargetPlatform.Win64 || TargetPlatform == UnrealTargetPlatform.WinUAP )
 			{
 				return "Windows";
 			}
diff --git a/Engine/Source/Programs/UnrealBuildTool/Configuration/UEBuildPlatform.cs b/Engine/Source/Programs/UnrealBuildTool/Configuration/UEBuildPlatform.cs
--- a/Engine/Source/Programs/UnrealBuildTool/Configuration/UEBuildPlatform.cs
+++ b/Engine/Source/Programs/UnrealBuildTool/Configuration/UEBuildPlatform.cs
@@ -436,6 +436,7 @@ namespace UnrealBuildTool
 					return ":";
 				case UnrealTargetPlatform.Win32:
 				case UnrealTargetPlatform.Win64:
+				case UnrealTargetPlatform.WinUAP:
 					return ";";
 				default:
 					Log.TraceWarning("PATH var delimiter unknown for platform " + BuildHostPlatform.Current.Platform.ToString() + " using ';'");
diff --git a/Engine/Source/Programs/UnrealBuildTool/Configuration/UEBuildTarget.cs b/Engine/Source/Programs/UnrealBuildTool/Configuration/UEBuildTarget.cs
--- a/Engine/Source/Programs/UnrealBuildTool/Configuration/UEBuildTarget.cs
+++ b/Engine/Source/Programs/UnrealBuildTool/Configuration/UEBuildTarget.cs
@@ -16,13 +16,14 @@ namespace UnrealBuildTool
 		Unknown,
 		Win32,
 		Win64,
+		WinRT,
+		WinRT_ARM,
+		WinUAP,
 		Mac,
 		XboxOne,
 		PS4,
 		IOS,
 		Android,
-		WinRT,
-		WinRT_ARM,
 		HTML5,
         Linux,
 		Desktop,
@@ -204,6 +205,7 @@ namespace UnrealBuildTool
 			{
 				case CPPTargetPlatform.Win32:			return UnrealTargetPlatform.Win32;
 				case CPPTargetPlatform.Win64:			return UnrealTargetPlatform.Win64;
+				case CPPTargetPlatform.WinUAP:          return UnrealTargetPlatform.WinUAP;
 				case CPPTargetPlatform.Mac:				return UnrealTargetPlatform.Mac;
 				case CPPTargetPlatform.XboxOne:			return UnrealTargetPlatform.XboxOne;
 				case CPPTargetPlatform.PS4:				return UnrealTargetPlatform.PS4;
@@ -1446,7 +1448,7 @@ namespace UnrealBuildTool
                 }
                 else
                 {
-                    IsCurrentPlatform = Platform == UnrealTargetPlatform.Win64 || Platform == UnrealTargetPlatform.Win32;
+                    IsCurrentPlatform = Platform == UnrealTargetPlatform.Win64 || Platform == UnrealTargetPlatform.Win32 || Platform == UnrealTargetPlatform.WinUAP;
 
                 }
 
diff --git a/Engine/Source/Programs/UnrealBuildTool/System/CPPEnvironment.cs b/Engine/Source/Programs/UnrealBuildTool/System/CPPEnvironment.cs
--- a/Engine/Source/Programs/UnrealBuildTool/System/CPPEnvironment.cs
+++ b/Engine/Source/Programs/UnrealBuildTool/System/CPPEnvironment.cs
@@ -14,6 +14,7 @@ namespace UnrealBuildTool
 	{
 		Win32,
 		Win64,
+		WinUAP,
 		Mac,
 		XboxOne,
 		PS4,
diff --git a/Engine/Source/Programs/UnrealBuildTool/System/ProjectFileGenerator.cs b/Engine/Source/Programs/UnrealBuildTool/System/ProjectFileGenerator.cs
--- a/Engine/Source/Programs/UnrealBuildTool/System/ProjectFileGenerator.cs
+++ b/Engine/Source/Programs/UnrealBuildTool/System/ProjectFileGenerator.cs
@@ -1483,6 +1483,10 @@ namespace UnrealBuildTool
 					{
 						if (UnrealBuildTool.IsValidPlatform(Platform))
 						{
+							// @todo UAP: Why would we exclude UAP from project files based on SDK install status?  We don't do this with anything else, right?
+                            if (BuildPlatform is WinUAPPlatform && ((WinUAPPlatform)BuildPlatform).HasRequiredSDKsInstalled() != SDKStatus.Valid)
+                                continue;
+
 							SupportedPlatforms.Add(Platform);
 
 							if (SupportedPlatformsString.Length > 0)
diff --git a/Engine/Source/Programs/UnrealBuildTool/System/UnrealBuildTool.cs b/Engine/Source/Programs/UnrealBuildTool/System/UnrealBuildTool.cs
--- a/Engine/Source/Programs/UnrealBuildTool/System/UnrealBuildTool.cs
+++ b/Engine/Source/Programs/UnrealBuildTool/System/UnrealBuildTool.cs
@@ -198,22 +198,22 @@ namespace UnrealBuildTool
             return UnrealBuildToolPath;
         }
 
-        // @todo projectfiles: Move this into the ProjectPlatformGeneration class?
-        /// <summary>
-        /// IsDesktopPlatform
-        /// </summary>
-        /// <param name="InPlatform">The platform of interest</param>
-        /// <returns>True if the given platform is a 'desktop' platform</returns>
-        static public bool IsDesktopPlatform(UnrealTargetPlatform InPlatform)
-        {
-            // Windows and Mac are desktop platforms.
-            return (
-                (InPlatform == UnrealTargetPlatform.Win64) ||
-                (InPlatform == UnrealTargetPlatform.Win32) ||
-                (InPlatform == UnrealTargetPlatform.Linux) ||
-                (InPlatform == UnrealTargetPlatform.Mac)
-                );
-        }
+		// @todo projectfiles: Move this into the ProjectPlatformGeneration class?
+		/// <summary>
+		/// IsDesktopPlatform
+		/// </summary>
+		/// <param name="InPlatform">The platform of interest</param>
+		/// <returns>True if the given platform is a 'desktop' platform</returns>
+		static public bool IsDesktopPlatform(UnrealTargetPlatform InPlatform)
+		{
+			// Windows and Mac are desktop platforms.
+			return (
+				(InPlatform == UnrealTargetPlatform.Win64) ||
+				(InPlatform == UnrealTargetPlatform.Win32) ||
+				(InPlatform == UnrealTargetPlatform.Linux) ||
+				(InPlatform == UnrealTargetPlatform.Mac)
+				);
+		}
 
         // @todo projectfiles: Move this into the ProjectPlatformGeneration class?
         /// <summary>
@@ -249,21 +249,22 @@ namespace UnrealBuildTool
                 );
         }
 
-        // @todo projectfiles: Move this into the ProjectPlatformGeneration class?
-        /// <summary>
-        /// PlatformSupportsCrashReporter
-        /// </summary>
-        /// <param name="InPlatform">The platform of interest</param>
-        /// <returns>True if the given platform supports a crash reporter client (i.e. it can be built for it)</returns>
-        static public bool PlatformSupportsCrashReporter(UnrealTargetPlatform InPlatform)
-        {
-            return (
-                (InPlatform == UnrealTargetPlatform.Win64) ||
-                (InPlatform == UnrealTargetPlatform.Win32) ||
-                (InPlatform == UnrealTargetPlatform.Linux) ||
-                (InPlatform == UnrealTargetPlatform.Mac)
-                );
-        }
+		// @todo projectfiles: Move this into the ProjectPlatformGeneration class?
+		/// <summary>
+		/// PlatformSupportsCrashReporter
+		/// </summary>
+		/// <param name="InPlatform">The platform of interest</param>
+		/// <returns>True if the given platform supports a crash reporter client (i.e. it can be built for it)</returns>
+		static public bool PlatformSupportsCrashReporter(UnrealTargetPlatform InPlatform)
+		{
+			return (
+				(InPlatform == UnrealTargetPlatform.WinUAP) ||
+				(InPlatform == UnrealTargetPlatform.Win64) ||
+				(InPlatform == UnrealTargetPlatform.Win32) ||
+				(InPlatform == UnrealTargetPlatform.Linux) ||
+				(InPlatform == UnrealTargetPlatform.Mac)
+				);
+		}
 
         /// <summary>
         /// Gets all platforms that satisfies the predicate.
@@ -356,33 +357,33 @@ namespace UnrealBuildTool
             return true;
         }
 
-        /// <summary>
-        /// Is this a valid platform
-        /// Used primarily for Rocket vs non-Rocket
-        /// </summary>
-        /// <param name="InPlatform"></param>
-        /// <returns>true if valid, false if not</returns>
-        static public bool IsValidPlatform(UnrealTargetPlatform InPlatform)
-        {
-            if (RunningRocket())
-            {
-                if(Utils.IsRunningOnMono)
-                {
-                    if (InPlatform != UnrealTargetPlatform.Mac && InPlatform != UnrealTargetPlatform.IOS && InPlatform != UnrealTargetPlatform.Linux)
-                    {
-                        return false;
-                    }
-                }
-                else
-                {
-                    if(InPlatform != UnrealTargetPlatform.Win32 && InPlatform != UnrealTargetPlatform.Win64 && InPlatform != UnrealTargetPlatform.Android)
-                    {
-                        return false;
-                    }
-                }
-            }
-            return true;
-        }
+		/// <summary>
+		/// Is this a valid platform
+		/// Used primarily for Rocket vs non-Rocket
+		/// </summary>
+		/// <param name="InPlatform"></param>
+		/// <returns>true if valid, false if not</returns>
+		static public bool IsValidPlatform(UnrealTargetPlatform InPlatform)
+		{
+			if (RunningRocket())
+			{
+				if(Utils.IsRunningOnMono)
+				{
+					if (InPlatform != UnrealTargetPlatform.Mac && InPlatform != UnrealTargetPlatform.IOS && InPlatform != UnrealTargetPlatform.Linux)
+					{
+						return false;
+					}
+				}
+				else
+				{
+					if(InPlatform != UnrealTargetPlatform.Win32 && InPlatform != UnrealTargetPlatform.Win64 && InPlatform != UnrealTargetPlatform.Android)
+					{
+						return false;
+					}
+				}
+			}
+			return true;
+		}
 
 
         /// <summary>
diff --git a/Engine/Source/Programs/UnrealBuildTool/System/Utils.cs b/Engine/Source/Programs/UnrealBuildTool/System/Utils.cs
--- a/Engine/Source/Programs/UnrealBuildTool/System/Utils.cs
+++ b/Engine/Source/Programs/UnrealBuildTool/System/Utils.cs
@@ -241,8 +241,9 @@ namespace UnrealBuildTool
 		 * Sets the environment variables from the passed in batch file
 		 * 
 		 * @param	BatchFileName	Name of the batch file to parse
+		 * @param	Parameters		Optional command-line parameters to pass to the batch file when running it
 		 */
-		public static void SetEnvironmentVariablesFromBatchFile(string BatchFileName)
+		public static void SetEnvironmentVariablesFromBatchFile(string BatchFileName, string Parameters = "")
 		{
 			// @todo ubtmake: Experiment with changing this to run asynchronously at startup, and only blocking if accessed before the .bat file finishes
 			if( File.Exists( BatchFileName ) )
@@ -267,7 +268,7 @@ namespace UnrealBuildTool
 					var ShortEnvVarsToXMLExePath = GetShortPathName(EnvVarsToXMLExePath);
 
 					// Run 'vcvars32.bat' (or similar x64 version) to set environment variables
-					EnvReaderBatchFileContent.Add( String.Format( "call \"{0}\"", ShortBatchFileName ) );
+					EnvReaderBatchFileContent.Add( String.Format("call \"{0}\" {1}", ShortBatchFileName, Parameters) );
 
 					// Pipe all environment variables to a file where we can read them in.
 					// We use a separate executable which runs after the batch file because we want to capture
@@ -294,9 +295,16 @@ namespace UnrealBuildTool
 					StartInfo.Arguments = String.Format("/U /C \"{0}\"", EnvReaderBatchFileName);
 					StartInfo.CreateNoWindow = true;
 					StartInfo.UseShellExecute = false;
-					StartInfo.RedirectStandardOutput = true;
-					StartInfo.RedirectStandardError = true;
-					StartInfo.RedirectStandardInput = true;
+
+					// @todo UAP: Win10-beta issue, sometimes the process hangs if standard I/O is redirected, and there's bytes
+					// in the buffer and there is no reader for those streams.  Loop back with a later Win10 build
+					if (System.Environment.OSVersion.Platform != PlatformID.Win32NT ||
+						System.Environment.OSVersion.Version.Major < 10)
+					{
+						StartInfo.RedirectStandardOutput = true;
+						StartInfo.RedirectStandardError = true;
+						StartInfo.RedirectStandardInput = true;
+					}
 
 					Log.TraceVerbose( "Launching {0} to harvest Visual Studio environment settings...", StartInfo.FileName );
 
diff --git a/Engine/Source/Programs/UnrealBuildTool/System/VCProject.cs b/Engine/Source/Programs/UnrealBuildTool/System/VCProject.cs
--- a/Engine/Source/Programs/UnrealBuildTool/System/VCProject.cs
+++ b/Engine/Source/Programs/UnrealBuildTool/System/VCProject.cs
@@ -322,20 +322,24 @@ namespace UnrealBuildTool
 
 			// Build up the new include search path string
 			var VCIncludeSearchPaths = new StringBuilder();
-			{ 
-				foreach( var CurPath in IntelliSenseIncludeSearchPaths )
+			{
+				foreach (var CurPath in IntelliSenseIncludeSearchPaths)
+				{
+					VCIncludeSearchPaths.Append(CurPath + ";");
+				}
+				foreach (var CurPath in IntelliSenseSystemIncludeSearchPaths)
 				{
-					VCIncludeSearchPaths.Append( CurPath + ";" );
+					VCIncludeSearchPaths.Append(CurPath + ";");
 				}
-				foreach( var CurPath in IntelliSenseSystemIncludeSearchPaths )
+				if (InPlatforms.Contains(UnrealTargetPlatform.WinUAP))
 				{
-					VCIncludeSearchPaths.Append( CurPath + ";" );
+					VCIncludeSearchPaths.Append(WinUAPToolChain.GetVCIncludePaths(CPPTargetPlatform.WinUAP) + ";");
 				}
-				if(InPlatforms.Contains(UnrealTargetPlatform.Win64))
+				else if (InPlatforms.Contains(UnrealTargetPlatform.Win64))
 				{
 					VCIncludeSearchPaths.Append(VCToolChain.GetVCIncludePaths(CPPTargetPlatform.Win64) + ";");
 				}
-				else if(InPlatforms.Contains(UnrealTargetPlatform.Win32))
+				else if (InPlatforms.Contains(UnrealTargetPlatform.Win32))
 				{
 					VCIncludeSearchPaths.Append(VCToolChain.GetVCIncludePaths(CPPTargetPlatform.Win32) + ";");
 				}
@@ -361,11 +365,10 @@ namespace UnrealBuildTool
 			var FiltersFileIsNeeded = false;
 
 			// Project file header
-			var ToolsVersion = VCProjectFileGenerator.ProjectFileFormat == VCProjectFileGenerator.VCProjectFileFormat.VisualStudio2013 ? "12.0" : "4.0";
 			VCProjectFileContent.Append(
 				"<?xml version=\"1.0\" encoding=\"utf-8\"?>" + ProjectFileGenerator.NewLine +
 				ProjectFileGenerator.NewLine +
-				"<Project DefaultTargets=\"Build\" ToolsVersion=\"" + ToolsVersion + "\" xmlns=\"http://schemas.microsoft.com/developer/msbuild/2003\">" + ProjectFileGenerator.NewLine);
+				"<Project DefaultTargets=\"Build\" ToolsVersion=\"" + VCProjectFileGenerator.ProjectFileToolVersionString + "\" xmlns=\"http://schemas.microsoft.com/developer/msbuild/2003\">" + ProjectFileGenerator.NewLine);
 
 			bool bGenerateUserFileContent = UEPlatformProjectGenerator.PlatformRequiresVSUserFileGeneration(InPlatforms, InConfigurations);
 			if (bGenerateUserFileContent)
@@ -373,7 +376,7 @@ namespace UnrealBuildTool
 				VCUserFileContent.Append(
 					"<?xml version=\"1.0\" encoding=\"utf-8\"?>" + ProjectFileGenerator.NewLine + 
 					ProjectFileGenerator.NewLine +
-					"<Project ToolsVersion=\"" + ToolsVersion + "\" xmlns=\"http://schemas.microsoft.com/developer/msbuild/2003\">" + ProjectFileGenerator.NewLine
+					"<Project ToolsVersion=\"" + VCProjectFileGenerator.ProjectFileToolVersionString + "\" xmlns=\"http://schemas.microsoft.com/developer/msbuild/2003\">" + ProjectFileGenerator.NewLine
 					);
 			}
 
@@ -843,14 +846,7 @@ namespace UnrealBuildTool
 			string PlatformToolsetString = (ProjGenerator != null) ? ProjGenerator.GetVisualStudioPlatformToolsetString(Combination.Platform, Combination.Configuration, this) : "";
 			if( String.IsNullOrEmpty( PlatformToolsetString ) )
 			{
-				if( VCProjectFileGenerator.ProjectFileFormat == VCProjectFileGenerator.VCProjectFileFormat.VisualStudio2013 )
-				{
-					PlatformToolsetString = "		<PlatformToolset>v120</PlatformToolset>" + ProjectFileGenerator.NewLine;
-				}
-				else if( VCProjectFileGenerator.ProjectFileFormat == VCProjectFileGenerator.VCProjectFileFormat.VisualStudio2012 )
-				{
-					PlatformToolsetString = "		<PlatformToolset>v110</PlatformToolset>" + ProjectFileGenerator.NewLine;
-				}
+				PlatformToolsetString = "		<PlatformToolset>" + VCProjectFileGenerator.ProjectFilePlatformToolsetVersionString + "</PlatformToolset>" + ProjectFileGenerator.NewLine;
 			}
 
 			string PlatformConfigurationType = (ProjGenerator == null)? "Makefile" : ProjGenerator.GetVisualStudioPlatformConfigurationType(Combination.Platform);	
@@ -893,15 +889,15 @@ namespace UnrealBuildTool
 					string ProjectRelativeUnusedDirectory = NormalizeProjectPath(Path.Combine(ProjectFileGenerator.EngineRelativePath, BuildConfiguration.BaseIntermediateFolder, "Unused"));
 
 					VCProjectFileContent.Append(
-					    "	<PropertyGroup " + ConditionString + ">" + ProjectFileGenerator.NewLine +
+						"	<PropertyGroup " + ConditionString + ">" + ProjectFileGenerator.NewLine +
 						"		<OutDir>" + ProjectRelativeUnusedDirectory + Path.DirectorySeparatorChar + "</OutDir>" + ProjectFileGenerator.NewLine +
 						"		<IntDir>" + ProjectRelativeUnusedDirectory + Path.DirectorySeparatorChar + "</IntDir>" + ProjectFileGenerator.NewLine +
 						"		<NMakeBuildCommandLine>@rem Nothing to do.</NMakeBuildCommandLine>" + ProjectFileGenerator.NewLine +
 						"		<NMakeReBuildCommandLine>@rem Nothing to do.</NMakeReBuildCommandLine>" + ProjectFileGenerator.NewLine +
 						"		<NMakeCleanCommandLine>@rem Nothing to do.</NMakeCleanCommandLine>" + ProjectFileGenerator.NewLine +
 						"		<NMakeOutput/>" + ProjectFileGenerator.NewLine +
-                        "	</PropertyGroup>" + ProjectFileGenerator.NewLine);
-                }
+						"	</PropertyGroup>" + ProjectFileGenerator.NewLine);
+				}
 				else
 				{
 					TargetRules TargetRulesObject = Combination.ProjectTarget.TargetRules;
@@ -920,7 +916,7 @@ namespace UnrealBuildTool
 					// Get the output directory
 					string EngineRootDirectory = Path.GetFullPath(ProjectFileGenerator.EngineRelativePath);
 					string RootDirectory = EngineRootDirectory;
-                    if ((TargetRules.IsAGame(TargetRulesObject.Type) || TargetRulesObject.Type == TargetRules.TargetType.Server) && bShouldCompileMonolithic && !TargetRulesObject.bOutputToEngineBinaries)
+					if ((TargetRules.IsAGame(TargetRulesObject.Type) || TargetRulesObject.Type == TargetRules.TargetType.Server) && bShouldCompileMonolithic && !TargetRulesObject.bOutputToEngineBinaries)
 					{
 						if (UnrealBuildTool.HasUProjectFile() && Utils.IsFileUnderDirectory(TargetFilePath, UnrealBuildTool.GetUProjectPath()))
 						{
@@ -970,8 +966,8 @@ namespace UnrealBuildTool
 					NMakePath += BuildPlatform.GetBinaryExtension(UEBuildBinaryType.Executable);
 					NMakePath = (BuildPlatform as UEBuildPlatform).ModifyNMakeOutput(NMakePath);
 
-                    VCProjectFileContent.Append(
-                        "	<PropertyGroup " + ConditionString + ">" + ProjectFileGenerator.NewLine);
+					VCProjectFileContent.Append(
+						"	<PropertyGroup " + ConditionString + ">" + ProjectFileGenerator.NewLine);
 
 					string PathStrings = (ProjGenerator != null) ? ProjGenerator.GetVisualStudioPathsEntries(Platform, Configuration, TargetRulesObject.Type, TargetFilePath, ProjectFilePath, NMakePath) : "";
 					if (string.IsNullOrEmpty(PathStrings) || (PathStrings.Contains("<IntDir>") == false))
@@ -1009,6 +1005,22 @@ namespace UnrealBuildTool
 					}
 					string BatchFilesDirectoryName = Path.Combine(ProjectFileGenerator.EngineRelativePath, "Build", "BatchFiles");
 
+					// @todo UAP: For the MS toolchains, if an override was set for project generation, push that into the build strings to override the build toolchain as well
+					string BuildToolOverride = "";
+					if (UnrealBuildTool.CommandLineContains("-2012"))
+					{
+						BuildToolOverride = " -2012";
+					}
+					if (UnrealBuildTool.CommandLineContains("-2013"))
+					{
+						BuildToolOverride = " -2013";
+					}
+					if (UnrealBuildTool.CommandLineContains("-2015"))
+					{
+						BuildToolOverride = " -2015";
+					}
+					BuildArguments += BuildToolOverride;
+
 					// NMake Build command line
 					VCProjectFileContent.Append("		<NMakeBuildCommandLine>");
 					VCProjectFileContent.Append(EscapePath(NormalizeProjectPath(Path.Combine(BatchFilesDirectoryName, "Build.bat"))) + BuildArguments.ToString());
@@ -1027,17 +1039,17 @@ namespace UnrealBuildTool
 					VCProjectFileContent.Append("		<NMakeOutput>");
 					VCProjectFileContent.Append(NormalizeProjectPath(NMakePath));
 					VCProjectFileContent.Append("</NMakeOutput>" + ProjectFileGenerator.NewLine);
-				    VCProjectFileContent.Append("	</PropertyGroup>" + ProjectFileGenerator.NewLine);
+					VCProjectFileContent.Append("	</PropertyGroup>" + ProjectFileGenerator.NewLine);
 
-                    string LayoutDirString = (ProjGenerator != null) ? ProjGenerator.GetVisualStudioLayoutDirSection(Platform, Configuration, ConditionString, Combination.ProjectTarget.TargetRules.Type, Combination.ProjectTarget.TargetFilePath, ProjectFilePath, NMakePath) : "";
-                    VCProjectFileContent.Append(LayoutDirString);
-                }
+					string LayoutDirString = (ProjGenerator != null) ? ProjGenerator.GetVisualStudioLayoutDirSection(Platform, Configuration, ConditionString, Combination.ProjectTarget.TargetRules.Type, Combination.ProjectTarget.TargetFilePath, ProjectFilePath, NMakePath) : "";
+					VCProjectFileContent.Append(LayoutDirString);
+				}
 
 				if (VCUserFileContent != null && Combination.ProjectTarget != null)
 				{
 					TargetRules TargetRulesObject = Combination.ProjectTarget.TargetRules;
 
-					if ((Platform == UnrealTargetPlatform.Win32) || (Platform == UnrealTargetPlatform.Win64))
+					if ((Platform == UnrealTargetPlatform.Win32) || (Platform == UnrealTargetPlatform.Win64) || (Platform == UnrealTargetPlatform.WinUAP))
 					{
 						VCUserFileContent.Append(
 							"	<PropertyGroup " + ConditionString + ">" + ProjectFileGenerator.NewLine);
@@ -1165,7 +1177,7 @@ namespace UnrealBuildTool
 
 			ProjectFileContent.Append(
 				"<Import Project=\"$(MSBuildExtensionsPath)\\$(MSBuildToolsVersion)\\Microsoft.Common.props\" Condition=\"Exists('$(MSBuildExtensionsPath)\\$(MSBuildToolsVersion)\\Microsoft.Common.props')\" />" +
- 					ProjectFileGenerator.NewLine);
+					ProjectFileGenerator.NewLine);
 
 			// Support single configuration only (for now).
 			ProjectFileContent.Append(
@@ -1176,7 +1188,7 @@ namespace UnrealBuildTool
 					"\t<DefineConstants>TRACE</DefineConstants>" + ProjectFileGenerator.NewLine +
 					"\t<ErrorReport>prompt</ErrorReport>" + ProjectFileGenerator.NewLine +
 					"\t<WarningLevel>4</WarningLevel>" + ProjectFileGenerator.NewLine +
-                    "\t<TreatWarningsAsErrors>true</TreatWarningsAsErrors>" + ProjectFileGenerator.NewLine +
+					"\t<TreatWarningsAsErrors>true</TreatWarningsAsErrors>" + ProjectFileGenerator.NewLine +
 				"</PropertyGroup>" + ProjectFileGenerator.NewLine);
 
 			ProjectFileContent.Append(
diff --git a/Engine/Source/Programs/UnrealBuildTool/System/VCProjectFileGenerator.cs b/Engine/Source/Programs/UnrealBuildTool/System/VCProjectFileGenerator.cs
--- a/Engine/Source/Programs/UnrealBuildTool/System/VCProjectFileGenerator.cs
+++ b/Engine/Source/Programs/UnrealBuildTool/System/VCProjectFileGenerator.cs
@@ -73,6 +73,11 @@ namespace UnrealBuildTool
 			{
 				File.Delete(MasterProjDeleteFilename);
 			}
+			MasterProjDeleteFilename = MasterProjectFile + ".v12.suo";
+			if (File.Exists(MasterProjDeleteFilename))
+			{
+				File.Delete(MasterProjDeleteFilename);
+			}
 
 			// Delete the project files folder
 			if (Directory.Exists(InIntermediateProjectFilesPath))
@@ -110,13 +115,65 @@ namespace UnrealBuildTool
 		public enum VCProjectFileFormat
 		{
 			VisualStudio2012,
-			VisualStudio2013
+			VisualStudio2013,
+			VisualStudio2015,
+		}
+
+		/// "4.0", "12.0", or "14.0", etc...
+		static public string ProjectFileToolVersionString
+		{
+			get
+			{
+				switch (ProjectFileFormat)
+				{
+					case VCProjectFileFormat.VisualStudio2012:
+						return "4.0";
+					case VCProjectFileFormat.VisualStudio2013:
+						return "12.0";
+					case VCProjectFileFormat.VisualStudio2015:
+						return "14.0";
+				}
+				return string.Empty;
+			}
+		}
+
+		/// for instance: <PlatformToolset>v110</PlatformToolset>
+		static public string ProjectFilePlatformToolsetVersionString
+		{
+			get
+			{
+				switch (ProjectFileFormat)
+				{
+					case VCProjectFileFormat.VisualStudio2012:
+						return "v110";
+					case VCProjectFileFormat.VisualStudio2013:
+						return "v120";
+					case VCProjectFileFormat.VisualStudio2015:
+						return "v140";
+				}
+				return string.Empty;
+			}
 		}
 
 		/// Which version of Visual Studio we should generate project files for
-		static public VCProjectFileFormat ProjectFileFormat = 
-			( WindowsPlatform.Compiler == WindowsCompiler.VisualStudio2012 ) ? VCProjectFileFormat.VisualStudio2012 : VCProjectFileFormat.VisualStudio2013;
+		static public VCProjectFileFormat ProjectFileFormat;
 
+		/// Just set up the defaults, later this will be overridden by command line or other options
+		static VCProjectFileGenerator()
+		{
+			switch (WindowsPlatform.Compiler)
+			{
+				case WindowsCompiler.VisualStudio2012:
+					ProjectFileFormat = VCProjectFileFormat.VisualStudio2012;
+					break;
+				case WindowsCompiler.VisualStudio2013:
+					ProjectFileFormat = VCProjectFileFormat.VisualStudio2013;
+					break;
+				case WindowsCompiler.VisualStudio2015:
+					ProjectFileFormat = VCProjectFileFormat.VisualStudio2015;
+					break;
+			}
+		}
 
 		/// This is the platform name that Visual Studio is always guaranteed to support.  We'll use this as
 		/// a platform for any project configurations where our actual platform is not supported by the
@@ -148,6 +205,10 @@ namespace UnrealBuildTool
 					case WindowsCompiler.VisualStudio2013:
 						ProjectFileFormat = VCProjectFileFormat.VisualStudio2013;
 						break;
+
+					case WindowsCompiler.VisualStudio2015:
+						ProjectFileFormat = VCProjectFileFormat.VisualStudio2015;
+						break;
 				}
 			}
 
@@ -162,6 +223,10 @@ namespace UnrealBuildTool
 					case "-2013":
 						ProjectFileFormat = VCProjectFileFormat.VisualStudio2013;
 						break;
+
+					case "-2015":
+						ProjectFileFormat = VCProjectFileFormat.VisualStudio2015;
+						break;
 				}
 			}
 		}
@@ -266,18 +331,22 @@ namespace UnrealBuildTool
 			const string VersionTag = "# UnrealEngineGeneratedSolutionVersion=1.0";
 
 			// Solution file header
-			if( ProjectFileFormat == VCProjectFileFormat.VisualStudio2013 )
+			if (ProjectFileFormat == VCProjectFileFormat.VisualStudio2015)
+			{
+				VCSolutionFileContent.Append(
+					ProjectFileGenerator.NewLine +
+					"Microsoft Visual Studio Solution File, Format Version 12.00" + ProjectFileGenerator.NewLine +
+					"# Visual Studio 14" + ProjectFileGenerator.NewLine +
+					"VisualStudioVersion = 14.0.22310.1" + ProjectFileGenerator.NewLine +
+					"MinimumVisualStudioVersion = 10.0.40219.1" + ProjectFileGenerator.NewLine);
+			}
+			else if ( ProjectFileFormat == VCProjectFileFormat.VisualStudio2013 )
 			{
 				VCSolutionFileContent.Append(
 					ProjectFileGenerator.NewLine +
 					"Microsoft Visual Studio Solution File, Format Version 12.00" + ProjectFileGenerator.NewLine +
 					"# Visual Studio 2013" + ProjectFileGenerator.NewLine +
 					VersionTag + ProjectFileGenerator.NewLine);
-				
-				/* This is not required by VS 2013 to load the projects
-				VCSolutionFileContent.Append(
-					"VisualStudioVersion = 12.0.20617.1 PREVIEW" + ProjectFileGenerator.NewLine +
-					"MinimumVisualStudioVersion = 10.0.40219.1" + ProjectFileGenerator.NewLine );*/
 			}
 			else if( ProjectFileFormat == VCProjectFileFormat.VisualStudio2012 )
 			{
@@ -741,8 +810,20 @@ namespace UnrealBuildTool
 			// Save a solution config file which selects the development editor configuration by default.
 			if (bSuccess)
 			{
-				// Figure out the filename for the SUO file. VS2013 will automatically import the VS2012 options if necessary.
-				string SolutionOptionsExtension = (ProjectFileFormat == VCProjectFileFormat.VisualStudio2012)? "v11.suo" : "v12.suo";
+				// Figure out the filename for the SUO file. VS will automatically import the options from earlier versions if necessary.
+				string SolutionOptionsExtension = "vUnknown.suo";
+				switch (ProjectFileFormat)
+				{
+					case VCProjectFileFormat.VisualStudio2012:
+						SolutionOptionsExtension = "v11.suo";
+						break;
+					case VCProjectFileFormat.VisualStudio2013:
+						SolutionOptionsExtension = "v12.suo";
+						break;
+					case VCProjectFileFormat.VisualStudio2015:
+						SolutionOptionsExtension = "v14.suo";
+						break;
+				}
 
 				// Check it doesn't exist before overwriting it. Since these files store the user's preferences, it'd be bad form to overwrite them.
 				string SolutionOptionsFileName = Path.Combine(MasterProjectRelativePath, Path.ChangeExtension(SolutionFileName, SolutionOptionsExtension));
diff --git a/Engine/Source/Programs/UnrealBuildTool/UnrealBuildTool.csproj b/Engine/Source/Programs/UnrealBuildTool/UnrealBuildTool.csproj
--- a/Engine/Source/Programs/UnrealBuildTool/UnrealBuildTool.csproj
+++ b/Engine/Source/Programs/UnrealBuildTool/UnrealBuildTool.csproj
@@ -117,15 +117,6 @@
       </ItemGroup>
     </When>
   </Choose>
-  <Choose>
-    <When Condition="Exists('Cygwin\UEBuildCygwin.cs')">
-      <ItemGroup>
-        <Compile Include="Cygwin\CygwinToolChain.cs" />
-        <Compile Include="Cygwin\UEBuildCygwin.cs" />
-        <Compile Include="Cygwin\CygwinProjectGenerator.cs" />
-      </ItemGroup>
-    </When>
-  </Choose>
   <ItemGroup>
     <Compile Include="Mac\MacToolChain.cs" />
     <Compile Include="Mac\UEBuildMac.cs" />
@@ -138,16 +129,18 @@
     <Compile Include="Android\AndroidToolChain.cs" />
     <Compile Include="Android\AndroidProjectGenerator.cs" />
   </ItemGroup>
-  <Choose>
-    <When Condition="Exists('WinRT\UEBuildWinRT.cs')">
-      <ItemGroup>
-        <Compile Include="WinRT\UEBuildWinRT.cs" />
-        <Compile Include="WinRT\WinRTToolChain.cs" />
-        <Compile Include="WinRT\WinRTDeploy.cs" />
-        <Compile Include="WinRT\WinRTProjectGenerator.cs" />
-      </ItemGroup>
-    </When>
-  </Choose>
+  <ItemGroup>
+    <Compile Include="WinRT\UEBuildWinRT.cs" />
+    <Compile Include="WinRT\WinRTToolChain.cs" />
+    <Compile Include="WinRT\WinRTDeploy.cs" />
+    <Compile Include="WinRT\WinRTProjectGenerator.cs" />
+  </ItemGroup>
+  <ItemGroup>
+    <Compile Include="WinUAP\WinUAPProjectGenerator.cs" />
+    <Compile Include="WinUAP\WinUAPToolChain.cs" />
+    <Compile Include="WinUAP\UEBuildWinUAP.cs" />
+    <Compile Include="WinUAP\WinUAPDeploy.cs" />
+  </ItemGroup>
   <ItemGroup>
     <Compile Include="IOS\IOSProjectGenerator.cs" />
     <Compile Include="IOS\IOSToolChain.cs" />
@@ -172,13 +165,6 @@
       </ItemGroup>
     </When>
   </Choose>
-  <Choose>
-    <When Condition="Exists('ToolChain\NoRedist\EpicMachineNames.cs')">
-      <ItemGroup>
-        <Compile Include="ToolChain\NoRedist\EpicMachineNames.cs" />
-      </ItemGroup>
-    </When>
-  </Choose>
   <ItemGroup>
     <Compile Include="..\DotNETCommon\MetaData.cs">
       <Link>Properties\MetaData.cs</Link>
@@ -313,4 +299,4 @@
     <PreBuildEvent>
     </PreBuildEvent>
   </PropertyGroup>
-</Project>
\ No newline at end of file
+</Project>
diff --git a/Engine/Source/Programs/UnrealBuildTool/WinRT/WinRTToolChain.cs b/Engine/Source/Programs/UnrealBuildTool/WinRT/WinRTToolChain.cs
--- a/Engine/Source/Programs/UnrealBuildTool/WinRT/WinRTToolChain.cs
+++ b/Engine/Source/Programs/UnrealBuildTool/WinRT/WinRTToolChain.cs
@@ -66,7 +66,16 @@ namespace UnrealBuildTool
 			// Disable "The file contains a character that cannot be represented in the current code page" warning for non-US windows.
 			Result += " /wd4819";
 
-			if( BuildConfiguration.bUseSharedPCHs )
+			// @todo UAP: UE4 is non-compliant when it comes to use of %s and %S
+			// Previously %s meant "the current character set" and %S meant "the other one".
+			// Now %s means multibyte and %S means wide. %Ts means "natural width".
+			// Reverting this behaviour until the UE4 source catches up.
+			if (WindowsPlatform.Compiler == WindowsCompiler.VisualStudio2015)
+			{
+				Result += " /D_CRT_STDIO_LEGACY_WIDE_SPECIFIERS=1";
+			}
+
+			if ( BuildConfiguration.bUseSharedPCHs )
 			{
 				// @todo SharedPCH: Disable warning about PCH defines not matching .cpp defines.  We "cheat" these defines a little
 				// bit to make shared PCHs work.  But it's totally safe.  Trust us.
diff --git a/Engine/Source/Programs/UnrealBuildTool/WinUAP/UEBuildWinUAP.cs b/Engine/Source/Programs/UnrealBuildTool/WinUAP/UEBuildWinUAP.cs
new file mode 100644
--- /dev/null
+++ b/Engine/Source/Programs/UnrealBuildTool/WinUAP/UEBuildWinUAP.cs
@@ -0,0 +1,496 @@
+﻿// Copyright 1998-2014 Epic Games, Inc. All Rights Reserved.
+
+using System;
+using System.Collections.Generic;
+using System.Text;
+using System.Diagnostics;
+using System.IO;
+using Microsoft.Win32;
+
+namespace UnrealBuildTool
+{
+	public class WinUAPPlatform : UEBuildPlatform
+	{
+		/// Property caching.
+		private static WindowsCompiler? CachedCompiler;
+
+		/// Version of the compiler toolchain to use for Windows Universal Applications (WinUAPs)
+		public static WindowsCompiler Compiler
+		{
+			get
+			{
+				// Cache the result because Compiler is often used.
+				if (CachedCompiler.HasValue)
+				{
+					return CachedCompiler.Value;
+				}
+
+				// First, default based on whether there is a command line override.
+				// Allows build chain to partially progress even in the absence of installed tools
+				if (UnrealBuildTool.CommandLineContains("-2015"))
+				{
+					CachedCompiler = WindowsCompiler.VisualStudio2015;
+				}
+				// Second, default based on what's installed
+				else if (!String.IsNullOrEmpty(WindowsPlatform.GetVSComnToolsPath(WindowsCompiler.VisualStudio2015)))
+				{
+					CachedCompiler = WindowsCompiler.VisualStudio2015;
+				}
+				else
+				{
+					CachedCompiler = null;
+				}
+
+				return CachedCompiler.Value;
+			}
+		}
+
+		// Enables the UAP platform and project file support in Unreal Build Tool
+		// @todo UAP: Remove this variable when UAP support is fully implemented
+		public static readonly bool bEnableUAPSupport = UnrealBuildTool.CommandLineContains("-uap");
+
+		/// True if we should only build against the app-local CRT and /APPCONTAINER linker flag
+		public static readonly bool bBuildForStore = true;
+
+		/// True if we should only build against the app-local CRT and set the API-set flags
+		public static readonly bool bWinApiFamilyApp = true;
+
+		/** True if VS EnvDTE is available (false when building using Visual Studio Express) */
+		public static bool bHasVisualStudioDTE
+		{
+			get
+			{
+				return WindowsPlatform.bHasVisualStudioDTE;
+			}
+		}
+
+		public override bool RequiresDeployPrepAfterCompile()
+		{
+			return true;
+		}
+
+		protected override SDKStatus HasRequiredManualSDKInternal()
+		{
+			try
+			{
+				// @todo UAP: wire this up once the rest of cleanup has been resolved. right now, just leaving this to always return valid, like VC does.
+				//WinUAPToolChain.FindBaseVSToolPath();
+				return SDKStatus.Valid;
+			}
+			catch (BuildException)
+			{
+				return SDKStatus.Invalid;
+			}
+		}
+
+		/**
+		 * Returns VisualStudio common tools path for current compiler.
+		 * 
+		 * @return Common tools path.
+		 */
+		public static string GetVSComnToolsPath()
+		{
+			return GetVSComnToolsPath(Compiler);
+		}
+
+		/**
+		 * Returns VisualStudio common tools path for given compiler.
+		 * 
+		 * @param Compiler Compiler for which to return tools path.
+		 * 
+		 * @return Common tools path.
+		 */
+		public static string GetVSComnToolsPath(WindowsCompiler Compiler)
+		{
+			int VSVersion;
+
+			switch(Compiler)
+			{
+				case WindowsCompiler.VisualStudio2015:
+					VSVersion = 14;
+					break;
+				default:
+					throw new NotSupportedException("Not supported compiler.");
+			}
+
+			string[] PossibleRegPaths = new string[] {
+				@"Wow6432Node\Microsoft\VisualStudio",	// Non-express VS on 64-bit machine.
+				@"Microsoft\VisualStudio",				// Non-express VS on 32-bit machine.
+				@"Wow6432Node\Microsoft\WDExpress",		// Express VS on 64-bit machine.
+				@"Microsoft\WDExpress"					// Express VS on 32-bit machine.
+			};
+
+			string VSPath = null;
+
+			foreach(var PossibleRegPath in PossibleRegPaths)
+			{
+				VSPath = (string) Registry.GetValue(string.Format(@"HKEY_LOCAL_MACHINE\SOFTWARE\{0}\{1}.0", PossibleRegPath, VSVersion), "InstallDir", null);
+
+				if(VSPath != null)
+				{
+					break;
+				}
+			}
+
+			if(VSPath == null)
+			{
+				return null;
+			}
+
+			return new DirectoryInfo(Path.Combine(VSPath, "..", "Tools")).FullName;
+		}
+
+
+		/**
+		 *	Register the platform with the UEBuildPlatform class
+		 */
+		protected override void RegisterBuildPlatformInternal()
+		{
+			// Register this build platform for WinUAP
+			Log.TraceVerbose("        Registering for {0}", UnrealTargetPlatform.WinUAP.ToString());
+			UEBuildPlatform.RegisterBuildPlatform(UnrealTargetPlatform.WinUAP, this);
+			UEBuildPlatform.RegisterPlatformWithGroup(UnrealTargetPlatform.WinUAP, UnrealPlatformGroup.Microsoft);
+		}
+
+		/**
+		 *	Retrieve the CPPTargetPlatform for the given UnrealTargetPlatform
+		 *
+		 *	@param	InUnrealTargetPlatform		The UnrealTargetPlatform being build
+		 *	
+		 *	@return	CPPTargetPlatform			The CPPTargetPlatform to compile for
+		 */
+		public override CPPTargetPlatform GetCPPTargetPlatform(UnrealTargetPlatform InUnrealTargetPlatform)
+		{
+			switch (InUnrealTargetPlatform)
+			{
+				case UnrealTargetPlatform.WinUAP:
+					return CPPTargetPlatform.WinUAP;
+			}
+			throw new BuildException("WinUAPPlatform::GetCPPTargetPlatform: Invalid request for {0}", InUnrealTargetPlatform.ToString());
+		}
+
+		/**
+		 *	Get the extension to use for the given binary type
+		 *	
+		 *	@param	InBinaryType		The binrary type being built
+		 *	
+		 *	@return	string				The binary extenstion (ie 'exe' or 'dll')
+		 */
+		public override string GetBinaryExtension(UEBuildBinaryType InBinaryType)
+		{
+			switch (InBinaryType)
+			{
+				case UEBuildBinaryType.DynamicLinkLibrary: 
+					return ".dll";
+				case UEBuildBinaryType.Executable:
+					return ".exe";
+				case UEBuildBinaryType.StaticLibrary:
+					return ".lib";
+				case UEBuildBinaryType.Object:
+					return ".obj";
+				case UEBuildBinaryType.PrecompiledHeader:
+					return ".pch";
+			}
+			return base.GetBinaryExtension(InBinaryType);
+		}
+
+		
+		/// <summary>
+		/// When using a Visual Studio compiler, returns the version name as a string
+		/// </summary>
+		/// <returns>The Visual Studio compiler version name (e.g. "2012")</returns>
+		public static string GetVisualStudioCompilerVersionName()
+		{
+			switch( Compiler )
+			{ 
+				case WindowsCompiler.VisualStudio2015:
+					return "2015";
+				default:
+					throw new BuildException("Unexpected WindowsCompiler version for GetVisualStudioCompilerVersionName().  Either not using a Visual Studio compiler or switch block needs to be updated");
+			}
+		}
+
+
+		/**
+		 *	Get the extension to use for debug info for the given binary type
+		 *	
+		 *	@param	InBinaryType		The binary type being built
+		 *	
+		 *	@return	string				The debug info extension (i.e. 'pdb')
+		 */
+		public override string GetDebugInfoExtension( UEBuildBinaryType InBinaryType )
+		{
+			switch (InBinaryType)
+			{
+				case UEBuildBinaryType.DynamicLinkLibrary:
+				case UEBuildBinaryType.Executable:
+					return ".pdb";
+			}
+			return "";
+		}
+
+		/**
+		 *	Whether incremental linking should be used
+		 *	
+		 *	@param	InPlatform			The CPPTargetPlatform being built
+		 *	@param	InConfiguration		The CPPTargetConfiguration being built
+		 *	
+		 *	@return	bool	true if incremental linking should be used, false if not
+		 */
+		public override bool ShouldUseIncrementalLinking(CPPTargetPlatform Platform, CPPTargetConfiguration Configuration)
+		{
+			return (Configuration == CPPTargetConfiguration.Debug);
+		}
+
+		/**
+		 *	Whether PDB files should be used
+		 *	
+		 *	@param	InPlatform			The CPPTargetPlatform being built
+		 *	@param	InConfiguration		The CPPTargetConfiguration being built
+		 *	@param	bInCreateDebugInfo	true if debug info is getting create, false if not
+		 *	
+		 *	@return	bool	true if PDB files should be used, false if not
+		 */
+		public override bool ShouldUsePDBFiles(CPPTargetPlatform Platform, CPPTargetConfiguration Configuration, bool bCreateDebugInfo)
+		{
+			// Only supported on PC.
+			if (bCreateDebugInfo && ShouldUseIncrementalLinking(Platform, Configuration))
+			{
+				return true;
+			}
+			return false;
+		}
+
+		/**
+		 *	Whether the editor should be built for this platform or not
+		 *	
+		 *	@param	InPlatform		The UnrealTargetPlatform being built
+		 *	@param	InConfiguration	The UnrealTargetConfiguration being built
+		 *	@return	bool			true if the editor should be built, false if not
+		 */
+		public override bool ShouldNotBuildEditor(UnrealTargetPlatform InPlatform, UnrealTargetConfiguration InConfiguration)
+		{
+			return true;
+		}
+
+		public override bool BuildRequiresCookedData(UnrealTargetPlatform InPlatform, UnrealTargetConfiguration InConfiguration)
+		{
+			return false;
+		}
+
+		public override void ResetBuildConfiguration(UnrealTargetPlatform InPlatform, UnrealTargetConfiguration InConfiguration)
+		{
+			UEBuildConfiguration.bCompileICU = true;
+		}
+
+
+		public override void ModifyNewlyLoadedModule(UEBuildModule InModule, TargetInfo Target)
+		{
+			if ((Target.Platform == UnrealTargetPlatform.WinUAP))
+			{
+				if (InModule.ToString() == "Core")
+				{
+					InModule.AddPrivateDependencyModule("WinUAPSDK");
+				}
+				else if (InModule.ToString() == "Engine")
+				{
+					InModule.AddPrivateDependencyModule("zlib");
+					InModule.AddPrivateDependencyModule("UElibPNG");
+					InModule.AddPublicDependencyModule("UEOgg");
+					InModule.AddPublicDependencyModule("Vorbis");
+				}
+				else if (InModule.ToString() == "D3D11RHI")
+				{
+					InModule.AddPublicDefinition("D3D11_WITH_DWMAPI=0");
+					InModule.AddPublicDefinition("WITH_DX_PERF=0");
+				}
+				else if (InModule.ToString() == "DX11")
+				{
+					// Clear out all the Windows include paths and libraries...
+					// The WinUAPSDK module handles proper paths and libs for UAP.
+					// However, the D3D11RHI module will include the DX11 module.
+					InModule.ClearPublicIncludePaths();
+					InModule.ClearPublicLibraryPaths();
+					InModule.ClearPublicAdditionalLibraries();
+					InModule.RemovePublicDefinition("WITH_D3DX_LIBS=1");
+					InModule.AddPublicDefinition("WITH_D3DX_LIBS=0");
+					InModule.RemovePublicAdditionalLibrary("X3DAudio.lib");
+					InModule.RemovePublicAdditionalLibrary("XAPOFX.lib");
+				}
+				else if (InModule.ToString() == "XAudio2")
+				{
+					InModule.AddPublicDefinition("XAUDIO_SUPPORTS_XMA2WAVEFORMATEX=0");
+					InModule.AddPublicDefinition("XAUDIO_SUPPORTS_DEVICE_DETAILS=0");
+					InModule.AddPublicDefinition("XAUDIO2_SUPPORTS_MUSIC=0");
+					InModule.AddPublicDefinition("XAUDIO2_SUPPORTS_SENDLIST=0");
+					InModule.AddPublicAdditionalLibrary("XAudio2.lib");
+				}
+				else if (InModule.ToString() == "DX11Audio")
+				{
+					InModule.RemovePublicAdditionalLibrary("X3DAudio.lib");
+					InModule.RemovePublicAdditionalLibrary("XAPOFX.lib");
+				}
+			}
+		}
+
+		/**
+		 *	Setup the target environment for building
+		 *	
+		 *	@param	InBuildTarget		The target being built
+		 */
+		public override void SetUpEnvironment(UEBuildTarget InBuildTarget)
+		{
+
+			InBuildTarget.GlobalCompileEnvironment.Config.Definitions.Add("PLATFORM_DESKTOP=0");
+			InBuildTarget.GlobalCompileEnvironment.Config.Definitions.Add("PLATFORM_64BITS=1");
+			InBuildTarget.GlobalCompileEnvironment.Config.Definitions.Add("EXCEPTIONS_DISABLED=0");
+
+			InBuildTarget.GlobalCompileEnvironment.Config.Definitions.Add("_WIN32_WINNT=0x0A00");
+			InBuildTarget.GlobalCompileEnvironment.Config.Definitions.Add("WINVER=0x0A00");
+
+			InBuildTarget.GlobalCompileEnvironment.Config.Definitions.Add("PLATFORM_WINUAP=1");
+			InBuildTarget.GlobalCompileEnvironment.Config.Definitions.Add("WINUAP=1");
+
+			if (bWinApiFamilyApp)
+			{
+				InBuildTarget.GlobalCompileEnvironment.Config.Definitions.Add("WINAPI_FAMILY=WINAPI_FAMILY_APP");
+			}
+
+			// @todo UAP: UE4 is non-compliant when it comes to use of %s and %S
+			// Previously %s meant "the current character set" and %S meant "the other one".
+			// Now %s means multibyte and %S means wide. %Ts means "natural width".
+			// Reverting this behavior until the UE4 source catches up.
+			InBuildTarget.GlobalCompileEnvironment.Config.Definitions.Add("_CRT_STDIO_LEGACY_WIDE_SPECIFIERS=1");
+
+			// No D3DX on WinUAP!
+			InBuildTarget.GlobalCompileEnvironment.Config.Definitions.Add("NO_D3DX_LIBS=1");
+
+
+			if (InBuildTarget.Rules != null)
+			{
+				// Explicitly exclude the MS C++ runtime libraries we're not using, to ensure other libraries we link with use the same
+				// runtime library as the engine.
+				if (InBuildTarget.Configuration == UnrealTargetConfiguration.Debug && BuildConfiguration.bDebugBuildsActuallyUseDebugCRT)
+				{
+					InBuildTarget.GlobalLinkEnvironment.Config.ExcludedLibraries.Add("MSVCRT");
+					InBuildTarget.GlobalLinkEnvironment.Config.ExcludedLibraries.Add("MSVCPRT");
+				}
+				else
+				{
+					InBuildTarget.GlobalLinkEnvironment.Config.ExcludedLibraries.Add("MSVCRTD");
+					InBuildTarget.GlobalLinkEnvironment.Config.ExcludedLibraries.Add("MSVCPRTD");
+				}
+				InBuildTarget.GlobalLinkEnvironment.Config.ExcludedLibraries.Add("LIBC");
+				InBuildTarget.GlobalLinkEnvironment.Config.ExcludedLibraries.Add("LIBCMT");
+				InBuildTarget.GlobalLinkEnvironment.Config.ExcludedLibraries.Add("LIBCPMT");
+				InBuildTarget.GlobalLinkEnvironment.Config.ExcludedLibraries.Add("LIBCP");
+				InBuildTarget.GlobalLinkEnvironment.Config.ExcludedLibraries.Add("LIBCD");
+				InBuildTarget.GlobalLinkEnvironment.Config.ExcludedLibraries.Add("LIBCMTD");
+				InBuildTarget.GlobalLinkEnvironment.Config.ExcludedLibraries.Add("LIBCPMTD");
+				InBuildTarget.GlobalLinkEnvironment.Config.ExcludedLibraries.Add("LIBCPD");
+			}
+
+			InBuildTarget.GlobalLinkEnvironment.Config.AdditionalLibraries.Add("mincore.lib");
+			InBuildTarget.GlobalLinkEnvironment.Config.AdditionalLibraries.Add("dloadhelper.lib");
+		   
+			// Disable Simplygon support if compiling against the NULL RHI.
+			if (InBuildTarget.GlobalCompileEnvironment.Config.Definitions.Contains("USE_NULL_RHI=1"))
+			{
+				UEBuildConfiguration.bCompileSimplygon = false;
+			}
+		}
+
+		/**
+		 *	Setup the configuration environment for building
+		 *	
+		 *	@param	InBuildTarget		The target being built
+		 */
+		public override void SetUpConfigurationEnvironment(UEBuildTarget InBuildTarget)
+		{
+			// Determine the C++ compile/link configuration based on the Unreal configuration.
+			CPPTargetConfiguration CompileConfiguration;
+			UnrealTargetConfiguration CheckConfig = InBuildTarget.Configuration;
+			switch (CheckConfig)
+			{
+				default:
+				case UnrealTargetConfiguration.Debug:
+					CompileConfiguration = CPPTargetConfiguration.Debug;
+					if( BuildConfiguration.bDebugBuildsActuallyUseDebugCRT )
+					{ 
+						InBuildTarget.GlobalCompileEnvironment.Config.Definitions.Add("_DEBUG=1"); // the engine doesn't use this, but lots of 3rd party stuff does
+					}
+					else
+					{
+						InBuildTarget.GlobalCompileEnvironment.Config.Definitions.Add("NDEBUG=1"); // the engine doesn't use this, but lots of 3rd party stuff does
+					}
+					InBuildTarget.GlobalCompileEnvironment.Config.Definitions.Add("UE_BUILD_DEBUG=1");
+					break;
+				case UnrealTargetConfiguration.DebugGame:
+					// Default to Development; can be overriden by individual modules.
+				case UnrealTargetConfiguration.Development:
+					CompileConfiguration = CPPTargetConfiguration.Development;
+					InBuildTarget.GlobalCompileEnvironment.Config.Definitions.Add("NDEBUG=1"); // the engine doesn't use this, but lots of 3rd party stuff does
+					InBuildTarget.GlobalCompileEnvironment.Config.Definitions.Add("UE_BUILD_DEVELOPMENT=1");
+					break;
+				case UnrealTargetConfiguration.Shipping:
+					CompileConfiguration = CPPTargetConfiguration.Shipping;
+					InBuildTarget.GlobalCompileEnvironment.Config.Definitions.Add("NDEBUG=1"); // the engine doesn't use this, but lots of 3rd party stuff does
+					InBuildTarget.GlobalCompileEnvironment.Config.Definitions.Add("UE_BUILD_SHIPPING=1");
+					break;
+				case UnrealTargetConfiguration.Test:
+					CompileConfiguration = CPPTargetConfiguration.Shipping;
+					InBuildTarget.GlobalCompileEnvironment.Config.Definitions.Add("NDEBUG=1"); // the engine doesn't use this, but lots of 3rd party stuff does
+					InBuildTarget.GlobalCompileEnvironment.Config.Definitions.Add("UE_BUILD_TEST=1");
+					break;
+			}
+
+			// Set up the global C++ compilation and link environment.
+			InBuildTarget.GlobalCompileEnvironment.Config.Target.Configuration = CompileConfiguration;
+			InBuildTarget.GlobalLinkEnvironment.Config.Target.Configuration    = CompileConfiguration;
+
+			// Create debug info based on the heuristics specified by the user.
+			InBuildTarget.GlobalCompileEnvironment.Config.bCreateDebugInfo =
+				!BuildConfiguration.bDisableDebugInfo && ShouldCreateDebugInfo(InBuildTarget.Platform, CheckConfig);
+
+			// NOTE: Even when debug info is turned off, we currently force the linker to generate debug info
+			//       anyway on Visual C++ platforms.  This will cause a PDB file to be generated with symbols
+			//       for most of the classes and function/method names, so that crashes still yield somewhat
+			//       useful call stacks, even though compiler-generate debug info may be disabled.  This gives
+			//       us much of the build-time savings of fully-disabled debug info, without giving up call
+			//       data completely.
+			InBuildTarget.GlobalLinkEnvironment.Config.bCreateDebugInfo = true;
+		}
+
+		/**
+		 *	Whether this platform should create debug information or not
+		 *	
+		 *	@param	InPlatform			The UnrealTargetPlatform being built
+		 *	@param	InConfiguration		The UnrealTargetConfiguration being built
+		 *	
+		 *	@return	bool				true if debug info should be generated, false if not
+		 */
+		public override bool ShouldCreateDebugInfo(UnrealTargetPlatform Platform, UnrealTargetConfiguration Configuration)
+		{
+			switch (Configuration)
+			{
+				case UnrealTargetConfiguration.Development: 
+				case UnrealTargetConfiguration.Shipping: 
+				case UnrealTargetConfiguration.Test: 
+					return !BuildConfiguration.bOmitPCDebugInfoInDevelopment;
+				case UnrealTargetConfiguration.DebugGame:
+				case UnrealTargetConfiguration.Debug:
+				default: 
+					return true;
+			};
+		}
+
+		/**
+		 *	Return whether this platform has uniquely named binaries across multiple games
+		 */
+		public override bool HasUniqueBinaries()
+		{
+			// Windows applications have many shared binaries between games
+			return false;
+		}
+	}
+}
diff --git a/Engine/Source/Programs/UnrealBuildTool/WinUAP/WinUAPDeploy.cs b/Engine/Source/Programs/UnrealBuildTool/WinUAP/WinUAPDeploy.cs
new file mode 100644
--- /dev/null
+++ b/Engine/Source/Programs/UnrealBuildTool/WinUAP/WinUAPDeploy.cs
@@ -0,0 +1,272 @@
+﻿// Copyright 1998-2014 Epic Games, Inc. All Rights Reserved.
+
+using System;
+using System.Collections.Generic;
+using System.Text;
+using System.IO;
+using System.Diagnostics;
+
+// @todo UAP: this file is a work in progress and is not yet complete for the F5 scenario for WinUAP
+
+namespace UnrealBuildTool
+{
+	/// <summary>
+	///  Base class to handle deploy of a target for a given platform
+	/// </summary>
+	class WinUAPDeploy : UEBuildDeploy
+	{
+		/**
+		 *	Register the platform with the UEBuildDeploy class
+		 */
+		public override void RegisterBuildDeploy()
+		{
+			// Register this deployment handle for WinUAP
+			Log.TraceVerbose("        Registering for {0}", UnrealTargetPlatform.WinUAP.ToString());
+			UEBuildDeploy.RegisterBuildDeploy(UnrealTargetPlatform.WinUAP, this);
+		}
+
+		/**
+		 *  Utility function to delete a file
+		 */
+		void DeployHelper_DeleteFile(string InFileToDelete)
+		{
+			Log.TraceInformation("WinUAPDeploy.DeployHelper_DeleteFile({0})", InFileToDelete);
+			if (File.Exists(InFileToDelete) == true)
+			{
+				FileAttributes attributes = File.GetAttributes(InFileToDelete);
+				if ((attributes & FileAttributes.ReadOnly) == FileAttributes.ReadOnly)
+				{
+					attributes &= ~FileAttributes.ReadOnly;
+					File.SetAttributes(InFileToDelete, attributes);
+				}
+				File.Delete(InFileToDelete);
+			}
+		}
+
+		/**
+		 *	Copy the contents of the given source directory to the given destination directory
+		 *	
+		 */
+		bool CopySourceToDestDir(string InSourceDirectory, string InDestinationDirectory, string InWildCard,
+			bool bInIncludeSubDirectories, bool bInRemoveDestinationOrphans)
+		{
+			Log.TraceInformation("WinUAPDeploy.CopySourceToDestDir({0}, {1}, {2},...)", InSourceDirectory, InDestinationDirectory, InWildCard);
+			if (Directory.Exists(InSourceDirectory) == false)
+			{
+				Log.TraceInformation("Warning: CopySourceToDestDir - SourceDirectory does not exist: {0}", InSourceDirectory);
+				return false;
+			}
+
+			// Make sure the destination directory exists!
+			Directory.CreateDirectory(InDestinationDirectory);
+
+			SearchOption OptionToSearch = bInIncludeSubDirectories ? SearchOption.AllDirectories : SearchOption.TopDirectoryOnly;
+
+			var SourceDirs = new List<string>(Directory.GetDirectories(InSourceDirectory, "*.*", OptionToSearch));
+			foreach (string SourceDir in SourceDirs)
+			{
+				string SubDir = SourceDir.Replace(InSourceDirectory, "");
+				string DestDir = InDestinationDirectory + SubDir;
+				Directory.CreateDirectory(DestDir);
+			}
+
+			var SourceFiles = new List<string>(Directory.GetFiles(InSourceDirectory, InWildCard, OptionToSearch));
+			var DestFiles = new List<string>(Directory.GetFiles(InDestinationDirectory, InWildCard, OptionToSearch));
+
+			// Keep a list of the files in the source directory... without the source path
+			List<string> FilesInSource = new List<string>();
+
+			// Copy all the source files that are newer...
+			foreach (string SourceFile in SourceFiles)
+			{
+				string Filename = SourceFile.Replace(InSourceDirectory, "");
+				FilesInSource.Add(Filename.ToUpperInvariant());
+				string DestFile = InDestinationDirectory + Filename;
+
+				System.DateTime SourceTime = File.GetLastWriteTime(SourceFile);
+				System.DateTime DestTime = File.GetLastWriteTime(DestFile);
+
+				if (SourceTime > DestTime)
+				{
+					try
+					{
+						DeployHelper_DeleteFile(DestFile);
+						File.Copy(SourceFile, DestFile, true);
+					}
+					catch (Exception exceptionMessage)
+					{
+						Log.TraceInformation("Failed to copy {0} to deployment: {1}", SourceFile, exceptionMessage);
+					}
+				}
+			}
+
+			if (bInRemoveDestinationOrphans == true)
+			{
+				// If requested, delete any destination files that do not have a corresponding
+				// file in the source directory
+				foreach (string DestFile in DestFiles)
+				{
+					string DestFilename = DestFile.Replace(InDestinationDirectory, "");
+					if (FilesInSource.Contains(DestFilename.ToUpperInvariant()) == false)
+					{
+						Log.TraceInformation("Destination file does not exist in Source - DELETING: {0}", DestFile);
+						FileAttributes attributes = File.GetAttributes(DestFile);
+						try
+						{
+							DeployHelper_DeleteFile(DestFile);
+						}
+						catch (Exception exceptionMessage)
+						{
+							Log.TraceInformation("Failed to delete {0} from deployment: {1}", DestFile, exceptionMessage);
+						}
+					}
+				}
+			}
+
+			return true;
+		}
+
+
+		/**
+		 *	Helper function for copying files
+		 */
+		void CopyFile(string InSource, string InDest, bool bForce)
+		{
+			Log.TraceInformation("WinUAPDeploy.CopyFile({0}, {1}, {2},...)", InSource, InDest, bForce);
+			if (File.Exists(InSource) == true)
+			{
+				if (bForce == true)
+				{
+					if (File.Exists(InDest) == true)
+					{
+						DeployHelper_DeleteFile(InDest);
+					}
+				}
+				File.Copy(InSource, InDest, true);
+			}
+			else
+			{
+				Log.TraceInformation("WinUAPDeploy: File didn't exist - {0}", InSource);
+			}
+		}
+
+		/**
+		 *	Helper function for copying a tree files
+		 */
+		//void CopyDirectory(string InSource, string InDest, bool bForce, bool bRecurse)
+		//{
+		//    if (Directory.Exists(InSource))
+		//    {
+		//        if (!Directory.Exists(InDest))
+		//        {
+		//            Directory.CreateDirectory(InDest);
+		//        }
+
+		//        // Copy all files
+		//        string[] FilesInDir = Directory.GetFiles(InSource);
+		//        foreach (string FileSourcePath in FilesInDir)
+		//        {
+		//            string FileDestPath = Path.Combine(InDest, Path.GetFileName(FileSourcePath));
+		//            if (bForce == true)
+		//            {
+		//                if (File.Exists(FileDestPath) == true)
+		//                {
+		//                    DeployHelper_DeleteFile(FileDestPath);
+		//                }
+		//            }
+		//            File.Copy(FileSourcePath, FileDestPath, true);
+		//        }
+
+		//        // Recurse sub directories
+		//        string[] DirsInDir = Directory.GetDirectories(InSource);
+		//        foreach (string DirSourcePath in DirsInDir)
+		//        {
+		//            string DirName = Path.GetFileName(DirSourcePath);
+		//            string DirDestPath = Path.Combine(InDest, DirName);
+		//            CopyDirectory(DirSourcePath, DirDestPath, bForce, bRecurse);
+		//        }
+		//    }
+		//}
+
+
+		public bool PrepForUATPackageOrDeploy(string InProjectName, string InProjectDirectory, string InExecutablePath, string InEngineDir, bool bForDistribution, string CookFlavor)
+		{
+			// Much of the WinUAPProjectGenerator code assumes that the path will be set to Engine/Source
+			string PreviousDirectory = Directory.GetCurrentDirectory();
+			string PreviousRelativeEnginePath = BuildConfiguration.RelativeEnginePath;
+
+			Directory.SetCurrentDirectory(Path.Combine(InEngineDir, "Binaries"));
+			BuildConfiguration.RelativeEnginePath = "../../Engine";
+
+			string ProjectBinaryFolder = new FileInfo(InExecutablePath).DirectoryName;
+			string EngineSourceRelativeBinaryPath = Utils.MakePathRelativeTo(ProjectBinaryFolder, Path.Combine(InEngineDir, "Source"));
+			bool bIsUE4Game = new FileInfo(InExecutablePath).Name.StartsWith("UE4Game", StringComparison.InvariantCultureIgnoreCase);
+
+			WinUAPProjectGenerator ProjectGen = new WinUAPProjectGenerator();
+			ProjectGen.GenerateDeployFiles(InProjectName, ProjectBinaryFolder, bIsUE4Game, EngineSourceRelativeBinaryPath);
+
+			// back to root for UAT
+			Directory.SetCurrentDirectory(PreviousDirectory);
+			BuildConfiguration.RelativeEnginePath = PreviousRelativeEnginePath;
+
+			return true;
+		}
+
+
+
+		public override bool PrepTargetForDeployment(UEBuildTarget InTarget)
+		{
+			string InAppName = InTarget.AppName;
+			Log.TraceInformation("Prepping {0} for deployment to {1}", InAppName, InTarget.Platform.ToString());
+			System.DateTime PrepDeployStartTime = DateTime.UtcNow;
+
+			string TargetFilename = RulesCompiler.GetTargetFilename(InAppName);
+			string ProjectSourceFolder = new FileInfo(TargetFilename).DirectoryName + "/";
+
+			string RelativeTargetDirectory;
+			string EngineSourceRelativeBinaryPath;
+			WinUAPProjectGenerator.GetTargetWinUAPPaths(InAppName, InTarget.Rules, out EngineSourceRelativeBinaryPath, out RelativeTargetDirectory);
+
+			PrepForUATPackageOrDeploy(InAppName, InTarget.ProjectDirectory, InTarget.OutputPath, BuildConfiguration.RelativeEnginePath, false, "");
+
+			// TODO - richiem - restore this if we find that it's needed.
+			//Log.TraceInformation("...copying the CELL dll...");
+			//string CELLPath = "../../../Engine/Source/" + UEBuildConfiguration.UEThirdPartySourceDirectory + "CELL/lib/win64/";
+			//string CELLPathRelease = CELLPath + "Release/";
+			//string CELLDllRelease = "CommonEventLoggingLibrary.dll";
+			//string CELLPathDebug = CELLPath + "Debug/";
+			//string CELLDllDebug = "CommonEventLoggingLibraryd.dll";
+
+			//CopyFile(EngineSourceRelativeBinaryPath + CELLPathRelease + CELLDllRelease, EngineSourceRelativeBinaryPath + CELLDllRelease, true);
+			//CopyFile(EngineSourceRelativeBinaryPath + CELLPathDebug + CELLDllDebug, EngineSourceRelativeBinaryPath + CELLDllDebug, true);
+
+			//string XSAPIPath            = EngineSourceRelativeBinaryPath + "../../../Engine/Source/ThirdParty/XSAPI/lib/";
+			//string XboxServicesConfig   = "xboxservices.config";
+			//string DesktopLogin         = "DesktopLogin.exe";
+
+			//Log.TraceInformation("...copying xboxservices.config");
+			//CopyFile(XSAPIPath + XboxServicesConfig, EngineSourceRelativeBinaryPath + XboxServicesConfig, true);
+			//Log.TraceInformation("...copying DesktopLogin.exe...");
+			//CopyFile(XSAPIPath + DesktopLogin, EngineSourceRelativeBinaryPath + DesktopLogin, true);
+
+			// TODO - richiem - restore this if we find that it's needed.
+			//if (InTarget.Configuration == UnrealTargetConfiguration.Development)
+			//{
+			//    Log.TraceInformation("...copying AutoLogin...");
+			//    string AutoLoginPath    = EngineSourceRelativeBinaryPath + "../../../Engine/Source/Tools/AutoLogin/";
+			//    string AutoLoginDLL     = "AutoLogin.dll";
+			//    string AutoLoginBat     = "AutoLogin.bat";
+			//    CopyFile(AutoLoginPath + "bin/Debug/" + AutoLoginDLL, EngineSourceRelativeBinaryPath + AutoLoginDLL, true);
+			//    CopyFile(AutoLoginPath + AutoLoginBat, EngineSourceRelativeBinaryPath + AutoLoginBat, true);
+			//}
+
+
+
+			// Log out the time taken to deploy...
+			double PrepDeployDuration = (DateTime.UtcNow - PrepDeployStartTime).TotalSeconds;
+			Log.TraceInformation("WinUAP deployment preparation took {0:0.00} seconds", PrepDeployDuration);
+
+			return true;
+		}
+	}
+}
diff --git a/Engine/Source/Programs/UnrealBuildTool/WinUAP/WinUAPProjectGenerator.cs b/Engine/Source/Programs/UnrealBuildTool/WinUAP/WinUAPProjectGenerator.cs
new file mode 100644
--- /dev/null
+++ b/Engine/Source/Programs/UnrealBuildTool/WinUAP/WinUAPProjectGenerator.cs
@@ -0,0 +1,407 @@
+// Copyright 1998-2014 Epic Games, Inc. All Rights Reserved.
+
+using System;
+using System.Collections.Generic;
+using System.Text;
+using System.Diagnostics;
+using System.IO;
+
+// @todo UAP: this file is a work in progress and is not yet complete for the F5 scenario for WinUAP
+
+namespace UnrealBuildTool
+{
+	/**
+	 *	Base class for platform-specific project generators 
+	 */
+	public class WinUAPProjectGenerator : UEPlatformProjectGenerator
+	{
+		/**
+		 *	Register the platform with the UEPlatformProjectGenerator class
+		 */
+		public override void RegisterPlatformProjectGenerator()
+		{
+			if( WinUAPPlatform.bEnableUAPSupport )
+			{ 
+				// Register this project generator for WinUAP
+				Log.TraceVerbose("        Registering for {0}", UnrealTargetPlatform.WinUAP.ToString());
+			
+				UEPlatformProjectGenerator.RegisterPlatformProjectGenerator(UnrealTargetPlatform.WinUAP, this);
+			}
+		}
+
+		public static void GetTargetWinUAPPaths(string InTargetName, TargetRules InTargetRules,
+			out string OutEngineSourceRelativeBinaryPath, out string OutRelativeTargetPath)
+		{
+			OutEngineSourceRelativeBinaryPath = "";
+			OutRelativeTargetPath = "";
+
+			string TargetFilename = RulesCompiler.GetTargetFilename(InTargetName);
+
+			string ProjectSourceFolder = new FileInfo(TargetFilename).DirectoryName;
+
+			string EnginePath = Path.Combine(ProjectFileGenerator.EngineRelativePath);
+			string EngineSourcePath = Path.Combine(EnginePath, "Source");
+			string RelativeTargetFilename = Utils.MakePathRelativeTo(TargetFilename, EngineSourcePath);
+
+			if ((RelativeTargetFilename.StartsWith("..") == false) && (RelativeTargetFilename.Contains(":") == false))
+			{
+				// This target must be UNDER Engine/Source...
+				RelativeTargetFilename = Path.Combine(EngineSourcePath, RelativeTargetFilename);
+			}
+			RelativeTargetFilename = RelativeTargetFilename.Replace("\\", "/");
+			EnginePath = EnginePath.Replace("\\", "/");
+
+			Int32 LastSourceIdx = RelativeTargetFilename.LastIndexOf("Source");
+			if (LastSourceIdx != -1)
+			{
+				RelativeTargetFilename = RelativeTargetFilename.Substring(0, LastSourceIdx);
+			}
+			else
+			{
+				RelativeTargetFilename = "";
+			}
+
+			OutRelativeTargetPath = RelativeTargetFilename;
+
+			if (InTargetRules.bOutputToEngineBinaries)
+			{
+				RelativeTargetFilename = EnginePath;
+			}
+			OutEngineSourceRelativeBinaryPath = Path.Combine(RelativeTargetFilename, "Binaries/WinUAP/");
+			OutEngineSourceRelativeBinaryPath = OutEngineSourceRelativeBinaryPath.Replace("\\", "/");
+		}
+
+		///
+		///	VisualStudio project generation functions
+		///	
+		/**
+		 *	Whether this build platform has native support for VisualStudio
+		 *	
+		 *	@param	InPlatform			The UnrealTargetPlatform being built
+		 *	@param	InConfiguration		The UnrealTargetConfiguration being built
+		 *	
+		 *	@return	bool				true if native VisualStudio support (or custom VSI) is available
+		 */
+		public override bool HasVisualStudioSupport(UnrealTargetPlatform InPlatform, UnrealTargetConfiguration InConfiguration)
+		{
+			return false;
+		}
+
+		/**
+		 * Get whether this platform deploys 
+		 * 
+		 * @return	bool		true if the 'Deploy' option should be enabled
+		 */
+		public override bool GetVisualStudioDeploymentEnabled(UnrealTargetPlatform InPlatform, UnrealTargetConfiguration InConfiguration)
+		{
+			return true;
+		}
+
+		private Dictionary<string, string> GameToBuildFolder = new Dictionary<string, string>();
+
+
+		protected void GeneratePackageAppXManifest(string InGameName, bool bIsUE4Game, string InManifestType, string InGameBinaries, string InGameIntermediates, string InOutputPath)
+		{
+			string BuildFolder = Path.Combine(InOutputPath, "../../Build/WinUAP");
+			string CheckPath;
+			if (GameToBuildFolder.TryGetValue(InGameName, out CheckPath) == false)
+			{
+				GameToBuildFolder.Add(InGameName, BuildFolder);
+			}
+
+			string OutputPackageAppXManifestPath = Path.Combine(BuildFolder, InManifestType + ".AppxManifest");
+
+
+			// first look for a per-game template
+			string PackageAppXManifestPath = OutputPackageAppXManifestPath + ".template";
+			if (!File.Exists(PackageAppXManifestPath))
+			{
+				// then fallback to shared template
+				PackageAppXManifestPath = "../../Engine/Build/WinUAP/Templates/" + InManifestType + ".appxmanifest.template";
+			}
+
+			GenerateFileFromTemplate(PackageAppXManifestPath, OutputPackageAppXManifestPath, bIsUE4Game, InGameName, "", InGameBinaries, InGameIntermediates);
+
+			// Generate the package name...
+			// todo - richiem - this is next to fix
+			string PackageName = "foobar"; //= XboxOneDeploy.XboxOnePackageNameHelper_Run(OutputPackageAppXManifestPath);
+			if (string.IsNullOrEmpty(PackageName) == false)
+			{
+				if (GameToBuildFolder.TryGetValue(InGameName, out CheckPath) == false)
+				{
+					GameToBuildFolder.Add(InGameName, PackageName);
+				}
+			}
+			else
+			{
+				throw new BuildException("Error: Failed to find package name for {0}", OutputPackageAppXManifestPath);
+			}
+		}
+
+		protected void GeneratePackagePriConfig(string InGameName, bool bIsUE4Game, string InGameBinaries, string InGameIntermediates, string InOutputPath)
+		{
+			string BuildFolder = Path.Combine(InOutputPath, "../../Build/WinUAP");
+
+			string OutputPackagePriConfigPath = Path.Combine(BuildFolder, "Package.PriConfig");
+
+			// first look for a per-game template
+			string PackagePriConfigPath = OutputPackagePriConfigPath + ".template";
+			if (!File.Exists(PackagePriConfigPath))
+			{
+				// then fallback to shared template
+				PackagePriConfigPath = "../../Engine/Build/WinUAP/Templates/Package.priconfig.template";
+			}
+
+			GenerateFileFromTemplate(PackagePriConfigPath, OutputPackagePriConfigPath, bIsUE4Game, InGameName, "", InGameBinaries, InGameIntermediates);
+		}
+
+		protected void GeneratePackageResources(string InGameName, bool bIsUE4Game, string InGameBinaries, string InGameIntermediates, string InOutputPath)
+		{
+			string BuildFolder = Path.Combine(InOutputPath, "../../Build/WinUAP");
+
+			string OutputPackageResourcesPath = Path.Combine(BuildFolder, "Resources", "Resources.resw");
+
+			// first look for a per-game template
+			string PackageResourcesPath = Path.Combine(BuildFolder, "Resources", "Package.resources.template");
+			if (!File.Exists(PackageResourcesPath))
+			{
+				// then fallback to shared template
+				PackageResourcesPath = "../../Engine/Build/WinUAP/Templates/Package.resources.template";
+			}
+
+			// Only generate the example language if the resources haven't been setup at all previously
+			if (!File.Exists(OutputPackageResourcesPath))
+			{
+				GenerateFileFromTemplate(PackageResourcesPath, Path.Combine(BuildFolder, "Resources", "en-US", "Resources.resw"), bIsUE4Game, InGameName, "", InGameBinaries, InGameIntermediates);
+			}
+			GenerateFileFromTemplate(PackageResourcesPath, OutputPackageResourcesPath, bIsUE4Game, InGameName, "", InGameBinaries, InGameIntermediates);
+		}
+
+		protected void GenerateFileFromTemplate(string InTemplateFile, string InOutputFile, bool bIsUE4Game, string InGameName, string InGameConfiguration, string InGameBinaries, string InGameIntermediates, string InGuidString = "", Dictionary<string, string> InExtraReplacements = null)
+		{
+			String ExecutableName = bIsUE4Game ? "UE4Game" : InGameName;
+
+			StringBuilder outputContents = new StringBuilder();
+			using (StreamReader reader = new StreamReader(InTemplateFile))
+			{
+				string LineStr = null;
+				while (reader.Peek() != -1)
+				{
+					LineStr = reader.ReadLine();
+					if (LineStr.Contains("%GAME%"))
+					{
+						if (LineStr.Contains("Executable") || LineStr.Contains("EntryPoint"))
+						{
+							LineStr = LineStr.Replace("%GAME%", ExecutableName);
+						}
+						else
+						{
+							LineStr = LineStr.Replace("%GAME%", InGameName);
+						}
+					}
+					if (LineStr.Contains("%PROJECT%"))
+					{
+						if (bIsUE4Game)
+						{
+							LineStr = LineStr.Replace("%PROJECT%", "Engine");
+						}
+						else
+						{
+							LineStr = LineStr.Replace("%PROJECT%", InGameName);
+						}
+					}
+
+					if (LineStr.Contains("%CONFIG%"))
+					{
+						LineStr = LineStr.Replace("%CONFIG%", InGameConfiguration);
+					}
+
+					if (LineStr.Contains("%DISPLAYCONFIG%"))
+					{
+						if (string.IsNullOrEmpty(InGameConfiguration))
+						{
+							LineStr = LineStr.Replace("%DISPLAYCONFIG%", "");
+						}
+						else
+						{
+							LineStr = LineStr.Replace("%DISPLAYCONFIG%", " (" + InGameConfiguration + ")");
+						}
+					}
+
+					if (LineStr.Contains("%GAMEBINARIES%"))
+					{
+						LineStr = LineStr.Replace("%GAMEBINARIES%", InGameBinaries);
+					}
+
+					if (LineStr.Contains("%GAMEBINARIESFOLDER%"))
+					{
+						LineStr = LineStr.Replace("%GAMEBINARIESFOLDER%", InGameBinaries);
+					}
+
+					if (LineStr.Contains("%GAMEINTERMEDIATESFOLDER%"))
+					{
+						LineStr = LineStr.Replace("%GAMEINTERMEDIATESFOLDER%", InGameIntermediates);
+					}
+
+					if (LineStr.Contains("%GUID%"))
+					{
+						LineStr = LineStr.Replace("%GUID%", InGuidString);
+					}
+
+					if (InExtraReplacements != null)
+					{
+						foreach (KeyValuePair<string, string> Replacement in InExtraReplacements)
+						{
+							if (LineStr.Contains(Replacement.Key))
+							{
+								LineStr = LineStr.Replace(Replacement.Key, Replacement.Value);
+							}
+						}
+					}
+
+					outputContents.AppendLine(LineStr);
+				}
+			}
+
+			if (outputContents.Length > 0)
+			{
+				bool bFileNeedsSave = true;
+
+				if (File.Exists(InOutputFile))
+				{
+					// Read in the original file completely
+					string LoadedFileContent = null;
+					var FileAlreadyExists = File.Exists(InOutputFile);
+					if (FileAlreadyExists)
+					{
+						try
+						{
+							LoadedFileContent = File.ReadAllText(InOutputFile);
+						}
+						catch (Exception)
+						{
+							Log.TraceInformation("Error while trying to load existing file {0}.  Ignored.", InOutputFile);
+						}
+					}
+
+					// Don't bother saving anything out if the new file content is the same as the old file's content
+					if (LoadedFileContent != null)
+					{
+						var bIgnoreProjectFileWhitespaces = true;
+						if (ProjectFileComparer.CompareOrdinalIgnoreCase(LoadedFileContent, outputContents.ToString(), bIgnoreProjectFileWhitespaces) == 0)
+						{
+							// Exact match!
+							bFileNeedsSave = false;
+						}
+
+						if (!bFileNeedsSave)
+						{
+							Log.TraceVerbose("Skipped saving {0} because contents haven't changed.", Path.GetFileName(InOutputFile));
+						}
+					}
+				}
+
+				if (bFileNeedsSave)
+				{
+					// Save the file
+					try
+					{
+						Directory.CreateDirectory(Path.GetDirectoryName(InOutputFile));
+						File.WriteAllText(InOutputFile, outputContents.ToString(), Encoding.UTF8);
+						Log.TraceVerbose("Saved {0}", Path.GetFileName(InOutputFile));
+					}
+					catch (Exception)
+					{
+						// Unable to write to the project file.
+						Log.TraceInformation("Error while trying to write file {0}.  The file is probably read-only.", InOutputFile);
+					}
+				}
+			}
+		}
+
+
+		protected void CopyDefaultImageFiles(string InOutputPath)
+		{
+			string OutputImageFolder = Path.Combine(InOutputPath, "../../Build/WinUAP/Resources");
+			if (Directory.Exists(OutputImageFolder) == false)
+			{
+				Directory.CreateDirectory(OutputImageFolder);
+			}
+
+			string[] ImageList = new string[] { "Logo.png", "SmallLogo.png", "SplashScreen.png", "StoreLogo.png", "WideLogo.png" };
+
+			foreach (string ImageName in ImageList)
+			{
+				string OutputImageFilePath = Path.Combine(OutputImageFolder, ImageName);
+				if (File.Exists(OutputImageFilePath) == false)
+				{
+					string SourceImagePath = "../../Engine/Build/WinUAP/DefaultImages/" + ImageName;
+					File.Copy(SourceImagePath, OutputImageFilePath);
+
+					string OutputEnglishImageFilePath = Path.Combine(OutputImageFolder, "en-US", ImageName);
+					if (Directory.Exists(OutputImageFolder + "/en-US") == false)
+					{
+						Directory.CreateDirectory(OutputImageFolder + "/en-US");
+					}
+					if (File.Exists(OutputEnglishImageFilePath) == false)
+					{
+						File.Copy(SourceImagePath, OutputEnglishImageFilePath);
+					}
+				}
+			}
+		}
+
+
+		public void GenerateDeployFiles(string InTargetName, string InTargetFilepath, bool bIsUE4Game, string EngineSourceRelativeBinaryPath)
+		{
+			const string SourceFolderString = "/Source/";
+			string OutTargetPath = InTargetFilepath;
+			OutTargetPath = OutTargetPath.Replace("\\", "/");
+			if (bIsUE4Game == false)
+			{
+				int LastSourceIndex = OutTargetPath.LastIndexOf(SourceFolderString);
+				if (LastSourceIndex != -1)
+				{
+					OutTargetPath = OutTargetPath.Substring(0, LastSourceIndex + SourceFolderString.Length);
+					OutTargetPath = Path.Combine(OutTargetPath, "WinUAP");
+				}
+				else
+				{
+					// @todo UAP: This is bad... throw a build exception?
+				}
+			}
+			else
+			{
+				OutTargetPath = Path.Combine((new FileInfo(InTargetFilepath)).DirectoryName, "WinUAP");
+			}
+
+			string GameBinariesFolder = Utils.MakePathRelativeTo(EngineSourceRelativeBinaryPath, OutTargetPath);
+			GameBinariesFolder = GameBinariesFolder.Replace("\\", "/");
+			if (GameBinariesFolder.EndsWith("/") == false)
+			{
+				GameBinariesFolder += "/";
+			}
+			string GameIntermediatesFolder = GameBinariesFolder.Replace("Binaries/WinAUP", "Intermediate/Build/Unused");
+
+			// Make sure the output directory exists...
+			if (Directory.Exists(OutTargetPath) == false)
+			{
+				Log.TraceVerbose("Generating output directory " + OutTargetPath);
+				Directory.CreateDirectory(OutTargetPath);
+			}
+
+			GeneratePackageAppXManifest(InTargetName, bIsUE4Game, "Package", GameBinariesFolder, GameIntermediatesFolder, OutTargetPath);
+			GeneratePackageAppXManifest(InTargetName, bIsUE4Game, "Deploy", GameBinariesFolder, GameIntermediatesFolder, OutTargetPath);
+			GeneratePackagePriConfig(InTargetName, bIsUE4Game, GameBinariesFolder, GameIntermediatesFolder, OutTargetPath);
+			GeneratePackageResources(InTargetName, bIsUE4Game, GameBinariesFolder, GameIntermediatesFolder, OutTargetPath);
+			CopyDefaultImageFiles(OutTargetPath);
+		}
+
+
+		public override void GenerateGameProperties(UnrealTargetConfiguration Configuration, StringBuilder VCProjectFileContent, TargetRules.TargetType TargetType, string RootDirectory, string TargetFilePath)
+		{
+			VCProjectFileContent.Append("		<NMakePreprocessorDefinitions>$(NMakePreprocessorDefinitions);PLATFORM_WINUAP=1;WINUAP=1;</NMakePreprocessorDefinitions>" + ProjectFileGenerator.NewLine);
+		}
+
+	}
+}
diff --git a/Engine/Source/Programs/UnrealBuildTool/WinUAP/WinUAPToolChain.cs b/Engine/Source/Programs/UnrealBuildTool/WinUAP/WinUAPToolChain.cs
new file mode 100644
--- /dev/null
+++ b/Engine/Source/Programs/UnrealBuildTool/WinUAP/WinUAPToolChain.cs
@@ -0,0 +1,1046 @@
+// Copyright 1998-2014 Epic Games, Inc. All Rights Reserved.
+
+using System;
+using System.Collections.Generic;
+using System.Diagnostics;
+using System.IO;
+using System.Text;
+using System.Text.RegularExpressions;
+
+using Microsoft.Win32;
+
+namespace UnrealBuildTool
+{
+	public class WinUAPToolChain : UEToolChain
+	{
+		public override void RegisterToolChain()
+		{
+			if( WinUAPPlatform.bEnableUAPSupport )
+			{ 
+				// Register this tool chain for WinUAP
+				Log.TraceVerbose("        Registered for {0}", CPPTargetPlatform.WinUAP.ToString());
+				UEToolChain.RegisterPlatformToolChain(CPPTargetPlatform.WinUAP, this);
+			}
+		}
+
+		static void AppendCLArguments_Global(CPPEnvironment CompileEnvironment, VCEnvironment EnvVars, StringBuilder Arguments)
+		{
+			//Arguments.Append(" /showIncludes");
+
+			if( BuildConfiguration.bEnableCodeAnalysis )
+			{
+				Arguments.Append(" /analyze");
+
+				// Don't cause analyze warnings to be errors
+				Arguments.Append(" /analyze:WX-");
+
+				// Report functions that use a LOT of stack space.  You can lower this value if you
+				// want more aggressive checking for functions that use a lot of stack memory.
+				Arguments.Append(" /analyze:stacksize81940");
+
+				// Don't bother generating code, only analyze code (may report fewer warnings though.)
+				//Arguments.Append(" /analyze:only");
+			}
+
+			// Prevents the compiler from displaying its logo for each invocation.
+			Arguments.Append(" /nologo");
+
+			// Enable intrinsic functions.
+			Arguments.Append(" /Oi");
+
+			// Pack struct members on 8-byte boundaries.
+			Arguments.Append(" /Zp8");
+
+			// Separate functions for linker.
+			Arguments.Append(" /Gy");
+
+			// Relaxes floating point precision semantics to allow more optimization.
+			Arguments.Append(" /fp:fast");
+
+			// Compile into an .obj file, and skip linking.
+			Arguments.Append(" /c");
+
+			// Allow 800% of the default memory allocation limit.
+			Arguments.Append(" /Zm800");
+
+			// Allow large object files to avoid hitting the 2^16 section limit when running with -StressTestUnity.
+			Arguments.Append(" /bigobj");
+
+			// Disable "The file contains a character that cannot be represented in the current code page" warning for non-US windows.
+			Arguments.Append(" /wd4819");
+
+			// @todo UAP: Disable "unreachable code" warning since auto-included vccorlib.h triggers it
+			Arguments.Append(" /wd4702");
+
+			// @todo UAP: Silence the hash_map deprecation errors for now. This should be replaced with unordered_map for the real fix.
+			if (WindowsPlatform.Compiler == WindowsCompiler.VisualStudio2015)
+			{
+				Arguments.Append(" /D_SILENCE_STDEXT_HASH_DEPRECATION_WARNINGS");
+			}
+
+			if ( BuildConfiguration.bUseSharedPCHs )
+			{
+				// @todo SharedPCH: Disable warning about PCH defines not matching .cpp defines.  We "cheat" these defines a little
+				// bit to make shared PCHs work.  But it's totally safe.  Trust us.
+				Arguments.Append(" /wd4651");
+
+				// @todo SharedPCH: Disable warning about redefining *API macros.  The PCH header is compiled with various DLLIMPORTs, but
+				// when a module that uses that PCH header *IS* one of those imports, that module is compiled with EXPORTS, so the macro
+				// is redefined on the command-line.  We need to clobber those defines to make shared PCHs work properly!
+				Arguments.Append(" /wd4005");
+			}
+
+			// If compiling as a DLL, set the relevant defines
+			if (CompileEnvironment.Config.bIsBuildingDLL)
+			{
+				Arguments.Append(" /D_WINDLL");
+			}
+
+			// Handle Common Language Runtime support (C++/CLI)
+			if (CompileEnvironment.Config.CLRMode == CPPCLRMode.CLREnabled)
+			{
+				Arguments.Append(" /clr");
+
+				// Don't use default lib path, we override it explicitly to use the 4.0 reference assemblies.
+				Arguments.Append(" /clr:nostdlib");
+			}
+
+			//
+			//	Debug
+			//
+			if (CompileEnvironment.Config.Target.Configuration == CPPTargetConfiguration.Debug)
+			{
+				// Disable compiler optimization.
+				Arguments.Append(" /Od");
+
+				// Favor code size (especially useful for embedded platforms).
+				Arguments.Append(" /Os");
+
+				// Allow inline method expansion unless E&C support is requested
+				if( !BuildConfiguration.bSupportEditAndContinue )
+				{
+					// @todo UAP: No inlining in Debug builds except in the editor where DLL exports/imports aren't handled properly in module _API macros.
+					if (UEBuildConfiguration.bBuildEditor)
+					{
+						Arguments.Append(" /Ob2");
+					}
+				}
+
+				// Runtime stack checks are not allowed when compiling for CLR
+				if (CompileEnvironment.Config.CLRMode == CPPCLRMode.CLRDisabled)
+				{
+					Arguments.Append(" /RTCs");
+				}
+			}
+			//
+			//	Development and LTCG
+			//
+			else
+			{
+				// Maximum optimizations if desired.
+				if (CompileEnvironment.Config.OptimizeCode >= ModuleRules.CodeOptimization.InNonDebugBuilds)
+				{
+					Arguments.Append(" /Ox");
+
+					// Allow optimized code to be debugged more easily.  This makes PDBs a bit larger, but doesn't noticeably affect
+					// compile times.  The executable code is not affected at all by this switch, only the debugging information.
+					Arguments.Append(" /Zo");
+				}
+
+				// Favor code speed.
+				Arguments.Append(" /Ot");
+
+				// Only omit frame pointers on the PC (which is implied by /Ox) if wanted.
+				if ( BuildConfiguration.bOmitFramePointers == false
+				&& (CompileEnvironment.Config.Target.Platform == CPPTargetPlatform.WinUAP))
+				{
+					Arguments.Append(" /Oy-");
+				}
+
+				// Allow inline method expansion
+				Arguments.Append(" /Ob2");
+
+				//
+				// LTCG
+				//
+				if (CompileEnvironment.Config.Target.Configuration == CPPTargetConfiguration.Shipping)
+				{
+					if( BuildConfiguration.bAllowLTCG )
+					{
+						// Enable link-time code generation.
+						Arguments.Append(" /GL");
+					}
+				}
+			}
+
+			//
+			//	PC
+			//
+
+			// Prompt the user before reporting internal errors to Microsoft.
+			Arguments.Append(" /errorReport:prompt");
+
+			// Enable C++ exception handling, but not C exceptions.
+			Arguments.Append(" /EHsc");
+
+			// If enabled, create debug information.
+			if (CompileEnvironment.Config.bCreateDebugInfo)
+			{
+				// Store debug info in .pdb files.
+				if (BuildConfiguration.bUsePDBFiles)
+				{
+					// Create debug info suitable for E&C if wanted.
+					if ( BuildConfiguration.bSupportEditAndContinue &&
+						// We only need to do this in debug as that's the only configuration that supports E&C.
+						CompileEnvironment.Config.Target.Configuration == CPPTargetConfiguration.Debug)
+					{
+						Arguments.Append(" /ZI");
+					}
+					// Regular PDB debug information.
+					else
+					{
+						Arguments.Append(" /Zi");
+					}
+					// We need to add this so VS won't lock the PDB file and prevent synchronous updates. This forces serialization through MSPDBSRV.exe.
+					// See http://msdn.microsoft.com/en-us/library/dn502518.aspx for deeper discussion of /FS switch.
+					if (BuildConfiguration.bUseIncrementalLinking && WindowsPlatform.Compiler >= WindowsCompiler.VisualStudio2013)
+					{
+						Arguments.Append(" /FS");
+					}
+				}
+				// Store C7-format debug info in the .obj files, which is faster.
+				else
+				{
+					Arguments.Append(" /Z7");
+				}
+			}
+
+			// Specify the appropriate runtime library based on the platform and config.
+			if (CompileEnvironment.Config.bUseStaticCRT)
+			{
+				if (CompileEnvironment.Config.Target.Configuration == CPPTargetConfiguration.Debug && BuildConfiguration.bDebugBuildsActuallyUseDebugCRT)
+				{
+					Arguments.Append(" /MTd");
+				}
+				else
+				{
+					Arguments.Append(" /MT");
+				}
+			}
+			else
+			{
+				if (CompileEnvironment.Config.Target.Configuration == CPPTargetConfiguration.Debug && BuildConfiguration.bDebugBuildsActuallyUseDebugCRT)
+				{
+					Arguments.Append(" /MDd");
+				}
+				else
+				{
+					Arguments.Append(" /MD");
+				}
+			}
+
+			if (WinUAPPlatform.bBuildForStore)
+			{
+				Arguments.Append(" /D_BUILD_FOR_STORE=1");
+			}
+
+			if (WindowsPlatform.Compiler == WindowsCompiler.VisualStudio2015)
+			{
+				// @todo UAP: These must be appended to the end of the system include list, lest they override some of the third party sources includes
+				if (Directory.Exists(EnvVars.WindowsSDKExtensionDir))
+				{
+					CompileEnvironment.Config.CPPIncludeInfo.SystemIncludePaths.Add(string.Format(@"{0}\Include\{1}\ucrt", EnvVars.WindowsSDKExtensionDir, EnvVars.WindowsSDKExtensionHeaderLibVersion));
+					CompileEnvironment.Config.CPPIncludeInfo.SystemIncludePaths.Add(string.Format(@"{0}\Include\{1}\um", EnvVars.WindowsSDKExtensionDir, EnvVars.WindowsSDKExtensionHeaderLibVersion));
+					CompileEnvironment.Config.CPPIncludeInfo.SystemIncludePaths.Add(string.Format(@"{0}\Include\{1}\shared", EnvVars.WindowsSDKExtensionDir, EnvVars.WindowsSDKExtensionHeaderLibVersion));
+					CompileEnvironment.Config.CPPIncludeInfo.SystemIncludePaths.Add(string.Format(@"{0}\Include\{1}\winrt", EnvVars.WindowsSDKExtensionDir, EnvVars.WindowsSDKExtensionHeaderLibVersion));
+				}
+			}
+
+			if (CompileEnvironment.Config.bIsBuildingLibrary == false) // will put in a config option, but for now...
+			{
+				// Enable Windows Runtime extensions
+				Arguments.Append(" /ZW");
+				Arguments.Append(" /DUSE_WINRT_MAIN=1");
+				// TODO - richiem - this will have to be updated when final layout SDKs are available
+				if (WindowsPlatform.Compiler == WindowsCompiler.VisualStudio2015 && 
+					Directory.Exists(Path.Combine(EnvVars.WindowsSDKExtensionDir, "UnionMetadata")))
+				{
+					Arguments.AppendFormat(@" /AI""{0}\UnionMetadata""", EnvVars.WindowsSDKExtensionDir);
+					Arguments.AppendFormat(@" /FU""{0}\UnionMetadata\windows.winmd""", EnvVars.WindowsSDKExtensionDir);
+				}
+				Arguments.AppendFormat(@" /AI""{0}\..\..\VC\vcpackages""", EnvVars.BaseVSToolPath);
+				Arguments.AppendFormat(@" /FU""{0}\..\..\VC\vcpackages\platform.winmd""", EnvVars.BaseVSToolPath);
+			}
+
+		}
+
+		static void AppendCLArguments_CPP( CPPEnvironment CompileEnvironment, StringBuilder Arguments)
+		{
+			// Explicitly compile the file as C++.
+			Arguments.Append(" /TP");
+
+			if (!CompileEnvironment.Config.bEnableBufferSecurityChecks)
+			{
+				// This will disable buffer security checks (which are enabled by default) that the MS compiler adds around arrays on the stack,
+				// Which can add some performance overhead, especially in performance intensive code
+				// Only disable this if you know what you are doing, because it will be disabled for the entire module!
+				Arguments.Append(" /GS-");
+			}
+
+			// C++/CLI requires that RTTI is left enabled
+			if (CompileEnvironment.Config.CLRMode == CPPCLRMode.CLRDisabled)
+			{
+				if (CompileEnvironment.Config.bUseRTTI)
+				{
+					// Enable C++ RTTI.
+					Arguments.Append(" /GR");
+				}
+				else
+				{
+					// Disable C++ RTTI.
+					Arguments.Append(" /GR-");
+				}
+			}
+
+			// Level 4 warnings.
+			Arguments.Append(" /W3");
+		}
+		
+		static void AppendCLArguments_C(StringBuilder Arguments)
+		{
+			// Explicitly compile the file as C.
+			Arguments.Append(" /TC");
+		
+			// Level 0 warnings.  Needed for external C projects that produce warnings at higher warning levels.
+			Arguments.Append(" /W0");
+		}
+
+		static void AppendLinkArguments(LinkEnvironment LinkEnvironment, StringBuilder Arguments)
+		{
+			// Don't create a side-by-side manifest file for the executable.
+			Arguments.Append(" /MANIFEST:NO");
+
+			// Prevents the linker from displaying its logo for each invocation.
+			Arguments.Append(" /NOLOGO");
+
+			if (LinkEnvironment.Config.bCreateDebugInfo)
+			{
+				// Output debug info for the linked executable.
+				Arguments.Append(" /DEBUG");
+			}
+
+			// Prompt the user before reporting internal errors to Microsoft.
+			Arguments.Append(" /errorReport:prompt");
+
+			//
+			//	PC
+			//
+			// Set machine type/ architecture to be 64 bit, and set as a store app
+			if (WinUAPPlatform.bBuildForStore)
+			{
+				Arguments.Append(" /MACHINE:x64");
+				Arguments.Append(" /APPCONTAINER");
+				// this helps with store API compliance validation tools, adding additional pdb info
+				Arguments.Append(" /PROFILE");
+			}
+
+			if (LinkEnvironment.Config.bIsBuildingConsoleApplication)
+			{
+				Arguments.Append(" /SUBSYSTEM:CONSOLE");
+			}
+			else
+			{
+				Arguments.Append(" /SUBSYSTEM:WINDOWS");
+			}
+
+			if (LinkEnvironment.Config.bIsBuildingConsoleApplication && !LinkEnvironment.Config.bIsBuildingDLL && !String.IsNullOrEmpty( LinkEnvironment.Config.WindowsEntryPointOverride ) )
+			{
+				// Use overridden entry point
+				Arguments.Append(" /ENTRY:" + LinkEnvironment.Config.WindowsEntryPointOverride);
+			}
+
+			// Allow the OS to load the EXE at different base addresses than its preferred base address.
+			Arguments.Append(" /FIXED:No");
+
+			// Explicitly declare that the executable is compatible with Data Execution Prevention.
+			Arguments.Append(" /NXCOMPAT");
+
+			// Set the default stack size.
+			Arguments.Append(" /STACK:5000000");
+
+			// Allow delay-loaded DLLs to be explicitly unloaded.
+			Arguments.Append(" /DELAY:UNLOAD");
+
+			if (LinkEnvironment.Config.bIsBuildingDLL)
+			{
+				Arguments.Append(" /DLL");
+			}
+
+			if (LinkEnvironment.Config.CLRMode == CPPCLRMode.CLREnabled)
+			{
+				// DLLs built with managed code aren't allowed to have entry points as they will try to initialize
+				// complex static variables.  Managed code isn't allowed to run during DLLMain, we can't allow
+				// these variables to be initialized here.
+				if (LinkEnvironment.Config.bIsBuildingDLL)
+				{
+					// NOTE: This appears to only be needed if we want to get call stacks for debugging exit crashes related to the above
+					//		Arguments.Append(" /NOENTRY /NODEFAULTLIB:nochkclr.obj");
+				}
+			}
+
+			// Don't embed the full PDB path; we want to be able to move binaries elsewhere. They will always be side by side.
+			Arguments.Append(" /PDBALTPATH:%_PDB%");
+
+			//
+			//	Shipping & LTCG
+			//
+			if( BuildConfiguration.bAllowLTCG &&
+				LinkEnvironment.Config.Target.Configuration == CPPTargetConfiguration.Shipping)
+			{
+				// Use link-time code generation.
+				Arguments.Append(" /LTCG");
+
+				// This is where we add in the PGO-Lite linkorder.txt if we are using PGO-Lite
+				//Arguments.Append(" /ORDER:@linkorder.txt");
+				//Arguments.Append(" /VERBOSE");
+			}
+
+			//
+			//	Shipping binary
+			//
+			if (LinkEnvironment.Config.Target.Configuration == CPPTargetConfiguration.Shipping)
+			{
+				// Generate an EXE checksum.
+				Arguments.Append(" /RELEASE");
+
+				// Eliminate unreferenced symbols.
+				Arguments.Append(" /OPT:REF");
+
+				// Remove redundant COMDATs.
+				Arguments.Append(" /OPT:ICF");
+			}
+			//
+			//	Regular development binary. 
+			//
+			else
+			{
+				// Keep symbols that are unreferenced.
+				Arguments.Append(" /OPT:NOREF");
+
+				// Disable identical COMDAT folding.
+				Arguments.Append(" /OPT:NOICF");
+			}
+
+			// Enable incremental linking if wanted.
+			// NOTE: Don't bother using incremental linking for C++/CLI projects, as that's not supported and the option
+			//		 will silently be ignored anyway
+			if (BuildConfiguration.bUseIncrementalLinking && LinkEnvironment.Config.CLRMode != CPPCLRMode.CLREnabled)
+			{
+				Arguments.Append(" /INCREMENTAL");
+			}
+			else
+			{
+				Arguments.Append(" /INCREMENTAL:NO");
+			}
+
+			// Disable 
+			//LINK : warning LNK4199: /DELAYLOAD:nvtt_64.dll ignored; no imports found from nvtt_64.dll
+			// type warning as we leverage the DelayLoad option to put third-party DLLs into a 
+			// non-standard location. This requires the module(s) that use said DLL to ensure that it 
+			// is loaded prior to using it.
+			Arguments.Append(" /ignore:4199");
+
+			// Suppress warnings about missing PDB files for statically linked libraries.  We often don't want to distribute
+			// PDB files for these libraries.
+			Arguments.Append(" /ignore:4099");		// warning LNK4099: PDB '<file>' was not found with '<file>'
+		}
+
+		static void AppendLibArguments(LinkEnvironment LinkEnvironment, StringBuilder Arguments)
+		{
+			// Prevents the linker from displaying its logo for each invocation.
+			Arguments.Append(" /NOLOGO");
+
+			// Prompt the user before reporting internal errors to Microsoft.
+			Arguments.Append(" /errorReport:prompt");
+
+			//
+			//	PC
+			//
+			// Set machine type/ architecture to be 64 bit.
+			Arguments.Append(" /MACHINE:x64");
+
+			if (LinkEnvironment.Config.bIsBuildingConsoleApplication)
+			{
+				Arguments.Append(" /SUBSYSTEM:CONSOLE");
+			}
+			else
+			{
+				Arguments.Append(" /SUBSYSTEM:WINDOWS");
+			}
+
+			//
+			//	Shipping & LTCG
+			//
+			if (LinkEnvironment.Config.Target.Configuration == CPPTargetConfiguration.Shipping)
+			{
+				// Use link-time code generation.
+				Arguments.Append(" /LTCG");
+			}
+		}
+
+		public override CPPOutput CompileCPPFiles(UEBuildTarget Target, CPPEnvironment CompileEnvironment, List<FileItem> SourceFiles, string ModuleName)
+		{
+			var EnvVars = VCEnvironment.SetEnvironment(CompileEnvironment.Config.Target.Platform);
+
+			StringBuilder Arguments = new StringBuilder();
+			AppendCLArguments_Global(CompileEnvironment, EnvVars, Arguments);
+
+			// Add include paths to the argument list.
+			foreach (string IncludePath in CompileEnvironment.Config.CPPIncludeInfo.IncludePaths)
+			{
+				string IncludePathRelative = Utils.CleanDirectorySeparators(Utils.MakePathRelativeTo(IncludePath, Path.Combine(ProjectFileGenerator.RootRelativePath, "Engine/Source")), '/');
+				Arguments.AppendFormat(" /I \"{0}\"", IncludePathRelative);
+			}
+			foreach (string IncludePath in CompileEnvironment.Config.CPPIncludeInfo.SystemIncludePaths)
+			{
+				string IncludePathRelative = Utils.CleanDirectorySeparators(Utils.MakePathRelativeTo(IncludePath, Path.Combine(ProjectFileGenerator.RootRelativePath, "Engine/Source")), '/');
+				Arguments.AppendFormat(" /I \"{0}\"", IncludePathRelative);
+			}
+
+
+			if (CompileEnvironment.Config.CLRMode == CPPCLRMode.CLREnabled)
+			{
+				// Add .NET framework assembly paths.  This is needed so that C++/CLI projects
+				// can reference assemblies with #using, without having to hard code a path in the
+				// .cpp file to the assembly's location.				
+				foreach (string AssemblyPath in CompileEnvironment.Config.SystemDotNetAssemblyPaths)
+				{
+					Arguments.AppendFormat(" /AI \"{0}\"", AssemblyPath);
+				}
+
+				// Add explicit .NET framework assembly references				
+				foreach (string AssemblyName in CompileEnvironment.Config.FrameworkAssemblyDependencies)
+				{
+					Arguments.AppendFormat( " /FU \"{0}\"", AssemblyName );
+				}
+
+				// Add private assembly references				
+				foreach( PrivateAssemblyInfo CurAssemblyInfo in CompileEnvironment.PrivateAssemblyDependencies )
+				{
+					Arguments.AppendFormat( " /FU \"{0}\"", CurAssemblyInfo.FileItem.AbsolutePath );
+				}
+			}
+
+
+			// Add preprocessor definitions to the argument list.
+			foreach (string Definition in CompileEnvironment.Config.Definitions)
+			{
+				// Escape all quotation marks so that they get properly passed with the command line.
+				var DefinitionArgument = Definition.Contains("\"") ? Definition.Replace("\"", "\\\"") : Definition;
+				Arguments.AppendFormat(" /D\"{0}\"", DefinitionArgument);
+			}
+
+			var BuildPlatform = UEBuildPlatform.GetBuildPlatformForCPPTargetPlatform(CompileEnvironment.Config.Target.Platform);
+
+			// Create a compile action for each source file.
+			CPPOutput Result = new CPPOutput();
+			foreach (FileItem SourceFile in SourceFiles)
+			{
+				Action CompileAction = new Action(ActionType.Compile);
+				CompileAction.CommandDescription = "Compile";
+
+				StringBuilder FileArguments = new StringBuilder();
+				bool bIsPlainCFile = Path.GetExtension(SourceFile.AbsolutePath).ToUpperInvariant() == ".C";
+
+				// Add the C++ source file and its included files to the prerequisite item list.
+				AddPrerequisiteSourceFile(Target, BuildPlatform, CompileEnvironment, SourceFile, CompileAction.PrerequisiteItems);
+
+				// If this is a CLR file then make sure our dependent assemblies are added as prerequisites
+				if (CompileEnvironment.Config.CLRMode == CPPCLRMode.CLREnabled)
+				{
+					foreach( PrivateAssemblyInfo CurPrivateAssemblyDependency in CompileEnvironment.PrivateAssemblyDependencies )
+					{
+						CompileAction.PrerequisiteItems.Add( CurPrivateAssemblyDependency.FileItem );
+					}
+				}
+
+				bool bEmitsObjectFile = true;
+				if (CompileEnvironment.Config.PrecompiledHeaderAction == PrecompiledHeaderAction.Create)
+				{
+					// Generate a CPP File that just includes the precompiled header.
+					string PCHCPPFilename = "PCH." + ModuleName + "." + Path.GetFileName(CompileEnvironment.Config.PrecompiledHeaderIncludeFilename) + ".cpp";
+					string PCHCPPPath = Path.Combine(CompileEnvironment.Config.OutputDirectory, PCHCPPFilename);
+					FileItem PCHCPPFile = FileItem.CreateIntermediateTextFile(
+						PCHCPPPath,
+						string.Format("#include \"{0}\"\r\n", CompileEnvironment.Config.PrecompiledHeaderIncludeFilename)
+						);
+
+					// Make sure the original source directory the PCH header file existed in is added as an include
+					// path -- it might be a private PCH header and we need to make sure that its found!
+					string OriginalPCHHeaderDirectory = Path.GetDirectoryName( SourceFile.AbsolutePath );
+					FileArguments.AppendFormat(" /I \"{0}\"", OriginalPCHHeaderDirectory);
+
+					var PrecompiledFileExtension = UEBuildPlatform.BuildPlatformDictionary[UnrealTargetPlatform.WinUAP].GetBinaryExtension(UEBuildBinaryType.PrecompiledHeader);
+					// Add the precompiled header file to the produced items list.
+					FileItem PrecompiledHeaderFile = FileItem.GetItemByPath(
+						Path.Combine(
+							CompileEnvironment.Config.OutputDirectory,
+						Path.GetFileName(SourceFile.AbsolutePath) + PrecompiledFileExtension
+							)
+						);
+					CompileAction.ProducedItems.Add(PrecompiledHeaderFile);
+					Result.PrecompiledHeaderFile = PrecompiledHeaderFile;
+
+					// Add the parameters needed to compile the precompiled header file to the command-line.
+					FileArguments.AppendFormat(" /Yc\"{0}\"", CompileEnvironment.Config.PrecompiledHeaderIncludeFilename);
+					FileArguments.AppendFormat(" /Fp\"{0}\"", PrecompiledHeaderFile.AbsolutePath);
+
+					// If we're creating a PCH that will be used to compile source files for a library, we need
+					// the compiled modules to retain a reference to PCH's module, so that debugging information
+					// will be included in the library.  This is also required to avoid linker warning "LNK4206"
+					// when linking an application that uses this library.
+					if (CompileEnvironment.Config.bIsBuildingLibrary)
+					{
+						// NOTE: The symbol name we use here is arbitrary, and all that matters is that it is
+						// unique per PCH module used in our library
+						string FakeUniquePCHSymbolName = Path.GetFileNameWithoutExtension(CompileEnvironment.Config.PrecompiledHeaderIncludeFilename);
+						FileArguments.AppendFormat(" /Yl{0}", FakeUniquePCHSymbolName);
+					}
+
+					FileArguments.AppendFormat(" \"{0}\"", PCHCPPFile.AbsolutePath);
+
+					CompileAction.StatusDescription = PCHCPPFilename;
+				}
+				else
+				{
+					if (CompileEnvironment.Config.PrecompiledHeaderAction == PrecompiledHeaderAction.Include)
+					{
+						CompileAction.bIsUsingPCH = true;
+						CompileAction.PrerequisiteItems.Add(CompileEnvironment.PrecompiledHeaderFile);
+
+						FileArguments.AppendFormat(" /Yu\"{0}\"", CompileEnvironment.Config.PCHHeaderNameInCode);
+						FileArguments.AppendFormat(" /Fp\"{0}\"", CompileEnvironment.PrecompiledHeaderFile.AbsolutePath);
+
+						// Is it unsafe to always force inclusion?  Clang is doing it, and .generated.cpp files
+						// won't work otherwise, because they're not located in the context of the module,
+						// so they can't access the module's PCH without an absolute path.
+						//if( CompileEnvironment.Config.bForceIncludePrecompiledHeader )
+						{
+							// Force include the precompiled header file.  This is needed because we may have selected a
+							// precompiled header that is different than the first direct include in the C++ source file, but
+							// we still need to make sure that our precompiled header is the first thing included!
+							FileArguments.AppendFormat( " /FI\"{0}\"", CompileEnvironment.Config.PCHHeaderNameInCode);
+						}
+					}
+
+					// Add the source file path to the command-line.
+					FileArguments.AppendFormat(" \"{0}\"", SourceFile.AbsolutePath);
+
+					CompileAction.StatusDescription = Path.GetFileName( SourceFile.AbsolutePath );
+				}
+
+				if( bEmitsObjectFile )
+				{
+					var ObjectFileExtension = UEBuildPlatform.BuildPlatformDictionary[UnrealTargetPlatform.WinUAP].GetBinaryExtension(UEBuildBinaryType.Object);
+					// Add the object file to the produced item list.
+					FileItem ObjectFile = FileItem.GetItemByPath(
+						Path.Combine(
+							CompileEnvironment.Config.OutputDirectory,
+							Path.GetFileName(SourceFile.AbsolutePath) + ObjectFileExtension
+							)
+						);
+					CompileAction.ProducedItems.Add(ObjectFile);
+					Result.ObjectFiles.Add(ObjectFile);
+					FileArguments.AppendFormat(" /Fo\"{0}\"", ObjectFile.AbsolutePath);
+				}
+
+				// Create PDB files if we were configured to do that.
+				//
+				// Also, when debug info is off and XGE is enabled, force PDBs, otherwise it will try to share
+				// a PDB file, which causes PCH creation to be serial rather than parallel (when debug info is disabled)
+				//		--> See https://udn.epicgames.com/lists/showpost.php?id=50619&list=unprog3
+				if (BuildConfiguration.bUsePDBFiles ||
+					(BuildConfiguration.bAllowXGE && !CompileEnvironment.Config.bCreateDebugInfo))
+				{
+					string PDBFileName;
+					bool bActionProducesPDB = false;
+
+					// All files using the same PCH are required to share the same PDB that was used when compiling the PCH
+					if (CompileEnvironment.Config.PrecompiledHeaderAction == PrecompiledHeaderAction.Include)
+					{
+						PDBFileName = "PCH." + ModuleName + "." + Path.GetFileName(CompileEnvironment.Config.PrecompiledHeaderIncludeFilename);
+					}
+					// Files creating a PCH use a PDB per file.
+					else if (CompileEnvironment.Config.PrecompiledHeaderAction == PrecompiledHeaderAction.Create)
+					{
+						PDBFileName = "PCH." + ModuleName + "." + Path.GetFileName(CompileEnvironment.Config.PrecompiledHeaderIncludeFilename);
+						bActionProducesPDB = true;
+					}
+					// Ungrouped C++ files use a PDB per file.
+					else if( !bIsPlainCFile )
+					{
+						PDBFileName = Path.GetFileName( SourceFile.AbsolutePath );
+						bActionProducesPDB = true;
+					}
+					// Group all plain C files that doesn't use PCH into the same PDB
+					else
+					{
+						PDBFileName = "MiscPlainC";
+					}
+
+					// Specify the PDB file that the compiler should write to.
+					FileItem PDBFile = FileItem.GetItemByPath(
+							Path.Combine(
+								CompileEnvironment.Config.OutputDirectory,
+								PDBFileName + ".pdb"
+								)
+							);
+					FileArguments.AppendFormat(" /Fd\"{0}\"", PDBFile.AbsolutePath);
+
+					// Only use the PDB as an output file if we want PDBs and this particular action is
+					// the one that produces the PDB (as opposed to no debug info, where the above code
+					// is needed, but not the output PDB, or when multiple files share a single PDB, so
+					// only the action that generates it should count it as output directly)
+					if( BuildConfiguration.bUsePDBFiles && bActionProducesPDB )
+					{
+						CompileAction.ProducedItems.Add(PDBFile);
+						Result.DebugDataFiles.Add(PDBFile);
+					}
+				}
+
+				// Add C or C++ specific compiler arguments.
+				if (bIsPlainCFile)
+				{
+					AppendCLArguments_C(FileArguments);
+				}
+				else
+				{
+					AppendCLArguments_CPP(CompileEnvironment, FileArguments);
+				}
+
+				CompileAction.WorkingDirectory = Path.GetFullPath(".");
+				CompileAction.CommandPath      = EnvVars.CompilerPath;
+				CompileAction.CommandArguments = Arguments.ToString() + FileArguments.ToString() + CompileEnvironment.Config.AdditionalArguments;
+
+				if( CompileEnvironment.Config.PrecompiledHeaderAction == PrecompiledHeaderAction.Create )
+				{
+					Log.TraceVerbose("Creating PCH: " + CompileEnvironment.Config.PrecompiledHeaderIncludeFilename);
+					Log.TraceVerbose("     Command: " + CompileAction.CommandArguments);
+				}
+				else
+				{
+					Log.TraceVerbose("   Compiling: " + CompileAction.StatusDescription);
+					Log.TraceVerbose("     Command: " + CompileAction.CommandArguments);
+				}
+
+				// VC++ always outputs the source file name being compiled, so we don't need to emit this ourselves
+				CompileAction.bShouldOutputStatusDescription = false;
+
+				// Don't farm out creation of precompiled headers as it is the critical path task.
+				CompileAction.bCanExecuteRemotely =
+					CompileEnvironment.Config.PrecompiledHeaderAction != PrecompiledHeaderAction.Create ||
+					BuildConfiguration.bAllowRemotelyCompiledPCHs
+					;
+
+				// @todo: XGE has problems remote compiling C++/CLI files that use .NET Framework 4.0
+				if (CompileEnvironment.Config.CLRMode == CPPCLRMode.CLREnabled)
+				{
+					CompileAction.bCanExecuteRemotely = false;
+				}
+			}
+			return Result;
+		}
+
+		public override CPPOutput CompileRCFiles(UEBuildTarget Target, CPPEnvironment Environment, List<FileItem> RCFiles)
+		{
+			var EnvVars = VCEnvironment.SetEnvironment(Environment.Config.Target.Platform);
+
+			CPPOutput Result = new CPPOutput();
+
+			var BuildPlatform = UEBuildPlatform.GetBuildPlatformForCPPTargetPlatform(Environment.Config.Target.Platform);
+
+			foreach (FileItem RCFile in RCFiles)
+			{
+				Action CompileAction = new Action(ActionType.Compile);
+				CompileAction.CommandDescription = "Resource";
+				CompileAction.WorkingDirectory = Path.GetFullPath(".");
+				CompileAction.CommandPath = EnvVars.ResourceCompilerPath;
+				CompileAction.StatusDescription = Path.GetFileName(RCFile.AbsolutePath);
+
+				// Suppress header spew
+				CompileAction.CommandArguments += " /nologo";
+
+				// If we're compiling for 64-bit Windows, also add the _WIN64 definition to the resource
+				// compiler so that we can switch on that in the .rc file using #ifdef.
+				CompileAction.CommandArguments += " /D_WIN64";
+
+				// Language
+				CompileAction.CommandArguments += " /l 0x409";
+
+				// Include paths.
+				foreach (string IncludePath in Environment.Config.CPPIncludeInfo.IncludePaths)
+				{
+					CompileAction.CommandArguments += string.Format( " /i \"{0}\"", IncludePath );
+				}
+
+				// System include paths.
+				foreach( var SystemIncludePath in Environment.Config.CPPIncludeInfo.SystemIncludePaths)
+				{
+					CompileAction.CommandArguments += string.Format( " /i \"{0}\"", SystemIncludePath );
+				}
+
+				// Preprocessor definitions.
+				foreach (string Definition in Environment.Config.Definitions)
+				{
+					CompileAction.CommandArguments += string.Format(" /d \"{0}\"", Definition);
+				}
+
+				// Add the RES file to the produced item list.
+				FileItem CompiledResourceFile = FileItem.GetItemByPath(
+					Path.Combine(
+						Environment.Config.OutputDirectory,
+						Path.GetFileName(RCFile.AbsolutePath) + ".res"
+						)
+					);
+				CompileAction.ProducedItems.Add(CompiledResourceFile);
+				CompileAction.CommandArguments += string.Format(" /fo \"{0}\"", CompiledResourceFile.AbsolutePath);
+				Result.ObjectFiles.Add(CompiledResourceFile);
+
+				// Add the RC file as a prerequisite of the action.
+				CompileAction.CommandArguments += string.Format(" \"{0}\"", RCFile.AbsolutePath);
+
+				// Add the C++ source file and its included files to the prerequisite item list.
+				AddPrerequisiteSourceFile(Target, BuildPlatform, Environment, RCFile, CompileAction.PrerequisiteItems);
+			}
+
+			return Result;
+		}
+
+		public override FileItem LinkFiles(LinkEnvironment LinkEnvironment, bool bBuildImportLibraryOnly)
+		{
+			var EnvVars = VCEnvironment.SetEnvironment(LinkEnvironment.Config.Target.Platform);
+
+			if (LinkEnvironment.Config.bIsBuildingDotNetAssembly)
+			{
+				return FileItem.GetItemByPath(LinkEnvironment.Config.OutputFilePath);
+			}
+
+			bool bIsBuildingLibrary = LinkEnvironment.Config.bIsBuildingLibrary || bBuildImportLibraryOnly;
+			bool bIncludeDependentLibrariesInLibrary = bIsBuildingLibrary && LinkEnvironment.Config.bIncludeDependentLibrariesInLibrary;
+
+			// Get link arguments.
+			StringBuilder Arguments = new StringBuilder();
+			if (bIsBuildingLibrary)
+			{
+				AppendLibArguments(LinkEnvironment, Arguments);
+			}
+			else
+			{
+				AppendLinkArguments(LinkEnvironment, Arguments);
+			}
+
+
+
+			// If we're only building an import library, add the '/DEF' option that tells the LIB utility
+			// to simply create a .LIB file and .EXP file, and don't bother validating imports
+			if (bBuildImportLibraryOnly)
+			{
+				Arguments.Append(" /DEF");
+
+				// Ensure that the import library references the correct filename for the linked binary.
+				Arguments.AppendFormat(" /NAME:\"{0}\"", Path.GetFileName(LinkEnvironment.Config.OutputFilePath));
+			}
+
+
+			// Add delay loaded DLLs.
+			if (!bIsBuildingLibrary)
+			{
+				// Delay-load these DLLs.
+				foreach (string DelayLoadDLL in LinkEnvironment.Config.DelayLoadDLLs)
+				{
+					Arguments.AppendFormat(" /DELAYLOAD:\"{0}\"", DelayLoadDLL);
+				}
+			}
+
+			// @todo UE4 DLL: Why do I need LIBPATHs to build only export libraries with /DEF? (tbbmalloc.lib)
+			if (!LinkEnvironment.Config.bIsBuildingLibrary || (LinkEnvironment.Config.bIsBuildingLibrary && bIncludeDependentLibrariesInLibrary))
+			{
+				// Add the library paths to the argument list.
+				foreach (string LibraryPath in LinkEnvironment.Config.LibraryPaths)
+				{
+					Arguments.AppendFormat(" /LIBPATH:\"{0}\"", LibraryPath);
+				}
+
+				// Add the excluded default libraries to the argument list.
+				foreach (string ExcludedLibrary in LinkEnvironment.Config.ExcludedLibraries)
+				{
+					Arguments.AppendFormat(" /NODEFAULTLIB:\"{0}\"", ExcludedLibrary);
+				}
+			}
+
+			// For targets that are cross-referenced, we don't want to write a LIB file during the link step as that
+			// file will clobber the import library we went out of our way to generate during an earlier step.  This
+			// file is not needed for our builds, but there is no way to prevent MSVC from generating it when
+			// linking targets that have exports.  We don't want this to clobber our LIB file and invalidate the
+			// existing timstamp, so instead we simply emit it with a different name
+			string ImportLibraryFilePath = Path.Combine( LinkEnvironment.Config.IntermediateDirectory, 
+														 Path.GetFileNameWithoutExtension(LinkEnvironment.Config.OutputFilePath) + ".lib" );
+
+			if (LinkEnvironment.Config.bIsCrossReferenced && !bBuildImportLibraryOnly)
+			{
+				ImportLibraryFilePath += ".suppressed";
+			}
+
+			FileItem OutputFile;
+			if (bBuildImportLibraryOnly)
+			{
+				OutputFile = FileItem.GetItemByPath(ImportLibraryFilePath);
+			}
+			else
+			{
+				OutputFile = FileItem.GetItemByPath(LinkEnvironment.Config.OutputFilePath);
+				OutputFile.bNeedsHotReloadNumbersDLLCleanUp = LinkEnvironment.Config.bIsBuildingDLL;
+			}
+
+			List<FileItem> ProducedItems = new List<FileItem>();
+			ProducedItems.Add(OutputFile);
+
+			List<FileItem> PrerequisiteItems = new List<FileItem>();
+
+			// Add the input files to a response file, and pass the response file on the command-line.
+			List<string> InputFileNames = new List<string>();
+			foreach (FileItem InputFile in LinkEnvironment.InputFiles)
+			{
+				InputFileNames.Add(string.Format("\"{0}\"", InputFile.AbsolutePath));
+				PrerequisiteItems.Add(InputFile);
+			}
+
+			if (!bBuildImportLibraryOnly)
+			{
+				// Add input libraries as prerequisites, too!
+				foreach (FileItem InputLibrary in LinkEnvironment.InputLibraries)
+				{
+					InputFileNames.Add(string.Format("\"{0}\"", InputLibrary.AbsolutePath));
+					PrerequisiteItems.Add(InputLibrary);
+				}
+			}
+
+			if (!bIsBuildingLibrary || (LinkEnvironment.Config.bIsBuildingLibrary && bIncludeDependentLibrariesInLibrary))
+			{
+				foreach (string AdditionalLibrary in LinkEnvironment.Config.AdditionalLibraries)
+				{
+					InputFileNames.Add(string.Format("\"{0}\"", AdditionalLibrary));
+
+					// If the library file name has a relative path attached (rather than relying on additional
+					// lib directories), then we'll add it to our prerequisites list.  This will allow UBT to detect
+					// when the binary needs to be relinked because a dependent external library has changed.
+					//if( !String.IsNullOrEmpty( Path.GetDirectoryName( AdditionalLibrary ) ) )
+					{
+						PrerequisiteItems.Add(FileItem.GetItemByPath(AdditionalLibrary));
+					}
+				}
+			}
+
+			// Create a response file for the linker
+			string ResponseFileName = GetResponseFileName( LinkEnvironment, OutputFile );
+
+			// Never create response files when we are only generating IntelliSense data
+			if( !ProjectFileGenerator.bGenerateProjectFiles )
+			{
+				ResponseFile.Create( ResponseFileName, InputFileNames );
+			}
+			Arguments.AppendFormat(" @\"{0}\"", ResponseFileName);
+
+			// Add the output file to the command-line.
+			Arguments.AppendFormat(" /OUT:\"{0}\"", OutputFile.AbsolutePath);
+
+			if (bBuildImportLibraryOnly || (LinkEnvironment.Config.bHasExports && !bIsBuildingLibrary))
+			{
+				// An export file is written to the output directory implicitly; add it to the produced items list.
+				string ExportFilePath = Path.ChangeExtension(ImportLibraryFilePath, ".exp");
+				FileItem ExportFile = FileItem.GetItemByPath(ExportFilePath);
+				ProducedItems.Add(ExportFile);
+			}
+
+			if (!bIsBuildingLibrary)
+			{
+				// Xbox 360 LTCG does not seem to produce those.
+				if (LinkEnvironment.Config.bHasExports &&
+					(LinkEnvironment.Config.Target.Configuration != CPPTargetConfiguration.Shipping))
+				{
+					// Write the import library to the output directory for nFringe support.
+					FileItem ImportLibraryFile = FileItem.GetItemByPath(ImportLibraryFilePath);
+					Arguments.AppendFormat(" /IMPLIB:\"{0}\"", ImportLibraryFilePath);
+					ProducedItems.Add(ImportLibraryFile);
+				}
+
+				if (LinkEnvironment.Config.bCreateDebugInfo)
+				{
+					// Write the PDB file to the output directory.			
+					string PDBFilePath = Path.Combine(LinkEnvironment.Config.OutputDirectory, Path.GetFileNameWithoutExtension(OutputFile.AbsolutePath) + ".pdb");
+					FileItem PDBFile = FileItem.GetItemByPath(PDBFilePath);
+					Arguments.AppendFormat(" /PDB:\"{0}\"", PDBFilePath);
+					ProducedItems.Add(PDBFile);
+
+					// Write the MAP file to the output directory.			
+#if false					
+					if (true)
+					{
+						string MAPFilePath = Path.Combine(LinkEnvironment.Config.OutputDirectory, Path.GetFileNameWithoutExtension(OutputFile.AbsolutePath) + ".map");
+						FileItem MAPFile = FileItem.GetItemByPath(MAPFilePath);
+						LinkAction.CommandArguments += string.Format(" /MAP:\"{0}\"", MAPFilePath);
+						LinkAction.ProducedItems.Add(MAPFile);
+					}
+#endif
+				}
+
+				// Add the additional arguments specified by the environment.
+				Arguments.Append(LinkEnvironment.Config.AdditionalArguments);
+			}
+
+			// Create an action that invokes the linker.
+			Action LinkAction = new Action(ActionType.Link);
+			LinkAction.CommandDescription = "Link";
+			LinkAction.WorkingDirectory   = Path.GetFullPath(".");
+			LinkAction.CommandPath        = bIsBuildingLibrary ? EnvVars.LibraryLinkerPath : EnvVars.LinkerPath;
+			LinkAction.CommandArguments   = Arguments.ToString();
+			LinkAction.ProducedItems    .AddRange(ProducedItems);
+			LinkAction.PrerequisiteItems.AddRange(PrerequisiteItems);
+			LinkAction.StatusDescription = Path.GetFileName(OutputFile.AbsolutePath);
+
+			// Tell the action that we're building an import library here and it should conditionally be
+			// ignored as a prerequisite for other actions
+			LinkAction.bProducesImportLibrary = bBuildImportLibraryOnly || LinkEnvironment.Config.bIsBuildingDLL;
+
+			// Only execute linking on the local PC.
+			LinkAction.bCanExecuteRemotely = false;
+
+			Log.TraceVerbose("     Linking: " + LinkAction.StatusDescription);
+			Log.TraceVerbose("     Command: " + LinkAction.CommandArguments);
+
+			return OutputFile;
+		}
+
+  
+		/** Gets the default include paths for the given platform. */
+		public static string GetVCIncludePaths(CPPTargetPlatform Platform)
+		{
+			Debug.Assert(Platform == CPPTargetPlatform.WinUAP);
+
+			// Make sure we've got the environment variables set up for this target
+			VCEnvironment.SetEnvironment(Platform);
+
+			// Also add any include paths from the INCLUDE environment variable.  MSVC is not necessarily running with an environment that
+			// matches what UBT extracted from the vcvars*.bat using SetEnvironmentVariablesFromBatchFile().  We'll use the variables we
+			// extracted to populate the project file's list of include paths
+			// @todo projectfiles: Should we only do this for VC++ platforms?
+			var IncludePaths = Environment.GetEnvironmentVariable("INCLUDE");
+			if (!String.IsNullOrEmpty(IncludePaths) && !IncludePaths.EndsWith(";"))
+			{
+				IncludePaths += ";";
+			}
+			
+			return IncludePaths;
+		}
+	};
+}
diff --git a/Engine/Source/Programs/UnrealBuildTool/Windows/UEBuildWindows.cs b/Engine/Source/Programs/UnrealBuildTool/Windows/UEBuildWindows.cs
--- a/Engine/Source/Programs/UnrealBuildTool/Windows/UEBuildWindows.cs
+++ b/Engine/Source/Programs/UnrealBuildTool/Windows/UEBuildWindows.cs
@@ -10,17 +10,20 @@ using System.Linq;
 
 namespace UnrealBuildTool
 {
-    /// <summary>
-    /// Available compiler toolchains on Windows platform
-    /// </summary>
-    public enum WindowsCompiler
-    {
-        /// Visual Studio 2012 (Visual C++ 11.0). No longer supported for building on Windows, but required for other platform toolchains.
-        VisualStudio2012,
+	/// <summary>
+	/// Available compiler toolchains on Windows platform
+	/// </summary>
+	public enum WindowsCompiler
+	{
+		/// Visual Studio 2012 (Visual C++ 11.0). No longer supported for building on Windows, but required for other platform toolchains.
+		VisualStudio2012,
 
-        /// Visual Studio 2013 (Visual C++ 12.0)
-        VisualStudio2013,
-    }
+		/// Visual Studio 2013 (Visual C++ 12.0)
+		VisualStudio2013,
+
+		/// Visual Studio 2015 (Visual C++ 14.0)
+		VisualStudio2015,
+	}
 
 
     public class WindowsPlatform : UEBuildPlatform
@@ -39,20 +42,44 @@ namespace UnrealBuildTool
                     return CachedCompiler.Value;
                 }
 
-                // First check if Visual Sudio 2013 is installed.
-                if (!String.IsNullOrEmpty(WindowsPlatform.GetVSComnToolsPath(WindowsCompiler.VisualStudio2013)))
-                {
-                    CachedCompiler = WindowsCompiler.VisualStudio2013;
-                }
-                // Finally assume 2013 is installed to defer errors somewhere else like VCToolChain
-                else
-                {
-                    CachedCompiler = WindowsCompiler.VisualStudio2013;
-                }
+				// First, default based on whether there is a command line override...
+				if (UnrealBuildTool.CommandLineContains("-2012"))
+				{
+					CachedCompiler = WindowsCompiler.VisualStudio2012;
+				}
+				else if (UnrealBuildTool.CommandLineContains("-2013"))
+				{
+					CachedCompiler = WindowsCompiler.VisualStudio2013;
+				}
+				else if (UnrealBuildTool.CommandLineContains("-2015"))
+				{
+					CachedCompiler = WindowsCompiler.VisualStudio2015;
+				}
 
-                return CachedCompiler.Value;
-            }
-        }
+				// Second, default based on what's installed, from newest to oldest
+				// @todo UAP: Currently we default to VS 2013, even if VS 2015 is installed, until the engine is fully updated to
+				// support compiling with VS 2015.  You can override this with the "-2015" command-line switch.
+				else if (!String.IsNullOrEmpty(WindowsPlatform.GetVSComnToolsPath(WindowsCompiler.VisualStudio2015)) && false)
+				{
+					CachedCompiler = WindowsCompiler.VisualStudio2015;
+				}
+				else if (!String.IsNullOrEmpty(WindowsPlatform.GetVSComnToolsPath(WindowsCompiler.VisualStudio2013)))
+				{
+					CachedCompiler = WindowsCompiler.VisualStudio2013;
+				}
+				else if (!String.IsNullOrEmpty(WindowsPlatform.GetVSComnToolsPath(WindowsCompiler.VisualStudio2012)))
+				{
+					CachedCompiler = WindowsCompiler.VisualStudio2012;
+				}
+				else
+				{
+					// Finally assume 2013 is installed to defer errors somewhere else like VCToolChain
+					CachedCompiler = WindowsCompiler.VisualStudio2013;
+				}
+
+				return CachedCompiler.Value;
+			}
+		}
 
         /// <summary>
         /// When true, throws some CL and link flags (/Bt+ and /link) to output detailed timing info.
@@ -69,6 +96,12 @@ namespace UnrealBuildTool
 		/// True if we should use the Clang linker (LLD) when bCompileWithClang is enabled, otherwise we use the MSVC linker
 		public static readonly bool bAllowClangLinker = bCompileWithClang && false;
 
+		/// Whether to compile against the Windows 10 SDK, instead of the Windows 8.1 SDK.  This requires the Visual Studio 2015
+		/// compiler or later, and the Windows 10 SDK must be installed.  The application will require at least Windows 8.x to run.
+		// @todo UAP: Expose this to be enabled more easily for building Windows 10 desktop apps
+		public static readonly bool bUseWindowsSDK10 = false;
+
+
         /// True if we're targeting Windows XP as a minimum spec.  In Visual Studio 2012 and higher, this may change how
         /// we compile and link the application (http://blogs.msdn.com/b/vcblog/archive/2012/10/08/10357555.aspx)
 		/// This is a flag to determine we should support XP if possible from XML
@@ -94,10 +127,22 @@ namespace UnrealBuildTool
 
                 try
                 {
-                    // Interrogate the Win32 registry
-                    string DTEKey = (Compiler == WindowsCompiler.VisualStudio2013) ? "VisualStudio.DTE.12.0" : "VisualStudio.DTE.11.0";
-                    return RegistryKey.OpenBaseKey(RegistryHive.ClassesRoot, RegistryView.Registry32).OpenSubKey(DTEKey) != null;
-                }
+					// Interrogate the Win32 registry
+					string DTEKey = null;
+					switch (Compiler)
+					{
+						case WindowsCompiler.VisualStudio2015:
+							DTEKey = "VisualStudio.DTE.14.0";
+							break;
+						case WindowsCompiler.VisualStudio2013:
+							DTEKey = "VisualStudio.DTE.12.0";
+							break;
+						case WindowsCompiler.VisualStudio2012:
+							DTEKey = "VisualStudio.DTE.11.0";
+							break;
+					}
+					return RegistryKey.OpenBaseKey(RegistryHive.ClassesRoot, RegistryView.Registry32).OpenSubKey(DTEKey) != null;
+				}
                 catch(Exception) 
                 {
                     return false;
@@ -137,17 +182,20 @@ namespace UnrealBuildTool
         {
             int VSVersion;
 
-            switch(Compiler)
-            {
-                case WindowsCompiler.VisualStudio2012:
-                    VSVersion = 11;
-                    break;
-                case WindowsCompiler.VisualStudio2013:
-                    VSVersion = 12;
-                    break;
-                default:
-                    throw new NotSupportedException("Not supported compiler.");
-            }
+			switch(Compiler)
+			{
+				case WindowsCompiler.VisualStudio2012:
+					VSVersion = 11;
+					break;
+				case WindowsCompiler.VisualStudio2013:
+					VSVersion = 12;
+					break;
+				case WindowsCompiler.VisualStudio2015:
+					VSVersion = 14;
+					break;
+				default:
+					throw new NotSupportedException("Not supported compiler.");
+			}
 
             string[] PossibleRegPaths = new string[] {
                 @"Wow6432Node\Microsoft\VisualStudio",	// Non-express VS2013 on 64-bit machine.
@@ -251,20 +299,23 @@ namespace UnrealBuildTool
         /// When using a Visual Studio compiler, returns the version name as a string
         /// </summary>
         /// <returns>The Visual Studio compiler version name (e.g. "2012")</returns>
-        public static string GetVisualStudioCompilerVersionName()
-        {
-            switch( Compiler )
-            { 
-                case WindowsCompiler.VisualStudio2012:
-                    return "2012";
+		public static string GetVisualStudioCompilerVersionName()
+		{
+			switch( Compiler )
+			{ 
+				case WindowsCompiler.VisualStudio2012:
+					return "2012";
 
-                case WindowsCompiler.VisualStudio2013:
-                    return "2013";
+				case WindowsCompiler.VisualStudio2013:
+					return "2013";
 
-                default:
-                    throw new BuildException( "Unexpected WindowsCompiler version for GetVisualStudioCompilerVersionName().  Either not using a Visual Studio compiler or switch block needs to be updated");
-            }
-        }
+				case WindowsCompiler.VisualStudio2015:
+					return "2015";
+
+				default:
+					throw new BuildException( "Unexpected WindowsCompiler version for GetVisualStudioCompilerVersionName().  Either not using a Visual Studio compiler or switch block needs to be updated");
+			}
+		}
 
 
         /**
@@ -439,21 +490,33 @@ namespace UnrealBuildTool
         {
             InBuildTarget.GlobalCompileEnvironment.Config.Definitions.Add("WIN32=1");
 
-			// Win32 XP is only supported at this time.
-			SupportWindowsXP = SupportWindowsXPIfAvailable && (GetCPPTargetPlatform(InBuildTarget.Platform) == CPPTargetPlatform.Win32);
-			if (IsWindowsXPSupported())
-            {
-                // Windows XP SP3 or higher required
-                InBuildTarget.GlobalCompileEnvironment.Config.Definitions.Add("_WIN32_WINNT=0x0502");
-                InBuildTarget.GlobalCompileEnvironment.Config.Definitions.Add("WINVER=0x0502");
-            }
-            else
-            {
-                // Windows Vista or higher required
-                InBuildTarget.GlobalCompileEnvironment.Config.Definitions.Add("_WIN32_WINNT=0x0600");
-                InBuildTarget.GlobalCompileEnvironment.Config.Definitions.Add("WINVER=0x0600");
-            }
-            InBuildTarget.GlobalCompileEnvironment.Config.Definitions.Add("PLATFORM_WINDOWS=1");
+			if (WindowsPlatform.bUseWindowsSDK10 && WindowsPlatform.Compiler == WindowsCompiler.VisualStudio2015)
+			{
+				if (SupportWindowsXP)
+				{
+					throw new NotSupportedException("Windows XP support is not possible when targeting the Windows 10 SDK");
+				}
+				// Windows 8 or higher required
+				InBuildTarget.GlobalCompileEnvironment.Config.Definitions.Add("_WIN32_WINNT=0x0602");
+				InBuildTarget.GlobalCompileEnvironment.Config.Definitions.Add("WINVER=0x0602");
+			}
+			else
+			{
+				SupportWindowsXP = SupportWindowsXP && (GetCPPTargetPlatform(InBuildTarget.Platform) == CPPTargetPlatform.Win32);
+				if (SupportWindowsXP)
+				{
+					// Windows XP SP3 or higher required
+					InBuildTarget.GlobalCompileEnvironment.Config.Definitions.Add("_WIN32_WINNT=0x0502");
+					InBuildTarget.GlobalCompileEnvironment.Config.Definitions.Add("WINVER=0x0502");
+				}
+				else
+				{
+					// Windows Vista or higher required
+					InBuildTarget.GlobalCompileEnvironment.Config.Definitions.Add("_WIN32_WINNT=0x0600");
+					InBuildTarget.GlobalCompileEnvironment.Config.Definitions.Add("WINVER=0x0600");
+				}
+			}
+			InBuildTarget.GlobalCompileEnvironment.Config.Definitions.Add("PLATFORM_WINDOWS=1");
 
             String MorpheusShaderPath = Path.Combine(BuildConfiguration.RelativeEnginePath, "Shaders/PS4/PostProcessHMDMorpheus.usf");
             if (File.Exists(MorpheusShaderPath))
diff --git a/Engine/Source/Programs/UnrealBuildTool/Windows/VCEnvironment.cs b/Engine/Source/Programs/UnrealBuildTool/Windows/VCEnvironment.cs
--- a/Engine/Source/Programs/UnrealBuildTool/Windows/VCEnvironment.cs
+++ b/Engine/Source/Programs/UnrealBuildTool/Windows/VCEnvironment.cs
@@ -16,6 +16,9 @@ namespace UnrealBuildTool
 		public readonly string            BaseVSToolPath;       // The path to Visual Studio's /Common7/Tools directory.
 		public readonly string            PlatformVSToolPath;   // The path to the platform tool binaries.
 		public readonly string            WindowsSDKDir;        // Installation folder of the Windows SDK, e.g. C:\Program Files\Microsoft SDKs\Windows\v6.0A\
+		public readonly string            WindowsSDKExtensionDir;  // Installation folder of the Windows SDK Extensions, e.g. C:\Program Files (x86)\Windows SDKs\10
+		public readonly string            NetFxSDKExtensionDir;    // Installation folder of the NetFx SDK, since that is split out from platform SDKs >= v10
+		public readonly Version           WindowsSDKExtensionHeaderLibVersion;  // 10.0.9910.0 for instance...
 		public readonly string            CompilerPath;         // The path to the linker for linking executables
 		public readonly Version           CLExeVersion;         // The version of cl.exe we're running
 		public readonly string            LinkerPath;           // The path to the linker for linking executables
@@ -58,19 +61,30 @@ namespace UnrealBuildTool
 			BaseVSToolPath = WindowsPlatform.GetVSComnToolsPath();
 			if (string.IsNullOrEmpty(BaseVSToolPath))
 			{
-				throw new BuildException("Visual Studio 2012 or Visual Studio 2013 must be installed in order to build this target.");
+				throw new BuildException("Visual Studio 2012, 2013 or 2015 must be installed in order to build this target.");
 			}
 
 			WindowsSDKDir        = FindWindowsSDKInstallationFolder(Platform);
-			PlatformVSToolPath   = GetPlatformVSToolPath      (Platform, BaseVSToolPath);
+			WindowsSDKExtensionDir = FindWindowsSDKExtensionInstallationFolder();
+			NetFxSDKExtensionDir = FindNetFxSDKExtensionInstallationFolder();
+			WindowsSDKExtensionHeaderLibVersion = FindWindowsSDKExtensionLatestVersion(WindowsSDKExtensionDir);
+			PlatformVSToolPath = GetPlatformVSToolPath      (Platform, BaseVSToolPath);
 			CompilerPath         = GetCompilerToolPath        (PlatformVSToolPath);
 			CLExeVersion         = FindCLExeVersion           (CompilerPath);
 			LinkerPath           = GetLinkerToolPath          (PlatformVSToolPath);
 			LibraryLinkerPath    = GetLibraryLinkerToolPath   (PlatformVSToolPath);
-			ResourceCompilerPath = GetResourceCompilerToolPath(Platform, WindowsSDKDir);
+			ResourceCompilerPath = GetResourceCompilerToolPath(Platform);
 
-			var VCVarsBatchFile = Path.Combine(BaseVSToolPath, (Platform == CPPTargetPlatform.Win64) ? "../../VC/bin/x86_amd64/vcvarsx86_amd64.bat" : "vsvars32.bat");
-			Utils.SetEnvironmentVariablesFromBatchFile(VCVarsBatchFile);
+			var VCVarsBatchFile = Path.Combine(BaseVSToolPath, 
+				(Platform == CPPTargetPlatform.Win64 || Platform == CPPTargetPlatform.WinUAP) ? "../../VC/bin/x86_amd64/vcvarsx86_amd64.bat" : "vsvars32.bat");
+			if (Platform == CPPTargetPlatform.WinUAP && WinUAPPlatform.bBuildForStore)
+			{
+				Utils.SetEnvironmentVariablesFromBatchFile(VCVarsBatchFile, "store");
+			}
+			else
+			{
+				Utils.SetEnvironmentVariablesFromBatchFile(VCVarsBatchFile);
+			}
 
 			// When targeting Windows XP on Visual Studio 2012+, we need to override the Windows SDK include and lib path set
 			// by the batch file environment (http://blogs.msdn.com/b/vcblog/archive/2012/10/08/10357555.aspx)
@@ -97,6 +111,17 @@ namespace UnrealBuildTool
 			}
 			else switch (WindowsPlatform.Compiler)
 			{
+				case WindowsCompiler.VisualStudio2015:
+					if( WindowsPlatform.bUseWindowsSDK10 )
+					{ 
+						Version = "v10.0";
+					}
+					else
+					{
+						Version = "v8.1";
+					}
+					break;
+
 				case WindowsCompiler.VisualStudio2013:
 					Version = "v8.1";
 					break;
@@ -110,17 +135,127 @@ namespace UnrealBuildTool
 			}
 
 			// Based on VCVarsQueryRegistry
-			var Result =
-					Microsoft.Win32.Registry.GetValue(@"HKEY_CURRENT_USER\SOFTWARE\Microsoft\Microsoft SDKs\Windows\"              + Version, "InstallationFolder", null)
-				?? Microsoft.Win32.Registry.GetValue(@"HKEY_LOCAL_MACHINE\SOFTWARE\Wow6432Node\Microsoft\Microsoft SDKs\Windows\" + Version, "InstallationFolder", null)
-				?? Microsoft.Win32.Registry.GetValue(@"HKEY_CURRENT_USER\SOFTWARE\Wow6432Node\Microsoft\Microsoft SDKs\Windows\"  + Version, "InstallationFolder", null);
+			string FinalResult = null;
+			foreach (string IndividualVersion in Version.Split('|'))
+			{
+				var Result = Microsoft.Win32.Registry.GetValue(@"HKEY_CURRENT_USER\SOFTWARE\Microsoft\Microsoft SDKs\Windows\" + IndividualVersion, "InstallationFolder", null)
+					?? Microsoft.Win32.Registry.GetValue(@"HKEY_LOCAL_MACHINE\SOFTWARE\Wow6432Node\Microsoft\Microsoft SDKs\Windows\" + IndividualVersion, "InstallationFolder", null)
+					?? Microsoft.Win32.Registry.GetValue(@"HKEY_CURRENT_USER\SOFTWARE\Wow6432Node\Microsoft\Microsoft SDKs\Windows\" + IndividualVersion, "InstallationFolder", null);
 
-			if (Result == null)
+				if (Result != null)
+				{
+					FinalResult = (string)Result;
+					break;
+				}
+			}
+			if (FinalResult == null)
 			{
 				throw new BuildException("Windows SDK {0} must be installed in order to build this target.", Version);
 			}
 
-			return (string)Result;
+			return FinalResult;
+		}
+
+		
+		private static string FindNetFxSDKExtensionInstallationFolder()
+		{
+			string Version;
+			switch (WindowsPlatform.Compiler)
+			{
+				case WindowsCompiler.VisualStudio2015:
+					if( WindowsPlatform.bUseWindowsSDK10 )
+					{ 
+						Version = "4.6";
+					}
+					else
+					{ 
+						return string.Empty;
+					}
+					break;
+
+				default:
+					return string.Empty;
+			}
+			string FinalResult = string.Empty;
+			var Result = Microsoft.Win32.Registry.GetValue(@"HKEY_LOCAL_MACHINE\SOFTWARE\Microsoft\Microsoft SDKs\NETFXSDK\" + Version, "KitsInstallationFolder", null)
+					  ?? Microsoft.Win32.Registry.GetValue(@"HKEY_LOCAL_MACHINE\SOFTWARE\Wow6432Node\Microsoft\Microsoft SDKs\NETFXSDK\" + Version, "KitsInstallationFolder", null);
+
+			if (Result != null)
+			{
+				FinalResult = ((string)Result).TrimEnd('\\');
+			}
+			return FinalResult;
+		}
+
+		private static string FindWindowsSDKExtensionInstallationFolder()
+		{
+			string Version;
+			switch (WindowsPlatform.Compiler)
+			{
+				case WindowsCompiler.VisualStudio2015:
+					if( WindowsPlatform.bUseWindowsSDK10 )
+					{ 
+						Version = "v10.0";
+					}
+					else
+					{
+						return string.Empty;
+					}
+					break;
+
+				default:
+					return string.Empty;
+			}
+
+			// Based on VCVarsQueryRegistry
+			string FinalResult = null;
+			{
+				var Result = Microsoft.Win32.Registry.GetValue(@"HKEY_LOCAL_MACHINE\SOFTWARE\Wow6432Node\Microsoft\Microsoft SDKs\Windows SDKs\" + Version, "InstallationFolder", null)
+						  ?? Microsoft.Win32.Registry.GetValue(@"HKEY_CURRENT_USER\SOFTWARE\Wow6432Node\Microsoft\Microsoft SDKs\Windows SDKs\" + Version, "InstallationFolder", null);
+				if (Result == null)
+				{
+					 Result = Microsoft.Win32.Registry.GetValue(@"HKEY_LOCAL_MACHINE\SOFTWARE\Wow6432Node\Microsoft\Microsoft SDKs\Windows\" + Version, "InstallationFolder", null)
+						   ?? Microsoft.Win32.Registry.GetValue(@"HKEY_CURRENT_USER\SOFTWARE\Wow6432Node\Microsoft\Microsoft SDKs\Windows\" + Version, "InstallationFolder", null);
+				}
+				if (Result != null)
+				{
+					FinalResult = ((string)Result).TrimEnd('\\');
+				}
+				
+			}
+			if (FinalResult == null)
+			{
+				FinalResult = string.Empty;
+			}
+
+			return FinalResult;
+		}
+
+		static Version FindWindowsSDKExtensionLatestVersion(string WindowsSDKExtensionDir)
+		{
+			Version LatestVersion = new Version(0, 0, 0, 0);
+
+			if (WindowsPlatform.Compiler == WindowsCompiler.VisualStudio2015 &&
+				WindowsPlatform.bUseWindowsSDK10 &&
+				!string.IsNullOrEmpty(WindowsSDKExtensionDir)  &&
+				Directory.Exists(WindowsSDKExtensionDir))
+			{
+				string IncludesBaseDirectory = Path.Combine(WindowsSDKExtensionDir, "include");
+				if (Directory.Exists(IncludesBaseDirectory))
+				{
+					string[] IncludeVersions = Directory.GetDirectories(IncludesBaseDirectory);
+					foreach (string IncludeVersion in IncludeVersions)
+					{
+						string VersionString = Path.GetFileName(IncludeVersion);
+						Version FoundVersion;
+						if (Version.TryParse(VersionString, out FoundVersion) && FoundVersion > LatestVersion)
+						{
+							LatestVersion = FoundVersion;
+						}
+					}
+				}
+			}
+			return LatestVersion;
 		}
 
 		/** Gets the path to the tool binaries for the specified platform. */
@@ -128,7 +263,7 @@ namespace UnrealBuildTool
 		{
 			// Regardless of the target, if we're linking on a 64 bit machine, we want to use the 64 bit linker (it's faster than the 32 bit linker)
 			[REDACTED_EMAIL]: Using the 64-bit linker appears to be broken at the moment.
-			if (Platform == CPPTargetPlatform.Win64)
+			if (Platform == CPPTargetPlatform.Win64 || Platform == CPPTargetPlatform.WinUAP)
 			{
 				// Use the native 64-bit compiler if present, otherwise use the amd64-on-x86 compiler. VS2012 Express only includes the latter.
 				var Result = Path.Combine(BaseVSToolPath, "../../VC/bin/amd64");
@@ -210,19 +345,33 @@ namespace UnrealBuildTool
 		}
 
 		/** Gets the path to the resource compiler's rc.exe for the specified platform. */
-		static string GetResourceCompilerToolPath(CPPTargetPlatform Platform, string WindowsSDKDir)
+		string GetResourceCompilerToolPath(CPPTargetPlatform Platform)
 		{
 			// 64 bit -- we can use the 32 bit version to target 64 bit on 32 bit OS.
-			if (Platform == CPPTargetPlatform.Win64)
+			if (Platform == CPPTargetPlatform.Win64 || Platform == CPPTargetPlatform.WinUAP)
 			{
-				return Path.Combine(WindowsSDKDir, "bin/x64/rc.exe");
+				if (WindowsPlatform.Compiler == WindowsCompiler.VisualStudio2015 && WindowsPlatform.bUseWindowsSDK10)
+				{
+					return Path.Combine(WindowsSDKExtensionDir, "bin/x64/rc.exe");
+				}
+				else
+				{
+					return Path.Combine(WindowsSDKDir, "bin/x64/rc.exe");
+				}
 			}
 
+			// @todo UAP: Verify that Windows XP will compile using VS 2015 (it should be supported)
 			if (!WindowsPlatform.IsWindowsXPSupported())	// Windows XP requires use to force Windows SDK 7.1 even on the newer compiler, so we need the old path RC.exe
 			{
-				return Path.Combine(WindowsSDKDir, "bin/x86/rc.exe");
+				if (WindowsPlatform.Compiler == WindowsCompiler.VisualStudio2015 && WindowsPlatform.bUseWindowsSDK10)
+				{
+					return Path.Combine(WindowsSDKExtensionDir, "bin/x86/rc.exe");
+				}
+				else
+				{
+					return Path.Combine(WindowsSDKDir, "bin/x86/rc.exe");
+				}
 			}
-
 			return Path.Combine(WindowsSDKDir, "bin/rc.exe");
 		}
 
diff --git a/Engine/Source/Programs/UnrealBuildTool/Windows/VCToolChain.cs b/Engine/Source/Programs/UnrealBuildTool/Windows/VCToolChain.cs
--- a/Engine/Source/Programs/UnrealBuildTool/Windows/VCToolChain.cs
+++ b/Engine/Source/Programs/UnrealBuildTool/Windows/VCToolChain.cs
@@ -91,6 +91,22 @@ namespace UnrealBuildTool
 
 		static void AppendCLArguments_Global(CPPEnvironment CompileEnvironment, VCEnvironment EnvVars, StringBuilder Arguments)
 		{
+			// @todo UAP: Why do we ever need WinRT headers when building regular Win32?  Is this just needed for the Windows 10 SDK?
+			// @todo UAP: These include paths should be added in SetUpEnvironment(), not here.  Do they need to be the last includes or something?
+			if (WindowsPlatform.Compiler == WindowsCompiler.VisualStudio2015 && WindowsPlatform.bUseWindowsSDK10)
+			{
+				if (Directory.Exists(EnvVars.WindowsSDKExtensionDir))
+				{
+					CompileEnvironment.Config.CPPIncludeInfo.SystemIncludePaths.Add(string.Format(@"{0}\Include\{1}\um", EnvVars.WindowsSDKExtensionDir, EnvVars.WindowsSDKExtensionHeaderLibVersion));
+					CompileEnvironment.Config.CPPIncludeInfo.SystemIncludePaths.Add(string.Format(@"{0}\Include\{1}\shared", EnvVars.WindowsSDKExtensionDir, EnvVars.WindowsSDKExtensionHeaderLibVersion));
+					CompileEnvironment.Config.CPPIncludeInfo.SystemIncludePaths.Add(string.Format(@"{0}\Include\{1}\winrt", EnvVars.WindowsSDKExtensionDir, EnvVars.WindowsSDKExtensionHeaderLibVersion));
+				}
+				if (Directory.Exists(EnvVars.NetFxSDKExtensionDir))
+				{
+					CompileEnvironment.Config.CPPIncludeInfo.SystemIncludePaths.Add(string.Format(@"{0}\Include\um", EnvVars.NetFxSDKExtensionDir));
+				}
+			}
+
 			// NOTE: Uncommenting this line will print includes as they are encountered by the preprocessor.  This can help with diagnosing include order problems.
 			if( WindowsPlatform.bCompileWithClang && !WindowsPlatform.bUseVCCompilerArgs )
 			{ 
@@ -127,6 +143,11 @@ namespace UnrealBuildTool
 						FullVersionString = "1800";
 						break;
 
+					case WindowsCompiler.VisualStudio2015:
+						VersionString = "19.0";
+						FullVersionString = "1900";
+						break;
+
 					default:
 						throw new BuildException( "Unexpected value for WindowsPlatform.Compiler: " + WindowsPlatform.Compiler.ToString() );
 				}
@@ -199,10 +220,25 @@ namespace UnrealBuildTool
 				Arguments.Append(" /wd4819");
 			}
 
-			if( BuildConfiguration.bUseSharedPCHs )
+			// @todo UAP: UE4 is non-compliant when it comes to use of %s and %S
+			// Previously %s meant "the current character set" and %S meant "the other one".
+			// Now %s means multibyte and %S means wide. %Ts means "natural width".
+			// Reverting this behaviour until the UE4 source catches up.
+			if (WindowsPlatform.Compiler == WindowsCompiler.VisualStudio2015)
 			{
-				if( WindowsPlatform.bUseVCCompilerArgs )
-				{ 
+				Arguments.Append(" /D_CRT_STDIO_LEGACY_WIDE_SPECIFIERS=1");
+			}
+
+			// @todo UAP: Silence the hash_map deprecation errors for now. This should be replaced with unordered_map for the real fix.
+			if (WindowsPlatform.Compiler == WindowsCompiler.VisualStudio2015)
+			{
+				Arguments.Append(" /D_SILENCE_STDEXT_HASH_DEPRECATION_WARNINGS");
+			}
+
+			if ( BuildConfiguration.bUseSharedPCHs )
+			{
+				if (WindowsPlatform.bUseVCCompilerArgs)
+				{
 					// @todo SharedPCH: Disable warning about PCH defines not matching .cpp defines.  We "cheat" these defines a little
 					// bit to make shared PCHs work.  But it's totally safe.  Trust us.
 					Arguments.Append(" /wd4651");
@@ -379,13 +415,13 @@ namespace UnrealBuildTool
 					// For C++/CLI all exceptions must be left enabled
 					Arguments.Append(" /EHa");
 				}
-            }
-    		else if (CompileEnvironment.Config.Target.Platform == CPPTargetPlatform.HTML5)
-            {
-                Arguments.Append(" /EHsc"); 
-            }
+			}
+			else if (CompileEnvironment.Config.Target.Platform == CPPTargetPlatform.HTML5)
+			{
+				Arguments.Append(" /EHsc"); 
+			}
 
-            // If enabled, create debug information.
+			// If enabled, create debug information.
 			if (CompileEnvironment.Config.bCreateDebugInfo)
 			{
 				if( WindowsPlatform.bUseVCCompilerArgs )
@@ -408,7 +444,7 @@ namespace UnrealBuildTool
 						}
 						// We need to add this so VS won't lock the PDB file and prevent synchronous updates. This forces serialization through MSPDBSRV.exe.
 						// See http://msdn.microsoft.com/en-us/library/dn502518.aspx for deeper discussion of /FS switch.
-						if (BuildConfiguration.bUseIncrementalLinking && WindowsPlatform.Compiler == WindowsCompiler.VisualStudio2013)
+						if (BuildConfiguration.bUseIncrementalLinking && WindowsPlatform.Compiler >= WindowsCompiler.VisualStudio2013)
 						{
 							Arguments.Append(" /FS");
 						}
@@ -1006,11 +1042,11 @@ namespace UnrealBuildTool
 				}
 			}
 
-            if (!WindowsPlatform.bCompileWithClang && WindowsPlatform.bLogDetailedCompilerTimingInfo)
-            {
-                // Force MSVC
-                SharedArguments.Append(" /Bt+");
-            }
+			if (!WindowsPlatform.bCompileWithClang && WindowsPlatform.bLogDetailedCompilerTimingInfo)
+			{
+				// Force MSVC
+				SharedArguments.Append(" /Bt+");
+			}
 
 			// Add preprocessor definitions to the argument list.
 			foreach (string Definition in CompileEnvironment.Config.Definitions)
@@ -1028,11 +1064,11 @@ namespace UnrealBuildTool
 			{
 				Action CompileAction = new Action(ActionType.Compile);
 				CompileAction.CommandDescription = "Compile";
-                // ensure compiler timings are captured when we execute the action.
-                if (!WindowsPlatform.bCompileWithClang && WindowsPlatform.bLogDetailedCompilerTimingInfo)
-                {
-                    CompileAction.bPrintDebugInfo = true;
-                }
+				// ensure compiler timings are captured when we execute the action.
+				if (!WindowsPlatform.bCompileWithClang && WindowsPlatform.bLogDetailedCompilerTimingInfo)
+				{
+					CompileAction.bPrintDebugInfo = true;
+				}
 
 				StringBuilder FileArguments = new StringBuilder();
 				bool bIsPlainCFile = Path.GetExtension(SourceFile.AbsolutePath).ToUpperInvariant() == ".C";
@@ -1302,9 +1338,9 @@ namespace UnrealBuildTool
 				CompileAction.bCanExecuteRemotely =
 					CompileEnvironment.Config.PrecompiledHeaderAction != PrecompiledHeaderAction.Create ||
 					BuildConfiguration.bAllowRemotelyCompiledPCHs
-                    ;
+					;
 
-                // @todo: XGE has problems remote compiling C++/CLI files that use .NET Framework 4.0
+				// @todo: XGE has problems remote compiling C++/CLI files that use .NET Framework 4.0
 				if (CompileEnvironment.Config.CLRMode == CPPCLRMode.CLREnabled)
 				{
 					CompileAction.bCanExecuteRemotely = false;
@@ -1404,6 +1440,19 @@ namespace UnrealBuildTool
 		{
 			var EnvVars = VCEnvironment.SetEnvironment(LinkEnvironment.Config.Target.Platform);
 
+			// @todo UAP: These paths should be added in SetUpEnvironment(), not here.  Also is this actually needed for classic desktop targets or only UAP?
+			if (WindowsPlatform.Compiler == WindowsCompiler.VisualStudio2015 && WindowsPlatform.bUseWindowsSDK10)
+			{
+				if (LinkEnvironment.Config.Target.Platform == CPPTargetPlatform.Win64)
+				{
+					LinkEnvironment.Config.LibraryPaths.Add(string.Format("{0}/Lib/um/x64", EnvVars.NetFxSDKExtensionDir));
+				}
+				else if (LinkEnvironment.Config.Target.Platform == CPPTargetPlatform.Win32)
+				{
+					LinkEnvironment.Config.LibraryPaths.Add(string.Format("{0}/Lib/um/x86", EnvVars.NetFxSDKExtensionDir));
+				}
+			}
+
 			if (LinkEnvironment.Config.bIsBuildingDotNetAssembly)
 			{
 				return FileItem.GetItemByPath(LinkEnvironment.Config.OutputFilePath);
@@ -1423,10 +1472,10 @@ namespace UnrealBuildTool
 				AppendLinkArguments(LinkEnvironment, Arguments);
 			}
 
-            if (!WindowsPlatform.bCompileWithClang && WindowsPlatform.bLogDetailedCompilerTimingInfo)
-            {
-                Arguments.Append(" /time+");
-            }
+			if (!WindowsPlatform.bCompileWithClang && WindowsPlatform.bLogDetailedCompilerTimingInfo)
+			{
+				Arguments.Append(" /time+");
+			}
 
 			// If we're only building an import library, add the '/DEF' option that tells the LIB utility
 			// to simply create a .LIB file and .EXP file, and don't bother validating imports
@@ -1598,11 +1647,11 @@ namespace UnrealBuildTool
 			LinkAction.ProducedItems    .AddRange(ProducedItems);
 			LinkAction.PrerequisiteItems.AddRange(PrerequisiteItems);
 			LinkAction.StatusDescription  = Path.GetFileName(OutputFile.AbsolutePath);
-            // ensure compiler timings are captured when we execute the action.
-            if (!WindowsPlatform.bCompileWithClang && WindowsPlatform.bLogDetailedCompilerTimingInfo)
-            {
-                LinkAction.bPrintDebugInfo = true;
-            }
+			// ensure compiler timings are captured when we execute the action.
+			if (!WindowsPlatform.bCompileWithClang && WindowsPlatform.bLogDetailedCompilerTimingInfo)
+			{
+				LinkAction.bPrintDebugInfo = true;
+			}
 
 			if( WindowsPlatform.bCompileWithClang )
 			{ 
@@ -1702,13 +1751,13 @@ namespace UnrealBuildTool
 			return IncludePaths;
 		}
 
-        public override void AddFilesToReceipt(BuildReceipt Receipt, UEBuildBinary Binary)
-        {
-            if (Binary.Config.Type == UEBuildBinaryType.DynamicLinkLibrary)
-            {
+		public override void AddFilesToReceipt(BuildReceipt Receipt, UEBuildBinary Binary)
+		{
+			if (Binary.Config.Type == UEBuildBinaryType.DynamicLinkLibrary)
+			{
 				Receipt.AddBuildProduct(Path.Combine(Binary.Config.IntermediateDirectory, Path.GetFileNameWithoutExtension(Binary.Config.OutputFilePath) + ".lib"), BuildProductType.ImportLibrary);
-            }
-        }
+			}
+		}
 
 
 		/** Formats compiler output from Clang so that it is clickable in Visual Studio */
