Fixed UBT compile warning (regress from CL 2537921)

diff --git a/Engine/Source/Programs/UnrealBuildTool/Windows/UEBuildWindows.cs b/Engine/Source/Programs/UnrealBuildTool/Windows/UEBuildWindows.cs
--- a/Engine/Source/Programs/UnrealBuildTool/Windows/UEBuildWindows.cs
+++ b/Engine/Source/Programs/UnrealBuildTool/Windows/UEBuildWindows.cs
@@ -59,10 +59,10 @@ namespace UnrealBuildTool
 				// Second, default based on what's installed, from newest to oldest
 				// @todo UAP: Currently we default to VS 2013, even if VS 2015 is installed, until the engine is fully updated to
 				// support compiling with VS 2015.  You can override this with the "-2015" command-line switch.
-				else if (!String.IsNullOrEmpty(WindowsPlatform.GetVSComnToolsPath(WindowsCompiler.VisualStudio2015)) && false)
-				{
-					CachedCompiler = WindowsCompiler.VisualStudio2015;
-				}
+// 				else if (!String.IsNullOrEmpty(WindowsPlatform.GetVSComnToolsPath(WindowsCompiler.VisualStudio2015)))
+// 				{
+// 					CachedCompiler = WindowsCompiler.VisualStudio2015;
+// 				}
 				else if (!String.IsNullOrEmpty(WindowsPlatform.GetVSComnToolsPath(WindowsCompiler.VisualStudio2013)))
 				{
 					CachedCompiler = WindowsCompiler.VisualStudio2013;
