Fix UBT handling for content-only plugins, where the module array in the descriptor is null.

diff --git a/Engine/Source/Programs/UnrealBuildTool/Configuration/UEBuildTarget.cs b/Engine/Source/Programs/UnrealBuildTool/Configuration/UEBuildTarget.cs
--- a/Engine/Source/Programs/UnrealBuildTool/Configuration/UEBuildTarget.cs
+++ b/Engine/Source/Programs/UnrealBuildTool/Configuration/UEBuildTarget.cs
@@ -2213,14 +2213,17 @@ namespace UnrealBuildTool
         public void AddPlugin(PluginInfo Plugin)
 		{
 			UEBuildBinaryType BinaryType = ShouldCompileMonolithic() ? UEBuildBinaryType.StaticLibrary : UEBuildBinaryType.DynamicLinkLibrary;
-			foreach(ModuleDescriptor Module in Plugin.Descriptor.Modules)
+			if(Plugin.Descriptor.Modules != null)
 			{
-				if (Module.IsCompiledInConfiguration(Platform, TargetType))
+				foreach(ModuleDescriptor Module in Plugin.Descriptor.Modules)
 				{
-					// Add the corresponding binary for it
-					string ModuleFileName = RulesCompiler.GetModuleFilename(Module.Name);
-					bool bHasSource = (!String.IsNullOrEmpty(ModuleFileName) && Directory.EnumerateFiles(Path.GetDirectoryName(ModuleFileName), "*.cpp", SearchOption.AllDirectories).Any());
-					AddBinaryForModule(Module.Name, BinaryType, bAllowCompilation: bHasSource, bIsCrossTarget: false);
+					if (Module.IsCompiledInConfiguration(Platform, TargetType))
+					{
+						// Add the corresponding binary for it
+						string ModuleFileName = RulesCompiler.GetModuleFilename(Module.Name);
+						bool bHasSource = (!String.IsNullOrEmpty(ModuleFileName) && Directory.EnumerateFiles(Path.GetDirectoryName(ModuleFileName), "*.cpp", SearchOption.AllDirectories).Any());
+						AddBinaryForModule(Module.Name, BinaryType, bAllowCompilation: bHasSource, bIsCrossTarget: false);
+					}
 				}
 			}
 		}
@@ -2409,7 +2412,7 @@ namespace UnrealBuildTool
 		/// <returns>Matching plugin, or null if not found</returns>
 		private PluginInfo FindPluginForModule(string ModuleName)
 		{
-			return ValidPlugins.FirstOrDefault(ValidPlugin => ValidPlugin.Descriptor.Modules.Any(Module => Module.Name == ModuleName));
+			return ValidPlugins.FirstOrDefault(ValidPlugin => ValidPlugin.Descriptor.Modules != null && ValidPlugin.Descriptor.Modules.Any(Module => Module.Name == ModuleName));
 		}
 
 		/**
