Unreal Build Tool: Fixed incorrect VS tools versions set in VC++ project files

diff --git a/Engine/Source/Programs/UnrealBuildTool/System/VCProject.cs b/Engine/Source/Programs/UnrealBuildTool/System/VCProject.cs
--- a/Engine/Source/Programs/UnrealBuildTool/System/VCProject.cs
+++ b/Engine/Source/Programs/UnrealBuildTool/System/VCProject.cs
@@ -438,7 +438,7 @@ namespace UnrealBuildTool
 			var ProjectPlatforms = new List<UnrealTargetPlatform>();
 			var ProjectPlatformNames = new List<string>();
 			var ProjectConfigNames = new List<string>();
-			foreach( var Combination in ProjectConfigAndTargetCombinations )
+			foreach ( var Combination in ProjectConfigAndTargetCombinations )
 			{
 				if( !ProjectPlatforms.Contains( Combination.Platform ) )
 				{
@@ -473,7 +473,7 @@ namespace UnrealBuildTool
 			VCFiltersFileContent.Append(
 				"<?xml version=\"1.0\" encoding=\"utf-8\"?>" + ProjectFileGenerator.NewLine +
 				ProjectFileGenerator.NewLine +
-				"<Project ToolsVersion=\"4.0\" xmlns=\"http://schemas.microsoft.com/developer/msbuild/2003\">" + ProjectFileGenerator.NewLine);
+				"<Project ToolsVersion=\"" + VCProjectFileGenerator.ProjectFileToolVersionString + "\" xmlns=\"http://schemas.microsoft.com/developer/msbuild/2003\">" + ProjectFileGenerator.NewLine);
 
 			// Platform specific PropertyGroups, etc.
 			StringBuilder AdditionalPropertyGroups = new StringBuilder();
@@ -1173,7 +1173,7 @@ namespace UnrealBuildTool
 			// Project file header.
 			ProjectFileContent.Append(
 				"<?xml version=\"1.0\" encoding=\"utf-8\"?>" + ProjectFileGenerator.NewLine +
-				"<Project DefaultTargets=\"Build\" ToolsVersion=\"4.0\" xmlns=\"http://schemas.microsoft.com/developer/msbuild/2003\">" + ProjectFileGenerator.NewLine);
+				"<Project DefaultTargets=\"Build\" ToolsVersion=\"" + VCProjectFileGenerator.ProjectFileToolVersionString + "\" xmlns=\"http://schemas.microsoft.com/developer/msbuild/2003\">" + ProjectFileGenerator.NewLine);
 
 			ProjectFileContent.Append(
 				"<Import Project=\"$(MSBuildExtensionsPath)\\$(MSBuildToolsVersion)\\Microsoft.Common.props\" Condition=\"Exists('$(MSBuildExtensionsPath)\\$(MSBuildToolsVersion)\\Microsoft.Common.props')\" />" +
