Fix generated java template path (only forward slashes)
#ue4
#android
#codeview Robert.Jones,Josh.Adams

diff --git a/Engine/Source/Programs/UnrealBuildTool/Android/UEDeployAndroid.cs b/Engine/Source/Programs/UnrealBuildTool/Android/UEDeployAndroid.cs
--- a/Engine/Source/Programs/UnrealBuildTool/Android/UEDeployAndroid.cs
+++ b/Engine/Source/Programs/UnrealBuildTool/Android/UEDeployAndroid.cs
@@ -1080,7 +1080,7 @@ namespace UnrealBuildTool.Android
                        
             // Generate Java files
             string PackageName = GetPackageName(ProjectName);
-            string TemplateDestinationBase = Path.Combine(ProjectDirectory, "Build", "Android", "src" , PackageName.Replace('.', '\\'));
+            string TemplateDestinationBase = Path.Combine(ProjectDirectory, "Build", "Android", "src" , PackageName.Replace('.', '/'));
             MakeDirectoryIfRequired(TemplateDestinationBase);
             // We'll be writing the OBB data into the same location as the download service files
             string UE4OBBDataFileName = GetUE4JavaOBBDataFileName(TemplateDestinationBase);
