Fix Android StripSymboils to only strip symbols (missing parameter)
#ue4
#android

diff --git a/Engine/Source/Programs/UnrealBuildTool/Android/AndroidToolChain.cs b/Engine/Source/Programs/UnrealBuildTool/Android/AndroidToolChain.cs
--- a/Engine/Source/Programs/UnrealBuildTool/Android/AndroidToolChain.cs
+++ b/Engine/Source/Programs/UnrealBuildTool/Android/AndroidToolChain.cs
@@ -981,7 +981,7 @@ namespace UnrealBuildTool
 			{
 				throw new BuildException("Couldn't determine Android architecture to strip symbols from {0}", SourceFileName);
 			}
-			StartInfo.Arguments = TargetFileName;
+			StartInfo.Arguments = "--strip-debug " + TargetFileName;
 			StartInfo.UseShellExecute = false;
 			StartInfo.CreateNoWindow = true;
 			Utils.RunLocalProcessAndLogOutput(StartInfo);
