Linux: only strip debug info and not symbols (UE-15001).

- The same change as in CL 2539603.

diff --git a/Engine/Source/Programs/UnrealBuildTool/Linux/LinuxToolChain.cs b/Engine/Source/Programs/UnrealBuildTool/Linux/LinuxToolChain.cs
--- a/Engine/Source/Programs/UnrealBuildTool/Linux/LinuxToolChain.cs
+++ b/Engine/Source/Programs/UnrealBuildTool/Linux/LinuxToolChain.cs
@@ -1138,7 +1138,7 @@ namespace UnrealBuildTool
 
 			ProcessStartInfo StartInfo = new ProcessStartInfo();
 			StartInfo.FileName = GetStripPath(UEBuildPlatform.GetBuildPlatform(UnrealTargetPlatform.Linux).GetActiveArchitecture());
-			StartInfo.Arguments = TargetFileName;
+			StartInfo.Arguments = "--strip-debug " + TargetFileName;
 			StartInfo.UseShellExecute = false;
 			StartInfo.CreateNoWindow = true;
 			Utils.RunLocalProcessAndLogOutput(StartInfo);
