Fixed Visual C++ project files not being compatible with VS 2015 ("-2015" option)
- We now emit ConfigurationType properties all of the possible valid project and configuration combinations for the project
- This fixes Visual Studio 2015 from reporting the project files as "Visual Studio 2010" files and offering to upgrade them
- This also results in "correct" Visual C++ project files for VS 2013 as well

diff --git a/Engine/Source/Programs/UnrealBuildTool/System/VCProject.cs b/Engine/Source/Programs/UnrealBuildTool/System/VCProject.cs
--- a/Engine/Source/Programs/UnrealBuildTool/System/VCProject.cs
+++ b/Engine/Source/Programs/UnrealBuildTool/System/VCProject.cs
@@ -436,32 +436,34 @@ namespace UnrealBuildTool
 
 			// Make a list of the platforms and configs as project-format names
 			var ProjectPlatforms = new List<UnrealTargetPlatform>();
-			var ProjectPlatformNames = new List<string>();
-			var ProjectConfigNames = new List<string>();
+			var ProjectPlatformNameAndPlatforms = new List<Tuple<string, UnrealTargetPlatform>>();	// ProjectPlatformName, Platform
+			var ProjectConfigurationNameAndConfigurations = new List<Tuple<string, UnrealTargetConfiguration>>();	// ProjectConfigurationName, Configuration
 			foreach ( var Combination in ProjectConfigAndTargetCombinations )
 			{
 				if( !ProjectPlatforms.Contains( Combination.Platform ) )
 				{
 					ProjectPlatforms.Add( Combination.Platform );
 				}
-				if( !ProjectPlatformNames.Contains( Combination.ProjectPlatformName ) )
+				if( !ProjectPlatformNameAndPlatforms.Any( ProjectPlatformNameAndPlatformTuple => ProjectPlatformNameAndPlatformTuple.Item1 == Combination.ProjectPlatformName ) )
 				{
-					ProjectPlatformNames.Add(Combination.ProjectPlatformName);
+					ProjectPlatformNameAndPlatforms.Add( Tuple.Create( Combination.ProjectPlatformName, Combination.Platform ) );
 				}
-				if( !ProjectConfigNames.Contains( Combination.ProjectConfigurationName ) )
+				if( !ProjectConfigurationNameAndConfigurations.Any( ProjectConfigurationNameAndConfigurationTuple => ProjectConfigurationNameAndConfigurationTuple.Item1 == Combination.ProjectConfigurationName ) )
 				{
-					ProjectConfigNames.Add(Combination.ProjectConfigurationName);
+					ProjectConfigurationNameAndConfigurations.Add( Tuple.Create( Combination.ProjectConfigurationName, Combination.Configuration ) ); 
 				}
-			}
+            }
 
 			// Output ALL the project's config-platform permutations (project files MUST do this)
-			foreach (var ProjectConfigName in ProjectConfigNames)
+			foreach( var ConfigurationTuple in ProjectConfigurationNameAndConfigurations )
 			{
-				foreach (var ProjectPlatformName in ProjectPlatformNames)
+				var ProjectConfigurationName = ConfigurationTuple.Item1;
+				foreach( var PlatformTuple in ProjectPlatformNameAndPlatforms )
 				{
-					VCProjectFileContent.Append(
-							"		<ProjectConfiguration Include=\"" + ProjectConfigName + "|" + ProjectPlatformName + "\">" + ProjectFileGenerator.NewLine +
-							"			<Configuration>" + ProjectConfigName + "</Configuration>" + ProjectFileGenerator.NewLine +
+					var ProjectPlatformName = PlatformTuple.Item1;
+                    VCProjectFileContent.Append(
+							"		<ProjectConfiguration Include=\"" + ProjectConfigurationName + "|" + ProjectPlatformName + "\">" + ProjectFileGenerator.NewLine +
+							"			<Configuration>" + ProjectConfigurationName + "</Configuration>" + ProjectFileGenerator.NewLine +
 							"			<Platform>" + ProjectPlatformName + "</Platform>" + ProjectFileGenerator.NewLine +
 							"		</ProjectConfiguration>" + ProjectFileGenerator.NewLine);
 				}              
@@ -574,9 +576,16 @@ namespace UnrealBuildTool
 			}
 
 			// Write each project configuration PreDefaultProps section
-			foreach( var Combination in ProjectConfigAndTargetCombinations )
+			foreach (var ConfigurationTuple in ProjectConfigurationNameAndConfigurations)
 			{
-				WritePreDefaultPropsConfiguration( Combination, VCProjectFileContent );
+				var ProjectConfigurationName = ConfigurationTuple.Item1;
+				var TargetConfiguration = ConfigurationTuple.Item2;
+				foreach (var PlatformTuple in ProjectPlatformNameAndPlatforms)
+				{
+					var ProjectPlatformName = PlatformTuple.Item1;
+					var TargetPlatform = PlatformTuple.Item2;
+					WritePreDefaultPropsConfiguration( TargetPlatform, TargetConfiguration, ProjectPlatformName, ProjectConfigurationName, VCProjectFileContent );
+				}
 			}
 
 			VCProjectFileContent.Append(
@@ -833,23 +842,24 @@ namespace UnrealBuildTool
 		}
 
 		// Anonymous function that writes pre-Default.props configuration data
-		private void WritePreDefaultPropsConfiguration(ProjectConfigAndTargetCombination Combination, StringBuilder VCProjectFileContent)
+		private void WritePreDefaultPropsConfiguration(UnrealTargetPlatform TargetPlatform, UnrealTargetConfiguration TargetConfiguration, string ProjectPlatformName, string ProjectConfigurationName, StringBuilder VCProjectFileContent)
 		{
-			UEPlatformProjectGenerator ProjGenerator = UEPlatformProjectGenerator.GetPlatformProjectGenerator(Combination.Platform, true);
-			if (((ProjGenerator == null) && (Combination.Platform != UnrealTargetPlatform.Unknown)))
+			UEPlatformProjectGenerator ProjGenerator = UEPlatformProjectGenerator.GetPlatformProjectGenerator(TargetPlatform, true);
+			if (((ProjGenerator == null) && (TargetPlatform != UnrealTargetPlatform.Unknown)))
 			{
 				return;
 			}
 
-			string ConditionString = "Condition=\"'$(Configuration)|$(Platform)'=='" + Combination.ProjectConfigurationAndPlatformName + "'\"";
+			var ProjectConfigurationAndPlatformName = ProjectConfigurationName + "|" + ProjectPlatformName;
+            string ConditionString = "Condition=\"'$(Configuration)|$(Platform)'=='" + ProjectConfigurationAndPlatformName + "'\"";
 
-			string PlatformToolsetString = (ProjGenerator != null) ? ProjGenerator.GetVisualStudioPlatformToolsetString(Combination.Platform, Combination.Configuration, this) : "";
+			string PlatformToolsetString = (ProjGenerator != null) ? ProjGenerator.GetVisualStudioPlatformToolsetString(TargetPlatform, TargetConfiguration, this) : "";
 			if( String.IsNullOrEmpty( PlatformToolsetString ) )
 			{
 				PlatformToolsetString = "		<PlatformToolset>" + VCProjectFileGenerator.ProjectFilePlatformToolsetVersionString + "</PlatformToolset>" + ProjectFileGenerator.NewLine;
 			}
 
-			string PlatformConfigurationType = (ProjGenerator == null)? "Makefile" : ProjGenerator.GetVisualStudioPlatformConfigurationType(Combination.Platform);	
+			string PlatformConfigurationType = (ProjGenerator == null)? "Makefile" : ProjGenerator.GetVisualStudioPlatformConfigurationType(TargetPlatform);	
 			VCProjectFileContent.Append(
 				"	<PropertyGroup " + ConditionString + " Label=\"Configuration\">" + ProjectFileGenerator.NewLine +
 				"		<ConfigurationType>" + PlatformConfigurationType + "</ConfigurationType>" + ProjectFileGenerator.NewLine +
