Fix issue deploying games with invalid characters in names or ABCD style guids.
Added some default values for project setup.
Added type support to default value macros for manifest settings.
Added a macro for manifest generation that filters all non-alphanumeric or dot characters (this is a common requirement for manifest fields).
Filtered game name with alphanumericdot and productid with guid type to support projects that didn't use the same format as some others.

diff --git a/Engine/Config/BaseGame.ini b/Engine/Config/BaseGame.ini
--- a/Engine/Config/BaseGame.ini
+++ b/Engine/Config/BaseGame.ini
@@ -26,7 +26,8 @@ CopyrightNotice=Fill out your copyright notice in the Description page of Projec
 Description=
 LicensingTerms=
 PrivacyPolicy=
-ProjectVersion=
+ProjectName=UE4Game
+ProjectVersion=1.0.0.0
 Homepage=
 SupportContact=
 
diff --git a/Engine/Source/Programs/UnrealBuildTool/Configuration/EngineConfiguration.cs b/Engine/Source/Programs/UnrealBuildTool/Configuration/EngineConfiguration.cs
--- a/Engine/Source/Programs/UnrealBuildTool/Configuration/EngineConfiguration.cs
+++ b/Engine/Source/Programs/UnrealBuildTool/Configuration/EngineConfiguration.cs
@@ -219,6 +219,53 @@ namespace UnrealBuildTool
 			return Result;
 		}
 
+		/// <summary>
+		/// Gets a single GUID value associated with the specified key.
+		/// </summary>
+		/// <param name="SectionName">Section name</param>
+		/// <param name="Key">Key name</param>
+		/// <param name="Value">Value associated with the specified key. If the key has more than one value, only the first one is returned</param>
+		/// <returns>True if the key exists</returns>
+		public bool GetGUID(string SectionName, string Key, out Guid Value)
+		{
+			Value = Guid.Empty;
+			string TextValue;
+			bool Result = GetString(SectionName, Key, out TextValue);
+			if (Result)
+			{
+                string HexString = "";
+                if (TextValue.Contains("A=") && TextValue.Contains("B=") && TextValue.Contains("C=") && TextValue.Contains("D="))
+                {
+                    char[] Separators = new char[] { '(', ')', '=', ',', ' ', 'A', 'B', 'C', 'D' };
+                    string[] ComponentValues = TextValue.Split(Separators, StringSplitOptions.RemoveEmptyEntries);
+                    if (ComponentValues.Length == 4)
+                    {
+                        for (int ComponentIndex = 0; ComponentIndex < 4; ComponentIndex++)
+                        {
+                            int IntegerValue;
+                            Result &= Int32.TryParse(ComponentValues[ComponentIndex], out IntegerValue);
+                            HexString += IntegerValue.ToString("X8");
+                        }
+                    }
+                }
+                else
+                {
+                    HexString = TextValue;
+                }
+
+                try
+                {
+                    Value = Guid.ParseExact(HexString, "N");
+                    Result = true;
+                }
+                catch (Exception)
+                {
+                    Result = false;
+                }
+			}
+			return Result;
+		}
+
 		/// <summary>
 		/// Gets a single float value associated with the specified key.
 		/// </summary>
