UBT: Add convenience targets to Linux makefiles.

- Makes "all" target actually usable.

diff --git a/Engine/Source/Programs/UnrealBuildTool/System/MakefileGenerator.cs b/Engine/Source/Programs/UnrealBuildTool/System/MakefileGenerator.cs
--- a/Engine/Source/Programs/UnrealBuildTool/System/MakefileGenerator.cs
+++ b/Engine/Source/Programs/UnrealBuildTool/System/MakefileGenerator.cs
@@ -108,7 +108,10 @@ namespace UnrealBuildTool
             MakefileContent.Append("\\\n\tconfigure");
 
             MakefileContent.Append("\n\n" + BuildCommand + ProjectBuildCommand + "\n" +
-                "all: $(TARGETS)\n"
+                "all: StandardSet\n\n" +
+                "RequiredTools: CrashReportClient ShaderCompileWorker UnrealPak UnrealLightmass\n\n" +
+                "StandardSet: RequiredTools UnrealFrontend UE4Editor\n\n" +
+                "DebugSet: RequiredTools UnrealFrontend-Linux-Debug UE4Editor-Linux-Debug\n\n"
             );
 
             foreach (string Target in DiscoverTargets())
