Added -logdetailedcompilertiminginfo for timing VC local builds.

diff --git a/Engine/Source/Programs/UnrealBuildTool/System/UnrealBuildTool.cs b/Engine/Source/Programs/UnrealBuildTool/System/UnrealBuildTool.cs
--- a/Engine/Source/Programs/UnrealBuildTool/System/UnrealBuildTool.cs
+++ b/Engine/Source/Programs/UnrealBuildTool/System/UnrealBuildTool.cs
@@ -681,6 +681,10 @@ namespace UnrealBuildTool
                 {
                     BuildConfiguration.LogFilename = LowercaseArg.Replace("-log=","");
                 }
+                else if (LowercaseArg.StartsWith("-logdetailedcompilertiminginfo="))
+                {
+                    WindowsPlatform.bLogDetailedCompilerTimingInfo = bool.Parse(LowercaseArg.Replace("-logdetailedcompilertiminginfo=",""));
+                }
                 else if (LowercaseArg == "-forceheadergeneration")
                 {
                     UEBuildConfiguration.bForceHeaderGeneration = true;
