Fixed regression when using '-2012' option to generate project files
- The VS 2012 compiler is not supported on Windows platform anymore, but using '-2012' was causing that toolchain to be invoked
- Now, this will once again generate project files for VS 2012 but compile using VS 2013

diff --git a/Engine/Source/Programs/UnrealBuildTool/Windows/UEBuildWindows.cs b/Engine/Source/Programs/UnrealBuildTool/Windows/UEBuildWindows.cs
--- a/Engine/Source/Programs/UnrealBuildTool/Windows/UEBuildWindows.cs
+++ b/Engine/Source/Programs/UnrealBuildTool/Windows/UEBuildWindows.cs
@@ -45,9 +45,10 @@ namespace UnrealBuildTool
 				// First, default based on whether there is a command line override...
 				if (UnrealBuildTool.CommandLineContains("-2012"))
 				{
-					CachedCompiler = WindowsCompiler.VisualStudio2012;
+					// We don't support compiling with VS 2012 on Windows platform, but you can still generate project files that are 2012-compatible.  That's handled elsewhere
 				}
-				else if (UnrealBuildTool.CommandLineContains("-2013"))
+
+				if (UnrealBuildTool.CommandLineContains("-2013"))
 				{
 					CachedCompiler = WindowsCompiler.VisualStudio2013;
 				}
@@ -67,10 +68,6 @@ namespace UnrealBuildTool
 				{
 					CachedCompiler = WindowsCompiler.VisualStudio2013;
 				}
-				else if (!String.IsNullOrEmpty(WindowsPlatform.GetVSComnToolsPath(WindowsCompiler.VisualStudio2012)))
-				{
-					CachedCompiler = WindowsCompiler.VisualStudio2012;
-				}
 				else
 				{
 					// Finally assume 2013 is installed to defer errors somewhere else like VCToolChain
