fix for mismatch executable name when there is an underscore in the project name
UE-15222
#ios

diff --git a/Engine/Source/Programs/UnrealBuildTool/IOS/UEDeployIOS.cs b/Engine/Source/Programs/UnrealBuildTool/IOS/UEDeployIOS.cs
--- a/Engine/Source/Programs/UnrealBuildTool/IOS/UEDeployIOS.cs
+++ b/Engine/Source/Programs/UnrealBuildTool/IOS/UEDeployIOS.cs
@@ -283,7 +283,7 @@ namespace UnrealBuildTool.IOS
 			Text.AppendLine("\t<key>CFBundleDisplayName</key>");
 			Text.AppendLine(string.Format("\t<string>{0}</string>", BundleDisplayName.Replace("[PROJECT_NAME]", ProjectName).Replace("_","")));
 			Text.AppendLine("\t<key>CFBundleExecutable</key>");
-			Text.AppendLine(string.Format("\t<string>{0}</string>", bIsUE4Game ? "UE4Game" : GameName.Replace("_", "")));
+			Text.AppendLine(string.Format("\t<string>{0}</string>", bIsUE4Game ? "UE4Game" : GameName));
 			Text.AppendLine("\t<key>CFBundleIdentifier</key>");
 			Text.AppendLine(string.Format("\t<string>{0}</string>", BundleIdentifier.Replace("[PROJECT_NAME]", ProjectName).Replace("_","")));
 			Text.AppendLine("\t<key>CFBundleInfoDictionaryVersion</key>");
