UBT: Fix generation of Manifest.xml path when we don't have a project directory.

diff --git a/Engine/Source/Programs/UnrealBuildTool/Configuration/UEBuildTarget.cs b/Engine/Source/Programs/UnrealBuildTool/Configuration/UEBuildTarget.cs
--- a/Engine/Source/Programs/UnrealBuildTool/Configuration/UEBuildTarget.cs
+++ b/Engine/Source/Programs/UnrealBuildTool/Configuration/UEBuildTarget.cs
@@ -1314,7 +1314,7 @@ namespace UnrealBuildTool
         public void GenerateManifest()
 		{
 			string ManifestPath;
-			if (UnrealBuildTool.RunningRocket())
+			if (UnrealBuildTool.RunningRocket() && UnrealBuildTool.HasUProjectFile())
 			{
 				ManifestPath = Path.Combine(UnrealBuildTool.GetUProjectPath(), BuildConfiguration.BaseIntermediateFolder, "Manifest.xml");
 			}
