UEBuildTarget stores its .Target.cs filename.
UEBuildModule stores its .Module.cs filename.
UEBuildTarget stores a map of dictionaries from module names to .Target.cs files.
Timestamps of serialized .Target.cs and .Module.cs files compared against makefile and used for invalidation.

More sensible Dictionary logic in ConditionallyCompileAndLoadRulesAssembly.
InstantiateModule can no longer return a null.
UEBuildModule constructors no longer default.

diff --git a/Engine/Source/Programs/UnrealBuildTool/Configuration/UEBuildClient.cs b/Engine/Source/Programs/UnrealBuildTool/Configuration/UEBuildClient.cs
--- a/Engine/Source/Programs/UnrealBuildTool/Configuration/UEBuildClient.cs
+++ b/Engine/Source/Programs/UnrealBuildTool/Configuration/UEBuildClient.cs
@@ -10,8 +10,8 @@ namespace UnrealBuildTool
 	[Serializable]
     public class UEBuildClient : UEBuildTarget
 	{
-		public UEBuildClient(TargetDescriptor InDesc, TargetRules InRulesObject)
-            : base(InDesc, InRulesObject, "UE4Client")
+		public UEBuildClient(TargetDescriptor InDesc, TargetRules InRulesObject, string InTargetCsFilename)
+            : base(InDesc, InRulesObject, "UE4Client", InTargetCsFilename)
         {
             if (ShouldCompileMonolithic())
             {
diff --git a/Engine/Source/Programs/UnrealBuildTool/Configuration/UEBuildEditor.cs b/Engine/Source/Programs/UnrealBuildTool/Configuration/UEBuildEditor.cs
--- a/Engine/Source/Programs/UnrealBuildTool/Configuration/UEBuildEditor.cs
+++ b/Engine/Source/Programs/UnrealBuildTool/Configuration/UEBuildEditor.cs
@@ -10,8 +10,8 @@ namespace UnrealBuildTool
 	[Serializable]
 	public class UEBuildEditor : UEBuildTarget
 	{
-		public UEBuildEditor(TargetDescriptor InDesc, TargetRules InRulesObject)
-			: base(InDesc, InRulesObject, "UE4Editor")
+		public UEBuildEditor(TargetDescriptor InDesc, TargetRules InRulesObject, string InTargetCsFilename)
+			: base(InDesc, InRulesObject, "UE4Editor", InTargetCsFilename)
 		{
 		}
 
diff --git a/Engine/Source/Programs/UnrealBuildTool/Configuration/UEBuildGame.cs b/Engine/Source/Programs/UnrealBuildTool/Configuration/UEBuildGame.cs
--- a/Engine/Source/Programs/UnrealBuildTool/Configuration/UEBuildGame.cs
+++ b/Engine/Source/Programs/UnrealBuildTool/Configuration/UEBuildGame.cs
@@ -10,8 +10,8 @@ namespace UnrealBuildTool
 	[Serializable]
 	public class UEBuildGame : UEBuildTarget
 	{
-		public UEBuildGame(TargetDescriptor InDesc, TargetRules InRulesObject)
-			: base(InDesc, InRulesObject, "UE4")
+		public UEBuildGame(TargetDescriptor InDesc, TargetRules InRulesObject, string InTargetCsFilename)
+			: base(InDesc, InRulesObject, "UE4", InTargetCsFilename)
 		{
 			if (ShouldCompileMonolithic())
 			{
diff --git a/Engine/Source/Programs/UnrealBuildTool/Configuration/UEBuildModule.cs b/Engine/Source/Programs/UnrealBuildTool/Configuration/UEBuildModule.cs
--- a/Engine/Source/Programs/UnrealBuildTool/Configuration/UEBuildModule.cs
+++ b/Engine/Source/Programs/UnrealBuildTool/Configuration/UEBuildModule.cs
@@ -177,6 +177,10 @@ namespace UnrealBuildTool
 		/** Is this module allowed to be redistributed. */
 		private readonly bool? IsRedistributableOverride;
 
+		/** The name of the .Build.cs file this module was created from, if any */
+		private readonly string BuildCsFilenameField;
+		public string BuildCsFilename { get { return BuildCsFilenameField; } }
+
 		/**
 		 * Tells if this module can be redistributed.
 		 * 
@@ -244,26 +248,27 @@ namespace UnrealBuildTool
 			UEBuildModuleType InType,
 			string InModuleDirectory,
 			bool? InIsRedistributableOverride,
-			IEnumerable<string> InPublicDefinitions = null,
-			IEnumerable<string> InPublicIncludePaths = null,
-			IEnumerable<string> InPublicSystemIncludePaths = null,
-			IEnumerable<string> InPublicLibraryPaths = null,
-			IEnumerable<string> InPublicAdditionalLibraries = null,
-			IEnumerable<string> InPublicFrameworks = null,
-			IEnumerable<string> InPublicWeakFrameworks = null,
-			IEnumerable<UEBuildFramework> InPublicAdditionalFrameworks = null,
-			IEnumerable<string> InPublicAdditionalShadowFiles = null,
-			IEnumerable<UEBuildBundleResource> InPublicAdditionalBundleResources = null,
-			IEnumerable<string> InPublicIncludePathModuleNames = null,
-			IEnumerable<string> InPublicDependencyModuleNames = null,
-			IEnumerable<string> InPublicDelayLoadDLLs = null,
-			IEnumerable<string> InPrivateIncludePaths = null,
-			IEnumerable<string> InPrivateIncludePathModuleNames = null,
-			IEnumerable<string> InPrivateDependencyModuleNames = null,
-			IEnumerable<string> InCircularlyReferencedDependentModules = null,
-			IEnumerable<string> InDynamicallyLoadedModuleNames = null,
-            IEnumerable<string> InPlatformSpecificDynamicallyLoadedModuleNames = null,
-			IEnumerable<RuntimeDependency> InRuntimeDependencies = null
+			IEnumerable<string> InPublicDefinitions,
+			IEnumerable<string> InPublicIncludePaths,
+			IEnumerable<string> InPublicSystemIncludePaths,
+			IEnumerable<string> InPublicLibraryPaths,
+			IEnumerable<string> InPublicAdditionalLibraries,
+			IEnumerable<string> InPublicFrameworks,
+			IEnumerable<string> InPublicWeakFrameworks,
+			IEnumerable<UEBuildFramework> InPublicAdditionalFrameworks,
+			IEnumerable<string> InPublicAdditionalShadowFiles,
+			IEnumerable<UEBuildBundleResource> InPublicAdditionalBundleResources,
+			IEnumerable<string> InPublicIncludePathModuleNames,
+			IEnumerable<string> InPublicDependencyModuleNames,
+			IEnumerable<string> InPublicDelayLoadDLLs,
+			IEnumerable<string> InPrivateIncludePaths,
+			IEnumerable<string> InPrivateIncludePathModuleNames,
+			IEnumerable<string> InPrivateDependencyModuleNames,
+			IEnumerable<string> InCircularlyReferencedDependentModules,
+			IEnumerable<string> InDynamicallyLoadedModuleNames,
+            IEnumerable<string> InPlatformSpecificDynamicallyLoadedModuleNames,
+			IEnumerable<RuntimeDependency> InRuntimeDependencies,
+			string InBuildCsFilename
 			)
 		{
 			Target = InTarget;
@@ -292,6 +297,9 @@ namespace UnrealBuildTool
 			RuntimeDependencies = (InRuntimeDependencies == null)? new List<RuntimeDependency>() : new List<RuntimeDependency>(InRuntimeDependencies);
 			IsRedistributableOverride = InIsRedistributableOverride;
 
+			Debug.Assert(InBuildCsFilename == null || InBuildCsFilename.EndsWith(".Build.cs", StringComparison.InvariantCultureIgnoreCase));
+			BuildCsFilenameField = InBuildCsFilename;
+
 			SetupModuleDistributionLevel();
 
 			Target.RegisterModule(this);
@@ -882,39 +890,48 @@ namespace UnrealBuildTool
 			string InName,
 			string InModuleDirectory,
 			bool? InIsRedistributableOverride,
-			IEnumerable<string> InPublicDefinitions = null,
-			IEnumerable<string> InPublicIncludePaths = null,
-			IEnumerable<string> InPublicSystemIncludePaths = null,
-			IEnumerable<string> InPublicLibraryPaths = null,
-			IEnumerable<string> InPublicAdditionalLibraries = null,
-			IEnumerable<string> InPublicFrameworks = null,
-			IEnumerable<string> InPublicWeakFrameworks = null,
-			IEnumerable<UEBuildFramework> InPublicAdditionalFrameworks = null,
-			IEnumerable<string> InPublicAdditionalShadowFiles = null,
-			IEnumerable<UEBuildBundleResource> InPublicAdditionalBundleResources = null,
-			IEnumerable<string> InPublicDependencyModuleNames = null,
-			IEnumerable<string> InPublicDelayLoadDLLs = null,		// Delay loaded DLLs that should be setup when including this module
-			IEnumerable<RuntimeDependency> InRuntimeDependencies = null
+			IEnumerable<string> InPublicDefinitions,
+			IEnumerable<string> InPublicIncludePaths,
+			IEnumerable<string> InPublicSystemIncludePaths,
+			IEnumerable<string> InPublicLibraryPaths,
+			IEnumerable<string> InPublicAdditionalLibraries,
+			IEnumerable<string> InPublicFrameworks,
+			IEnumerable<string> InPublicWeakFrameworks,
+			IEnumerable<UEBuildFramework> InPublicAdditionalFrameworks,
+			IEnumerable<string> InPublicAdditionalShadowFiles,
+			IEnumerable<UEBuildBundleResource> InPublicAdditionalBundleResources,
+			IEnumerable<string> InPublicDependencyModuleNames,
+			IEnumerable<string> InPublicDelayLoadDLLs,		// Delay loaded DLLs that should be setup when including this module
+			IEnumerable<RuntimeDependency> InRuntimeDependencies,
+			string InBuildCsFilename
 			)
 		: base(
-			InTarget:						InTarget,
-			InType:							InType,
-			InName:							InName,
-			InModuleDirectory:				InModuleDirectory,
-			InIsRedistributableOverride:	InIsRedistributableOverride,
-			InPublicDefinitions:			InPublicDefinitions,
-			InPublicIncludePaths:			InPublicIncludePaths,
-			InPublicSystemIncludePaths:		InPublicSystemIncludePaths,
-			InPublicLibraryPaths:			InPublicLibraryPaths,
-			InPublicAdditionalLibraries:	InPublicAdditionalLibraries,
-			InPublicFrameworks:             InPublicFrameworks,
-			InPublicWeakFrameworks:			InPublicWeakFrameworks,
-			InPublicAdditionalFrameworks: 	InPublicAdditionalFrameworks,
-			InPublicAdditionalShadowFiles:  InPublicAdditionalShadowFiles,
-			InPublicAdditionalBundleResources: InPublicAdditionalBundleResources,
-			InPublicDependencyModuleNames:	InPublicDependencyModuleNames,
-			InPublicDelayLoadDLLs:			InPublicDelayLoadDLLs,
-			InRuntimeDependencies:			InRuntimeDependencies
+			InTarget:										InTarget,
+			InType:											InType,
+			InName:											InName,
+			InModuleDirectory:								InModuleDirectory,
+			InIsRedistributableOverride:					InIsRedistributableOverride,
+			InPublicDefinitions:							InPublicDefinitions,
+			InPublicIncludePaths:							InPublicIncludePaths,
+			InPublicSystemIncludePaths:						InPublicSystemIncludePaths,
+			InPublicLibraryPaths:							InPublicLibraryPaths,
+			InPublicAdditionalLibraries:					InPublicAdditionalLibraries,
+			InPublicFrameworks:								InPublicFrameworks,
+			InPublicWeakFrameworks:							InPublicWeakFrameworks,
+			InPublicAdditionalFrameworks:					InPublicAdditionalFrameworks,
+			InPublicAdditionalShadowFiles:					InPublicAdditionalShadowFiles,
+			InPublicAdditionalBundleResources:				InPublicAdditionalBundleResources,
+			InPublicIncludePathModuleNames:					null,
+			InPublicDependencyModuleNames:					InPublicDependencyModuleNames,
+			InPublicDelayLoadDLLs:							InPublicDelayLoadDLLs,
+			InPrivateIncludePaths:							null,
+			InPrivateIncludePathModuleNames:				null,
+			InPrivateDependencyModuleNames:					null,
+			InCircularlyReferencedDependentModules:			null,
+			InDynamicallyLoadedModuleNames:					null,
+			InPlatformSpecificDynamicallyLoadedModuleNames: null,
+			InRuntimeDependencies:							InRuntimeDependencies,
+			InBuildCsFilename:								InBuildCsFilename
 			)
 		{
 			bIncludedInTarget = true;
@@ -1118,7 +1135,8 @@ namespace UnrealBuildTool
 			bool InFasterWithoutUnity,
 			int InMinFilesUsingPrecompiledHeaderOverride,
 			bool InEnableExceptions,
-			bool bInBuildSourceFiles
+			bool bInBuildSourceFiles,
+			string InBuildCsFilename
 			)
 			: base(	InTarget,
 					InName, 
@@ -1144,7 +1162,8 @@ namespace UnrealBuildTool
 					InCircularlyReferencedDependentModules,
 					InDynamicallyLoadedModuleNames,
                     InPlatformSpecificDynamicallyLoadedModuleNames,
-					InRuntimeDependencies
+					InRuntimeDependencies,
+					InBuildCsFilename
 				)
 		{
 			GeneratedCodeDirectory = InGeneratedCodeDirectory;
@@ -2100,7 +2119,8 @@ namespace UnrealBuildTool
 			bool InFasterWithoutUnity,
 			int InMinFilesUsingPrecompiledHeaderOverride,
 			bool InEnableExceptions,
-			bool bInBuildSourceFiles
+			bool bInBuildSourceFiles,
+			string InBuildCsFilename
 			)
 			: base(InTarget,InName,InType,InModuleDirectory,InGeneratedCodeDirectory,InIsRedistributableOverride,InIntelliSenseGatherer,
 			InSourceFiles,InPublicIncludePaths,InPublicSystemIncludePaths,null,InDefinitions,
@@ -2108,7 +2128,7 @@ namespace UnrealBuildTool
 			InPrivateIncludePaths,InPrivateIncludePathModuleNames,InPrivateDependencyModuleNames,
             InCircularlyReferencedDependentModules, InDynamicallyLoadedModuleNames, InPlatformSpecificDynamicallyLoadedModuleNames, InRuntimeDependencies, InOptimizeCode,
 			InAllowSharedPCH, InSharedPCHHeaderFile, InUseRTTI, InEnableBufferSecurityChecks, InFasterWithoutUnity, InMinFilesUsingPrecompiledHeaderOverride,
-			InEnableExceptions, bInBuildSourceFiles)
+			InEnableExceptions, bInBuildSourceFiles, InBuildCsFilename)
 		{
 			PrivateAssemblyReferences = HashSetFromOptionalEnumerableStringParameter(InPrivateAssemblyReferences);
 		}
diff --git a/Engine/Source/Programs/UnrealBuildTool/Configuration/UEBuildServer.cs b/Engine/Source/Programs/UnrealBuildTool/Configuration/UEBuildServer.cs
--- a/Engine/Source/Programs/UnrealBuildTool/Configuration/UEBuildServer.cs
+++ b/Engine/Source/Programs/UnrealBuildTool/Configuration/UEBuildServer.cs
@@ -10,8 +10,8 @@ namespace UnrealBuildTool
 	[Serializable]
 	public class UEBuildServer : UEBuildTarget
 	{
-		public UEBuildServer(TargetDescriptor InDesc, TargetRules InRulesObject)
-			: base(InDesc, InRulesObject, "UE4Server")
+		public UEBuildServer(TargetDescriptor InDesc, TargetRules InRulesObject, string InTargetCsFilename)
+			: base(InDesc, InRulesObject, "UE4Server", InTargetCsFilename)
 		{
 		}
 
diff --git a/Engine/Source/Programs/UnrealBuildTool/Configuration/UEBuildTarget.cs b/Engine/Source/Programs/UnrealBuildTool/Configuration/UEBuildTarget.cs
--- a/Engine/Source/Programs/UnrealBuildTool/Configuration/UEBuildTarget.cs
+++ b/Engine/Source/Programs/UnrealBuildTool/Configuration/UEBuildTarget.cs
@@ -187,6 +187,11 @@ namespace UnrealBuildTool
 			{
 				OnlyModules = new List<OnlyModule>();
 			}
+
+			if (ModuleCsMap == null)
+			{
+				ModuleCsMap = new Dictionary<string,string>();
+			}
 		}
 
 		public string GetAppName()
@@ -772,6 +777,9 @@ namespace UnrealBuildTool
 		[NonSerialized]
 		private Dictionary<string, UEBuildModule> Modules = new Dictionary<string, UEBuildModule>(StringComparer.InvariantCultureIgnoreCase);
 
+		/** Used to map names of modules to their .Build.cs filename */
+		private Dictionary<string, string> ModuleCsMap = new Dictionary<string, string>(StringComparer.InvariantCultureIgnoreCase);
+
 		/** The receipt for this target, which contains a record of this build. */
 		private BuildReceipt Receipt;
 
@@ -779,6 +787,19 @@ namespace UnrealBuildTool
 		[NonSerialized]
 		private string ForceReceiptFileName;
 
+		/** The name of the .Target.cs file, if the target was created with one */
+		private readonly string TargetCsFilenameField;
+		public string TargetCsFilename { get { return TargetCsFilenameField; } }
+
+		/// <summary>
+		/// A list of the module filenames which were used to build this target.
+		/// </summary>
+		/// <returns></returns>
+		public List<string> GetAllModuleBuildCsFilenames()
+		{
+			return new List<string>(ModuleCsMap.Values);
+		}
+
 		/// <summary>
 		/// Whether this target should be compiled in monolithic mode
 		/// </summary>
@@ -794,7 +815,8 @@ namespace UnrealBuildTool
 		/// <param name="InDesc">Target descriptor</param>
 		/// <param name="InRules">The target rules, as created by RulesCompiler.</param>
 		/// <param name="InPossibleAppName">The AppName for shared binaries of this target type, if used (null if there is none).</param>
-		public UEBuildTarget(TargetDescriptor InDesc, TargetRules InRules, string InPossibleAppName)
+		/// <param name="InTargetCsFilename">The name of the target </param>
+		public UEBuildTarget(TargetDescriptor InDesc, TargetRules InRules, string InPossibleAppName, string InTargetCsFilename)
 		{
 			AppName = InDesc.TargetName;
 			TargetName = InDesc.TargetName;
@@ -808,6 +830,9 @@ namespace UnrealBuildTool
 			ForeignPlugins = InDesc.ForeignPlugins;
 			ForceReceiptFileName = InDesc.ForceReceiptFileName;
 
+			Debug.Assert(InTargetCsFilename == null || InTargetCsFilename.EndsWith(".Target.cs", StringComparison.InvariantCultureIgnoreCase));
+			TargetCsFilenameField = InTargetCsFilename;
+
 			{
 				bCompileMonolithic = Rules.ShouldCompileMonolithic(InDesc.Platform, InDesc.Configuration);
 
@@ -2105,7 +2130,8 @@ namespace UnrealBuildTool
 				InFasterWithoutUnity: true,
 				InMinFilesUsingPrecompiledHeaderOverride: 0,
 				InEnableExceptions: false,
-				bInBuildSourceFiles: true);
+				bInBuildSourceFiles: true,
+				InBuildCsFilename: null);
 		}
 
 
@@ -2888,6 +2914,7 @@ namespace UnrealBuildTool
 		{
 			Debug.Assert(Module.Target == this);
 			Modules.Add(Module.Name, Module);
+			ModuleCsMap.Add(Module.Name, Module.BuildCsFilename);
 		}
 
 		/** Finds a module given its name.  Throws an exception if the module couldn't be found. */
@@ -3037,11 +3064,7 @@ namespace UnrealBuildTool
 				}
 
 				// Now, go ahead and create the module builder instance
-				Module = InstantiateModule(RulesObject, ModuleName, ModuleType, ModuleDirectory, GeneratedCodeDirectory, IntelliSenseGatherer, FoundSourceFiles, bBuildFiles);
-				if(Module == null)
-				{
-					throw new BuildException("Unrecognized module type specified by 'Rules' object {0}", RulesObject.ToString());
-				}
+				Module = InstantiateModule(RulesObject, ModuleName, ModuleType, ModuleDirectory, GeneratedCodeDirectory, IntelliSenseGatherer, FoundSourceFiles, bBuildFiles, ModuleFileName);
 
 				UnrealTargetPlatform Only = UnrealBuildTool.GetOnlyPlatformSpecificFor();
 
@@ -3063,10 +3086,11 @@ namespace UnrealBuildTool
 			string               ModuleName,
 			UEBuildModuleType    ModuleType,
 			string               ModuleDirectory,
-			string				 GeneratedCodeDirectory,
+			string               GeneratedCodeDirectory,
 			IntelliSenseGatherer IntelliSenseGatherer,
 			List<FileItem>       ModuleSourceFiles,
-			bool                 bBuildSourceFiles)
+			bool                 bBuildSourceFiles,
+			string               InBuildCsFile)
 		{
 			switch (RulesObject.Type)
 			{
@@ -3108,7 +3132,8 @@ namespace UnrealBuildTool
 							InFasterWithoutUnity: RulesObject.bFasterWithoutUnity,
 							InMinFilesUsingPrecompiledHeaderOverride: RulesObject.MinFilesUsingPrecompiledHeaderOverride,
 							InEnableExceptions: RulesObject.bEnableExceptions,
-							bInBuildSourceFiles: bBuildSourceFiles
+							bInBuildSourceFiles: bBuildSourceFiles,
+							InBuildCsFilename: InBuildCsFile
 						);
 
 				case ModuleRules.ModuleType.CPlusPlusCLR:
@@ -3149,7 +3174,8 @@ namespace UnrealBuildTool
 							InFasterWithoutUnity: RulesObject.bFasterWithoutUnity,
 							InMinFilesUsingPrecompiledHeaderOverride: RulesObject.MinFilesUsingPrecompiledHeaderOverride,
 							InEnableExceptions: RulesObject.bEnableExceptions,
-							bInBuildSourceFiles : bBuildSourceFiles
+							bInBuildSourceFiles : bBuildSourceFiles,
+							InBuildCsFilename: InBuildCsFile
 						);
 
 				case ModuleRules.ModuleType.External:
@@ -3171,11 +3197,12 @@ namespace UnrealBuildTool
 							InPublicAdditionalBundleResources: RulesObject.AdditionalBundleResources,
 							InPublicDependencyModuleNames: RulesObject.PublicDependencyModuleNames,
 							InPublicDelayLoadDLLs: RulesObject.PublicDelayLoadDLLs,
-							InRuntimeDependencies: RulesObject.RuntimeDependencies
+							InRuntimeDependencies: RulesObject.RuntimeDependencies,
+							InBuildCsFilename: InBuildCsFile
 						);
 
 				default:
-					return null;
+					throw new BuildException("Unrecognized module type specified by 'Rules' object {0}", RulesObject.ToString());
 			}
 		}
 
diff --git a/Engine/Source/Programs/UnrealBuildTool/System/RulesCompiler.cs b/Engine/Source/Programs/UnrealBuildTool/System/RulesCompiler.cs
--- a/Engine/Source/Programs/UnrealBuildTool/System/RulesCompiler.cs
+++ b/Engine/Source/Programs/UnrealBuildTool/System/RulesCompiler.cs
@@ -1274,10 +1274,21 @@ namespace UnrealBuildTool
 		/// Maps target names to their actual xxx.Target.cs file on disk
 		private static Dictionary<string, string> TargetNameToTargetFileMap = new Dictionary<string, string>( StringComparer.InvariantCultureIgnoreCase );
 
-		/// Map of assembly names we've already compiled and loaded to their list of game folders.  This is used to prevent
-		/// trying to recompile the same assembly when ping-ponging between different types of targets
-		private static Dictionary<Assembly, List<string>> LoadedAssemblyMap = new Dictionary<Assembly, List<string>>();
+		private class LoadedAssemblyData
+		{
+			public LoadedAssemblyData(Assembly InAssembly, List<string> InGameFolders)
+			{
+				ExistingAssembly    = InAssembly;
+				ExistingGameFolders = InGameFolders;
+			}
 
+			public Assembly     ExistingAssembly     { get; private set; }
+			public List<string> ExistingGameFolders  { get; private set; }
+		}
+
+		/// Map of assembly names we've already compiled and loaded to their Assembly and list of game folders.  This is used to prevent
+		/// trying to recompile the same assembly when ping-ponging between different types of targets
+		private static Dictionary<string, LoadedAssemblyData> LoadedAssemblyMap = new Dictionary<string, LoadedAssemblyData>(StringComparer.InvariantCultureIgnoreCase);
 
 		private static void ConditionallyCompileAndLoadRulesAssembly()
 		{
@@ -1286,127 +1297,118 @@ namespace UnrealBuildTool
 				throw new BuildException( "Module or target rules data was requested, but not rules assembly name was set yet!" );
 			}
 
-			RulesAssembly = null;
-
 			// Did we already have a RulesAssembly and cached data about rules files and modules?  If so, then we'll
 			// check to see if we need to flush everything and start over.  This can happen if UBT wants to built
 			// different targets in a single invocation, or when generating project files before or after building
 			// a target
-			foreach( var ExistingAssembly in LoadedAssemblyMap.Keys )
+			LoadedAssemblyData LoadedAssembly;
+			if (LoadedAssemblyMap.TryGetValue(AssemblyName, out LoadedAssembly))
 			{
-				var ExistingAssemblyName = Path.GetFileNameWithoutExtension( ExistingAssembly.Location );
-				if( ExistingAssemblyName.Equals( AssemblyName, StringComparison.InvariantCultureIgnoreCase ) )
-				{
-					// We already have this assembly!
-					RulesAssembly = ExistingAssembly;
+				Assembly     ExistingAssembly    = LoadedAssembly.ExistingAssembly;
+				List<string> ExistingGameFolders = LoadedAssembly.ExistingGameFolders;
 
-					var ExistingGameFolders = LoadedAssemblyMap[ ExistingAssembly ];
+				RulesAssembly = ExistingAssembly;
 
-					// Make sure the game folder list wasn't changed since we last compiled this assembly
-					if( ExistingGameFolders != AllGameFolders )	// Quick-check pointers first to avoid iterating
+				// Make sure the game folder list wasn't changed since we last compiled this assembly
+				if( ExistingGameFolders != AllGameFolders )	// Quick-check pointers first to avoid iterating
+				{
+					var AnyGameFoldersDifferent = false;
+					if( ExistingGameFolders.Count != AllGameFolders.Count )
 					{
-						var AnyGameFoldersDifferent = false;
-						if( ExistingGameFolders.Count != AllGameFolders.Count )
-						{
-							AnyGameFoldersDifferent = true;
-						}
-						else
+						AnyGameFoldersDifferent = true;
+					}
+					else
+					{
+						foreach( var NewGameFolder in AllGameFolders )
 						{
-							foreach( var NewGameFolder in AllGameFolders )
+							if( !ExistingGameFolders.Contains( NewGameFolder ) )
 							{
-								if( !ExistingGameFolders.Contains( NewGameFolder ) )
-								{
-									AnyGameFoldersDifferent = true;
-									break;
-								}
-							}
-							foreach( var OldGameFolder in ExistingGameFolders )
-							{
-								if( !AllGameFolders.Contains( OldGameFolder ) )
-								{
-									AnyGameFoldersDifferent = true;
-									break;
-								}
+								AnyGameFoldersDifferent = true;
+								break;
 							}
 						}
-
-						if( AnyGameFoldersDifferent )
+						foreach( var OldGameFolder in ExistingGameFolders )
 						{
-							throw new BuildException( "SetAssemblyNameAndGameFolders() was called with an assembly name that had already been compiled, but with DIFFERENT game folders.  This is not allowed." );
+							if( !AllGameFolders.Contains( OldGameFolder ) )
+							{
+								AnyGameFoldersDifferent = true;
+								break;
+							}
 						}
 					}
 
-					break;
+					if( AnyGameFoldersDifferent )
+					{
+						throw new BuildException( "SetAssemblyNameAndGameFolders() was called with an assembly name that had already been compiled, but with DIFFERENT game folders.  This is not allowed." );
+					}
 				}
+
+				return;
 			}
 
+			RulesAssembly = null;
+
+			var AdditionalSearchPaths = new List<string>();
 
-			// Does anything need to be compiled?
-			if( RulesAssembly == null )
+			if (UnrealBuildTool.HasUProjectFile())
 			{
-				var AdditionalSearchPaths = new List<string>();
-
-				if (UnrealBuildTool.HasUProjectFile())
+				// Add the game project's source folder
+				var ProjectSourceDirectory = Path.Combine( UnrealBuildTool.GetUProjectPath(), "Source" );
+				if( Directory.Exists( ProjectSourceDirectory ) )
 				{
-					// Add the game project's source folder
-					var ProjectSourceDirectory = Path.Combine( UnrealBuildTool.GetUProjectPath(), "Source" );
-					if( Directory.Exists( ProjectSourceDirectory ) )
-					{
-						AdditionalSearchPaths.Add( ProjectSourceDirectory );
-					}
-					// Add the games project's intermediate source folder
-					var ProjectIntermediateSourceDirectory = Path.Combine(UnrealBuildTool.GetUProjectPath(), "Intermediate", "Source");
-					if (Directory.Exists(ProjectIntermediateSourceDirectory))
-					{
-						AdditionalSearchPaths.Add(ProjectIntermediateSourceDirectory);
-					}
+					AdditionalSearchPaths.Add( ProjectSourceDirectory );
 				}
-				var ModuleFileNames = FindAllRulesSourceFiles(RulesFileType.Module, AdditionalSearchPaths);
-
-				var AssemblySourceFiles = new List<string>();
-				AssemblySourceFiles.AddRange( ModuleFileNames );
-				if( AssemblySourceFiles.Count == 0 )
+				// Add the games project's intermediate source folder
+				var ProjectIntermediateSourceDirectory = Path.Combine(UnrealBuildTool.GetUProjectPath(), "Intermediate", "Source");
+				if (Directory.Exists(ProjectIntermediateSourceDirectory))
 				{
-					throw new BuildException("No module rules source files were found in any of the module base directories!");
+					AdditionalSearchPaths.Add(ProjectIntermediateSourceDirectory);
 				}
-				var TargetFileNames = FindAllRulesSourceFiles(RulesFileType.Target, AdditionalSearchPaths);
-				AssemblySourceFiles.AddRange( TargetFileNames );
+			}
+			var ModuleFileNames = FindAllRulesSourceFiles(RulesFileType.Module, AdditionalSearchPaths);
 
-				// Create a path to the assembly that we'll either load or compile
-				string BaseIntermediatePath = UnrealBuildTool.HasUProjectFile() ?
-					Path.Combine(UnrealBuildTool.GetUProjectPath(), BuildConfiguration.BaseIntermediateFolder) : BuildConfiguration.BaseIntermediatePath;
+			var AssemblySourceFiles = new List<string>();
+			AssemblySourceFiles.AddRange( ModuleFileNames );
+			if( AssemblySourceFiles.Count == 0 )
+			{
+				throw new BuildException("No module rules source files were found in any of the module base directories!");
+			}
+			var TargetFileNames = FindAllRulesSourceFiles(RulesFileType.Target, AdditionalSearchPaths);
+			AssemblySourceFiles.AddRange( TargetFileNames );
 
-				string OutputAssemblyPath = Path.GetFullPath(Path.Combine(BaseIntermediatePath, "BuildRules", AssemblyName + ".dll"));
+			// Create a path to the assembly that we'll either load or compile
+			string BaseIntermediatePath = UnrealBuildTool.HasUProjectFile() ?
+				Path.Combine(UnrealBuildTool.GetUProjectPath(), BuildConfiguration.BaseIntermediateFolder) : BuildConfiguration.BaseIntermediatePath;
 
-				RulesAssembly = DynamicCompilation.CompileAndLoadAssembly( OutputAssemblyPath, AssemblySourceFiles );
+			string OutputAssemblyPath = Path.GetFullPath(Path.Combine(BaseIntermediatePath, "BuildRules", AssemblyName + ".dll"));
 
-				{
-					// Setup the module map
-					foreach( var CurModuleFileName in ModuleFileNames )
-					{
-						var CleanFileName = Utils.CleanDirectorySeparators( CurModuleFileName );
-						var ModuleName = Path.GetFileNameWithoutExtension( Path.GetFileNameWithoutExtension( CleanFileName ) );	// Strip both extensions
-						if( !ModuleNameToModuleFileMap.ContainsKey( ModuleName ) )
-						{
-							ModuleNameToModuleFileMap.Add( ModuleName, CurModuleFileName );
-						}
-					}
+			RulesAssembly = DynamicCompilation.CompileAndLoadAssembly( OutputAssemblyPath, AssemblySourceFiles );
 
-					// Setup the target map
-					foreach( var CurTargetFileName in TargetFileNames )
-					{
-						var CleanFileName = Utils.CleanDirectorySeparators( CurTargetFileName );
-						var TargetName = Path.GetFileNameWithoutExtension( Path.GetFileNameWithoutExtension( CleanFileName ) );	// Strip both extensions
-						if( !TargetNameToTargetFileMap.ContainsKey( TargetName ) )
-						{
-							TargetNameToTargetFileMap.Add( TargetName, CurTargetFileName );
-						}
-					}
+			// Setup the module map
+			foreach( var CurModuleFileName in ModuleFileNames )
+			{
+				var CleanFileName = Utils.CleanDirectorySeparators( CurModuleFileName );
+				var ModuleName = Path.GetFileNameWithoutExtension( Path.GetFileNameWithoutExtension( CleanFileName ) );	// Strip both extensions
+				if( !ModuleNameToModuleFileMap.ContainsKey( ModuleName ) )
+				{
+					ModuleNameToModuleFileMap.Add( ModuleName, CurModuleFileName );
 				}
+			}
 
-				// Remember that we loaded this assembly
-				LoadedAssemblyMap[ RulesAssembly ] = AllGameFolders;
+			// Setup the target map
+			foreach( var CurTargetFileName in TargetFileNames )
+			{
+				var CleanFileName = Utils.CleanDirectorySeparators( CurTargetFileName );
+				var TargetName = Path.GetFileNameWithoutExtension( Path.GetFileNameWithoutExtension( CleanFileName ) );	// Strip both extensions
+				if( !TargetNameToTargetFileMap.ContainsKey( TargetName ) )
+				{
+					TargetNameToTargetFileMap.Add( TargetName, CurTargetFileName );
+				}
 			}
+
+			// Remember that we loaded this assembly
+			var RulesAssemblyName = Path.GetFileNameWithoutExtension( RulesAssembly.Location );
+			LoadedAssemblyMap[RulesAssemblyName] = new LoadedAssemblyData(RulesAssembly, AllGameFolders);
 		}
 
 		/// <summary>
@@ -1936,19 +1938,19 @@ namespace UnrealBuildTool
 			switch (RulesObject.Type)
 			{
 				case TargetRules.TargetType.Game:
-					BuildTarget = new UEBuildGame(Desc, RulesObject);
+					BuildTarget = new UEBuildGame(Desc, RulesObject, TargetFileName);
 					break;
 				case TargetRules.TargetType.Editor:
-					BuildTarget = new UEBuildEditor(Desc, RulesObject);
+					BuildTarget = new UEBuildEditor(Desc, RulesObject, TargetFileName);
 					break;
                 case TargetRules.TargetType.Client:
-                    BuildTarget = new UEBuildClient(Desc, RulesObject);
+                    BuildTarget = new UEBuildClient(Desc, RulesObject, TargetFileName);
                     break;
 				case TargetRules.TargetType.Server:
-					BuildTarget = new UEBuildServer(Desc, RulesObject);
+					BuildTarget = new UEBuildServer(Desc, RulesObject, TargetFileName);
 					break;
 				case TargetRules.TargetType.Program:
-					BuildTarget = new UEBuildTarget(Desc, RulesObject, null);
+					BuildTarget = new UEBuildTarget(Desc, RulesObject, null, TargetFileName);
 					break;
 			}
 
diff --git a/Engine/Source/Programs/UnrealBuildTool/System/UnrealBuildTool.cs b/Engine/Source/Programs/UnrealBuildTool/System/UnrealBuildTool.cs
--- a/Engine/Source/Programs/UnrealBuildTool/System/UnrealBuildTool.cs
+++ b/Engine/Source/Programs/UnrealBuildTool/System/UnrealBuildTool.cs
@@ -1498,16 +1498,14 @@ namespace UnrealBuildTool
 				if (UnrealBuildTool.bIsInvalidatingMakefilesOnly)
 				{
 					Log.TraceInformation("Invalidating makefiles only in this run.");
-					if (TargetDescs != null && TargetDescs.Count == 1)
-					{
-						InvalidateMakefiles(TargetDescs[0]);
-						return ECompilationResult.Succeeded;
-					}
-					else
+					if (TargetDescs.Count != 1)
 					{
 						Log.TraceError("You have to provide one target name for makefile invalidation.");
 						return ECompilationResult.OtherCompilationError;
 					}
+
+					InvalidateMakefiles(TargetDescs[0]);
+					return ECompilationResult.Succeeded;
 				}
 
 				UEBuildConfiguration.bHotReloadFromIDE = UEBuildConfiguration.bAllowHotReloadFromIDE && TargetDescs.Count == 1 && !TargetDescs[0].bIsEditorRecompile && ShouldDoHotReloadFromIDE(TargetDescs[0]);
@@ -2267,8 +2265,6 @@ namespace UnrealBuildTool
 						ReasonNotLoaded = "UnrealBuildTool.exe is newer";
 						bForceOutOfDate = true;
 					}
-
-
 				}
 			}
 			else
@@ -2312,6 +2308,46 @@ namespace UnrealBuildTool
                     LoadedUBTMakefile = null;
 					ReasonNotLoaded = "existing makefile appears to be invalid";
                 }
+
+				if( LoadedUBTMakefile != null )
+				{
+					// Check if any of the target's Build.cs files are newer than the makefile
+					foreach (var Target in LoadedUBTMakefile.Targets)
+					{
+						string TargetCsFilename = Target.TargetCsFilename;
+						if (TargetCsFilename != null)
+						{
+							var TargetCsFile = new FileInfo(TargetCsFilename);
+							bool bTargetCsFileExists = TargetCsFile.Exists;
+							if (!bTargetCsFileExists || TargetCsFile.LastWriteTime > UBTMakefileItem.LastWriteTime)
+							{
+								Log.TraceWarning("{0} has been {1} since makefile was built, ignoring it ({2})", TargetCsFilename, bTargetCsFileExists ? "changed" : "deleted", UBTMakefileItem.AbsolutePath);
+								LoadedUBTMakefile = null;
+								ReasonNotLoaded = string.Format("changes to target files");
+								goto SkipRemainingTimestampChecks;
+							}
+						}
+
+						List<string> BuildCsFilenames = Target.GetAllModuleBuildCsFilenames();
+						foreach (var BuildCsFilename in BuildCsFilenames)
+						{
+							if (BuildCsFilename != null)
+							{
+								var BuildCsFile = new FileInfo(BuildCsFilename);
+								bool bBuildCsFileExists = BuildCsFile.Exists;
+								if (!bBuildCsFileExists || BuildCsFile.LastWriteTime > UBTMakefileItem.LastWriteTime)
+								{
+									Log.TraceWarning("{0} has been {1} since makefile was built, ignoring it ({2})", BuildCsFilename, bBuildCsFileExists ? "changed" : "deleted", UBTMakefileItem.AbsolutePath);
+									LoadedUBTMakefile = null;
+									ReasonNotLoaded = string.Format("changes to module files");
+									goto SkipRemainingTimestampChecks;
+								}
+							}
+						}
+					}
+
+				SkipRemainingTimestampChecks:;
+				}
             }
 
             return LoadedUBTMakefile;
