Add a Marketplace URL field to the .uplugin descriptor. The URL will be embedded into projects that enable the plugin, which will prompt to open the marketplace at the given URL if the user doesn't have the plugin installed.

diff --git a/Engine/Plugins/Editor/PluginBrowser/Source/PluginBrowser/Private/SPluginTile.cpp b/Engine/Plugins/Editor/PluginBrowser/Source/PluginBrowser/Private/SPluginTile.cpp
--- a/Engine/Plugins/Editor/PluginBrowser/Source/PluginBrowser/Private/SPluginTile.cpp
+++ b/Engine/Plugins/Editor/PluginBrowser/Source/PluginBrowser/Private/SPluginTile.cpp
@@ -347,7 +347,7 @@ void SPluginTile::OnEnablePluginCheckboxChanged(ECheckBoxState NewCheckedState)
 	FGameProjectGenerationModule::Get().TryMakeProjectFileWriteable(FPaths::GetProjectFilePath());
 	FText FailMessage;
 
-	if (!IProjectManager::Get().SetPluginEnabled(Plugin->GetName(), bNewEnabledState, FailMessage))
+	if (!IProjectManager::Get().SetPluginEnabled(Plugin->GetName(), bNewEnabledState, FailMessage, PluginDescriptor.MarketplaceURL))
 	{
 		FMessageDialog::Open(EAppMsgType::Ok, FailMessage);
 	}
diff --git a/Engine/Source/Programs/UnrealBuildTool/System/PluginDescriptor.cs b/Engine/Source/Programs/UnrealBuildTool/System/PluginDescriptor.cs
--- a/Engine/Source/Programs/UnrealBuildTool/System/PluginDescriptor.cs
+++ b/Engine/Source/Programs/UnrealBuildTool/System/PluginDescriptor.cs
@@ -52,6 +52,9 @@ namespace UnrealBuildTool
 		// Documentation URL string.
 		public string DocsURL;
 
+		// Marketplace URL for this plugin. This URL will be embedded into projects that enable this plugin, so we can redirect to the marketplace if a user doesn't have it installed.
+		public string MarketplaceURL;
+
 		// List of all modules associated with this plugin
 		public ModuleDescriptor[] Modules;
 
@@ -124,6 +127,7 @@ namespace UnrealBuildTool
 				RawObject.TryGetStringField("CreatedBy", out Descriptor.CreatedBy);
 				RawObject.TryGetStringField("CreatedByURL", out Descriptor.CreatedByURL);
 				RawObject.TryGetStringField("DocsURL", out Descriptor.DocsURL);
+				RawObject.TryGetStringField("MarketplaceURL", out Descriptor.MarketplaceURL);
 
 				JsonObject[] ModulesArray;
 				if(RawObject.TryGetObjectArrayField("Modules", out ModulesArray))
@@ -163,6 +167,7 @@ namespace UnrealBuildTool
 				Writer.WriteValue("CreatedBy", CreatedBy);
 				Writer.WriteValue("CreatedByURL", CreatedByURL);
 				Writer.WriteValue("DocsURL", DocsURL);
+				Writer.WriteValue("MarketplaceURL", MarketplaceURL);
 
 				ModuleDescriptor.WriteArray(Writer, "Modules", Modules);
 
@@ -188,6 +193,9 @@ namespace UnrealBuildTool
 		// Description of the plugin for users that do not have it installed.
 		public string Description;
 
+		// URL for this plugin on the marketplace, if the user doesn't have it installed.
+		public string MarketplaceURL;
+
 		// If enabled, list of platforms for which the plugin should be enabled (or all platforms if blank).
 		UnrealTargetPlatform[] WhitelistPlatforms;
 
@@ -199,9 +207,10 @@ namespace UnrealBuildTool
 		/// </summary>
 		/// <param name="InName">Name of the plugin</param>
 		/// <param name="bInEnabled">Whether the plugin is enabled</param>
-		public PluginReferenceDescriptor(string InName, bool bInEnabled)
+		public PluginReferenceDescriptor(string InName, string InMarketplaceURL, bool bInEnabled)
 		{
 			Name = InName;
+			MarketplaceURL = InMarketplaceURL;
 			bEnabled = bInEnabled;
 		}
 
@@ -212,8 +221,9 @@ namespace UnrealBuildTool
 		/// <returns>New PluginReferenceDescriptor object</returns>
 		public static PluginReferenceDescriptor FromJsonObject(JsonObject RawObject)
 		{
-			PluginReferenceDescriptor Descriptor = new PluginReferenceDescriptor(RawObject.GetStringField("Name"), RawObject.GetBoolField("Enabled"));
+			PluginReferenceDescriptor Descriptor = new PluginReferenceDescriptor(RawObject.GetStringField("Name"), null, RawObject.GetBoolField("Enabled"));
 			RawObject.TryGetStringField("Description", out Descriptor.Description);
+			RawObject.TryGetStringField("MarketplaceURL", out Descriptor.MarketplaceURL);
 			RawObject.TryGetEnumArrayField<UnrealTargetPlatform>("WhitelistPlatforms", out Descriptor.WhitelistPlatforms);
 			RawObject.TryGetEnumArrayField<UnrealTargetPlatform>("BlacklistPlatforms", out Descriptor.BlacklistPlatforms);
 			return Descriptor;
diff --git a/Engine/Source/Runtime/Projects/Private/PluginDescriptor.cpp b/Engine/Source/Runtime/Projects/Private/PluginDescriptor.cpp
--- a/Engine/Source/Runtime/Projects/Private/PluginDescriptor.cpp
+++ b/Engine/Source/Runtime/Projects/Private/PluginDescriptor.cpp
@@ -89,6 +89,7 @@ bool FPluginDescriptor::Read(const FString& Text, FText& OutFailReason)
 	Object.TryGetStringField(TEXT("CreatedBy"), CreatedBy);
 	Object.TryGetStringField(TEXT("CreatedByURL"), CreatedByURL);
 	Object.TryGetStringField(TEXT("DocsURL"), DocsURL);
+	Object.TryGetStringField(TEXT("MarketplaceURL"), MarketplaceURL);
 
 	if (!FModuleDescriptor::ReadArray(Object, TEXT("Modules"), Modules, OutFailReason))
 	{
@@ -135,6 +136,7 @@ FString FPluginDescriptor::ToString() const
 	Writer.WriteValue(TEXT("CreatedBy"), CreatedBy);
 	Writer.WriteValue(TEXT("CreatedByURL"), CreatedByURL);
 	Writer.WriteValue(TEXT("DocsURL"), DocsURL);
+	Writer.WriteValue(TEXT("MarketplaceURL"), MarketplaceURL);
 
 	FModuleDescriptor::WriteArray(Writer, TEXT("Modules"), Modules);
 
@@ -151,8 +153,9 @@ FString FPluginDescriptor::ToString() const
 
 //////////////////////////////////////////////////////////////////////////////////////////////////////////////////
 
-FPluginReferenceDescriptor::FPluginReferenceDescriptor( const FString &InName, bool bInEnabled )
+FPluginReferenceDescriptor::FPluginReferenceDescriptor( const FString& InName, const FString& InMarketplaceURL, bool bInEnabled )
 	: Name(InName)
+	, MarketplaceURL(InMarketplaceURL)
 	, bEnabled(bInEnabled)
 { }
 
@@ -197,8 +200,9 @@ bool FPluginReferenceDescriptor::Read( const FJsonObject& Object, FText& OutFail
 		return false;
 	}
 
-	// Read the description
+	// Read the metadata for users that don't have the plugin installed
 	Object.TryGetStringField(TEXT("Description"), Description);
+	Object.TryGetStringField(TEXT("MarketplaceURL"), MarketplaceURL);
 
 	// Get the platform lists
 	Object.TryGetStringArrayField(TEXT("WhitelistPlatforms"), WhitelistPlatforms);
@@ -247,6 +251,11 @@ void FPluginReferenceDescriptor::Write( TJsonWriter<>& Writer ) const
 		Writer.WriteValue(TEXT("Description"), Description);
 	}
 
+	if (MarketplaceURL.Len() > 0)
+	{
+		Writer.WriteValue(TEXT("MarketplaceURL"), MarketplaceURL);
+	}
+
 	if (WhitelistPlatforms.Num() > 0)
 	{
 		Writer.WriteArrayStart(TEXT("WhitelistPlatforms"));
diff --git a/Engine/Source/Runtime/Projects/Private/PluginManager.cpp b/Engine/Source/Runtime/Projects/Private/PluginManager.cpp
--- a/Engine/Source/Runtime/Projects/Private/PluginManager.cpp
+++ b/Engine/Source/Runtime/Projects/Private/PluginManager.cpp
@@ -343,8 +343,20 @@ bool FPluginManager::ConfigureEnabledPlugins()
 				if(Plugin.bEnabled && !FindPluginInstance(Plugin.Name).IsValid())
 				{
 					FText Caption(LOCTEXT("PluginMissingCaption", "Plugin missing"));
-					FString Description = (Plugin.Description.Len() > 0)? FString::Printf(TEXT("\n\n%s"), *Plugin.Description) : FString();
-					FMessageDialog::Open(EAppMsgType::Ok, FText::Format(LOCTEXT("PluginMissingError", "This project requires the {0} plugin. {1}"), FText::FromString(Plugin.Name), FText::FromString(Description)), &Caption);
+					if(Plugin.MarketplaceURL.Len() > 0)
+					{
+						if(FMessageDialog::Open(EAppMsgType::YesNo, FText::Format(LOCTEXT("PluginMissingError", "This project requires the {0} plugin.\n\nWould you like to download it from the the Marketplace?"), FText::FromString(Plugin.Name)), &Caption) == EAppReturnType::Yes)
+						{
+							FString Error;
+							FPlatformProcess::LaunchURL(*Plugin.MarketplaceURL, nullptr, &Error);
+							if(Error.Len() > 0) FMessageDialog::Open(EAppMsgType::Ok, FText::FromString(Error));
+						}
+					}
+					else
+					{
+						FString Description = (Plugin.Description.Len() > 0)? FString::Printf(TEXT("\n\n%s"), *Plugin.Description) : FString();
+						FMessageDialog::Open(EAppMsgType::Ok, FText::Format(LOCTEXT("PluginMissingError", "This project requires the {0} plugin. {1}"), FText::FromString(Plugin.Name), FText::FromString(Description)), &Caption);
+					}
 					return false;
 				}
 			}
diff --git a/Engine/Source/Runtime/Projects/Private/ProjectManager.cpp b/Engine/Source/Runtime/Projects/Private/ProjectManager.cpp
--- a/Engine/Source/Runtime/Projects/Private/ProjectManager.cpp
+++ b/Engine/Source/Runtime/Projects/Private/ProjectManager.cpp
@@ -275,7 +275,7 @@ bool FProjectManager::IsNonDefaultPluginEnabled() const
 	return false;
 }
 
-bool FProjectManager::SetPluginEnabled(const FString& PluginName, bool bEnabled, FText& OutFailReason)
+bool FProjectManager::SetPluginEnabled(const FString& PluginName, bool bEnabled, FText& OutFailReason, const FString& MarketplaceURL)
 {
 	// Don't go any further if there's no project loaded
 	if(!CurrentProject.IsValid())
@@ -290,7 +290,7 @@ bool FProjectManager::SetPluginEnabled(const FString& PluginName, bool bEnabled,
 	{
 		if(PluginRefIdx == CurrentProject->Plugins.Num())
 		{
-			PluginRefIdx = CurrentProject->Plugins.Add(FPluginReferenceDescriptor(PluginName, bEnabled));
+			PluginRefIdx = CurrentProject->Plugins.Add(FPluginReferenceDescriptor(PluginName, MarketplaceURL, bEnabled));
 			break;
 		}
 		else if(CurrentProject->Plugins[PluginRefIdx].Name == PluginName)
diff --git a/Engine/Source/Runtime/Projects/Private/ProjectManager.h b/Engine/Source/Runtime/Projects/Private/ProjectManager.h
--- a/Engine/Source/Runtime/Projects/Private/ProjectManager.h
+++ b/Engine/Source/Runtime/Projects/Private/ProjectManager.h
@@ -28,7 +28,7 @@ public:
 	virtual FOnTargetPlatformsForCurrentProjectChangedEvent& OnTargetPlatformsForCurrentProjectChanged() override { return OnTargetPlatformsForCurrentProjectChangedEvent; }
 	virtual void GetEnabledPlugins(TArray<FString>& OutPluginNames) const override;
 	virtual bool IsNonDefaultPluginEnabled() const override;
-	virtual bool SetPluginEnabled(const FString& PluginName, bool bEnabled, FText& OutFailReason) override;
+	virtual bool SetPluginEnabled(const FString& PluginName, bool bEnabled, FText& OutFailReason, const FString& MarketplaceURL) override;
 
 private:
 	static void QueryStatusForProjectImpl(const FProjectDescriptor& Project, const FString& FilePath, FProjectStatus& OutProjectStatus);
diff --git a/Engine/Source/Runtime/Projects/Public/Interfaces/IProjectManager.h b/Engine/Source/Runtime/Projects/Public/Interfaces/IProjectManager.h
--- a/Engine/Source/Runtime/Projects/Public/Interfaces/IProjectManager.h
+++ b/Engine/Source/Runtime/Projects/Public/Interfaces/IProjectManager.h
@@ -198,7 +198,8 @@ public:
 	 * @param	PluginName		Name of the plugin
 	 * @param	bEnabled		Whether to enable or disable the plugin
 	 * @param	OutFailReason	On failure, gives an error message
+	 * @param	MarketplaceURL	Marketplace URL to open if the user does not have this plugin installed
 	 * @return	True if the plugin has been marked as enabled, and the project descriptor has been saved.
 	 */
-	virtual bool SetPluginEnabled(const FString& PluginName, bool bEnabled, FText& OutFailReason) = 0;
+	virtual bool SetPluginEnabled(const FString& PluginName, bool bEnabled, FText& OutFailReason, const FString& MarketplaceURL = TEXT("")) = 0;
 };
\ No newline at end of file
diff --git a/Engine/Source/Runtime/Projects/Public/PluginDescriptor.h b/Engine/Source/Runtime/Projects/Public/PluginDescriptor.h
--- a/Engine/Source/Runtime/Projects/Public/PluginDescriptor.h
+++ b/Engine/Source/Runtime/Projects/Public/PluginDescriptor.h
@@ -58,6 +58,9 @@ struct PROJECTS_API FPluginDescriptor
 	/** Documentation URL string. */
 	FString DocsURL;
 
+	/** Marketplace URL for this plugin. This URL will be embedded into projects that enable this plugin, so we can redirect to the marketplace if a user doesn't have it installed. */
+	FString MarketplaceURL;
+
 	/** List of all modules associated with this plugin */
 	TArray<FModuleDescriptor> Modules;
 
@@ -103,6 +106,9 @@ struct PROJECTS_API FPluginReferenceDescriptor
 	/** Description of the plugin for users that do not have it installed. */
 	FString Description;
 
+	/** URL for this plugin on the marketplace, if the user doesn't have it installed. */
+	FString MarketplaceURL;
+
 	/** If enabled, list of platforms for which the plugin should be enabled (or all platforms if blank). */
 	TArray<FString> WhitelistPlatforms;
 
@@ -110,7 +116,7 @@ struct PROJECTS_API FPluginReferenceDescriptor
 	TArray<FString> BlacklistPlatforms;
 
 	/** Constructor */
-	FPluginReferenceDescriptor(const FString& InName = TEXT(""), bool bInEnabled = false);
+	FPluginReferenceDescriptor(const FString& InName = TEXT(""), const FString& InMarketplaceURL = TEXT(""), bool bInEnabled = false);
 
 	/** Determines whether the plugin is enabled for the given platform */
 	bool IsEnabledForPlatform(const FString& Platform) const;
