Fix a UBT error trying to scrape the MSVC environment when packaging from the editor. The editor adds a lot of DLL search paths to the PATH environment variable, which can cause the variable to exceed the max allowed length when MSVC tries to add its own. Reset it to the user default first.

diff --git a/Engine/Source/Programs/UnrealBuildTool/System/Utils.cs b/Engine/Source/Programs/UnrealBuildTool/System/Utils.cs
--- a/Engine/Source/Programs/UnrealBuildTool/System/Utils.cs
+++ b/Engine/Source/Programs/UnrealBuildTool/System/Utils.cs
@@ -296,6 +296,10 @@ namespace UnrealBuildTool
 					StartInfo.CreateNoWindow = true;
 					StartInfo.UseShellExecute = false;
 
+					// The engine adds a lot of DLL search paths to the PATH environment variable, and this gets propagated through to UBT. MSVC wants
+					// to add a lot more, so reset it to the user default before spawning the batch file, otherwise we can overflow the max length and fail.
+					StartInfo.EnvironmentVariables["PATH"] = Environment.GetEnvironmentVariable("PATH", EnvironmentVariableTarget.User);
+
 					// @todo UAP: Win10-beta issue, sometimes the process hangs if standard I/O is redirected, and there's bytes
 					// in the buffer and there is no reader for those streams.  Loop back with a later Win10 build
 					if (System.Environment.OSVersion.Platform != PlatformID.Win32NT ||
