UE-13112 : Android : The game is 'launched by 'launch' button doesn't work after using "Package Project..." feature
- changed the deploy script so that it will ignore the OBB when we aren't allowing data packaging

diff --git a/Engine/Source/Programs/UnrealBuildTool/Android/UEDeployAndroid.cs b/Engine/Source/Programs/UnrealBuildTool/Android/UEDeployAndroid.cs
--- a/Engine/Source/Programs/UnrealBuildTool/Android/UEDeployAndroid.cs
+++ b/Engine/Source/Programs/UnrealBuildTool/Android/UEDeployAndroid.cs
@@ -1115,7 +1115,7 @@ namespace UnrealBuildTool.Android
 			bool bBuildSettingsMatch = false;
 
 			string ManifestFile = Path.Combine(UE4BuildPath, "AndroidManifest.xml");
-			string NewManifest = GenerateManifest(ProjectName, bForDistribution, bPackageDataInsideApk, GameBuildFilesPath, File.Exists(ObbFileLocation), bDisableVerifyOBBOnStartUp);
+            string NewManifest = GenerateManifest(ProjectName, bForDistribution, bPackageDataInsideApk, GameBuildFilesPath, bDisallowPackagingDataInApk ? false : File.Exists(ObbFileLocation), bDisableVerifyOBBOnStartUp);
 			string OldManifest = File.Exists(ManifestFile) ? File.ReadAllText(ManifestFile) : "";
 			if (NewManifest == OldManifest) 
 			{
