UE-13706: GitHub 1050 : Build configuration method call inconsistencies and code clean up

GitHub PR #1050

diff --git a/Engine/Source/Programs/UnrealBuildTool/Utilities/XmlConfigLoader.cs b/Engine/Source/Programs/UnrealBuildTool/Utilities/XmlConfigLoader.cs
--- a/Engine/Source/Programs/UnrealBuildTool/Utilities/XmlConfigLoader.cs
+++ b/Engine/Source/Programs/UnrealBuildTool/Utilities/XmlConfigLoader.cs
@@ -36,14 +36,7 @@ namespace UnrealBuildTool
 		/// <typeparam name="ConfigClass">The class to reset.</typeparam>
 		public static void Reset<ConfigClass>()
 		{
-			var Type = typeof(ConfigClass);
-
-			InvokeIfExists(Type, "LoadDefaults");
-
-			// Load eventual XML configuration files if they exist to override default values.
-			Load(Type);
-
-			InvokeIfExists(Type, "PostReset");
+			Load(typeof(ConfigClass));
 		}
 
 		/// <summary>
@@ -82,11 +75,6 @@ namespace UnrealBuildTool
 		/// <returns>Default config XML.</returns>
 		public static string GetDefaultXML()
 		{
-			foreach(var ConfigType in GetAllConfigurationTypes())
-			{
-				InvokeIfExists(ConfigType, "LoadDefaults");
-			}
-
 			var Comment = @"
 	#########################################################################
 	#																		#
@@ -145,12 +133,16 @@ namespace UnrealBuildTool
 		/// <param name="Class">A type of a config class.</param>
 		public static void Load(Type Class)
 		{
-			XmlConfigLoaderClassData ClassData;
+			InvokeIfExists(Class, "LoadDefaults");
+			InvokeIfExists(Class, "Reset");
 
+			XmlConfigLoaderClassData ClassData;
 			if(Data.TryGetValue(Class, out ClassData))
 			{
 				ClassData.LoadXmlData();
 			}
+
+			InvokeIfExists(Class, "PostReset");
 		}
 
 		/// <summary>
@@ -207,13 +199,14 @@ namespace UnrealBuildTool
 
             OverwriteIfDifferent(GetXSDPath(), BuildXSD());
 
+			LoadDefaults();
+			CreateUserXmlConfigTemplate();
+
 			LoadData();
 
-			foreach(var ClassData in Data)
+			foreach(var ConfClass in Data.Keys)
 			{
-				ClassData.Value.ResetData();
-
-				InvokeIfExists(ClassData.Key, "PostReset");
+				Load(ConfClass);
 			}
 		}
 
@@ -265,35 +258,21 @@ namespace UnrealBuildTool
 		}
 
 		/// <summary>
-		/// Cache entry class to store loaded info for given class.
+		/// Loads default values for all configuration classes in assembly.
 		/// </summary>
-		class XmlConfigLoaderClassData
+		private static void LoadDefaults()
 		{
-			public XmlConfigLoaderClassData(Type ConfigClass)
-			{
-				// Adding empty types array to make sure this is parameterless Reset
-				// in case of overloading.
-				DefaultValuesLoader = ConfigClass.GetMethod("Reset", new Type[] { });
-			}
-
-			/// <summary>
-			/// Resets previously stored data into class.
-			/// </summary>
-			public void ResetData()
+			foreach (var ConfigType in GetAllConfigurationTypes())
 			{
-				bDoneLoading = false;
-
-				if(DefaultValuesLoader != null)
-				{
-					DefaultValuesLoader.Invoke(null, new object[] { });
-				}
-
-				if (!bDoneLoading)
-				{
-					LoadXmlData();
-				}
+				InvokeIfExists(ConfigType, "LoadDefaults");
 			}
+		}
 
+		/// <summary>
+		/// Cache entry class to store loaded info for given class.
+		/// </summary>
+		class XmlConfigLoaderClassData
+		{
 			/// <summary>
 			/// Loads previously stored data into class.
 			/// </summary>
@@ -308,8 +287,6 @@ namespace UnrealBuildTool
                 {
                     PropertyPair.Key.SetValue(null, PropertyPair.Value, null);
                 }
-
-                bDoneLoading = true;
 			}
 
 			/// <summary>
@@ -346,13 +323,6 @@ namespace UnrealBuildTool
                 }
             }
 
-            // A variable to indicate if loading was done during invoking of
-			// default values loader.
-			bool bDoneLoading = false;
-
-			// Method info loader to invoke before overriding fields with XML files data.
-			MethodInfo DefaultValuesLoader = null;
-
 			// Loaded data map.
 			Dictionary<FieldInfo, object> DataMap = new Dictionary<FieldInfo, object>();
         
@@ -550,16 +520,6 @@ namespace UnrealBuildTool
 		/// </summary>
 		private static void LoadData()
 		{
-			foreach (var PossibleConfigLocation in ConfigLocationHierarchy)
-			{
-				if (!PossibleConfigLocation.IfShouldCreateFile())
-				{
-					continue;
-				}
-
-				PossibleConfigLocation.CreateUserXmlConfigTemplate();
-			}
-
 			foreach (var PossibleConfigLocation in ConfigLocationHierarchy)
 			{
 				if(!PossibleConfigLocation.bExists)
@@ -569,7 +529,7 @@ namespace UnrealBuildTool
 
 				try
 				{
-					Load(PossibleConfigLocation.FSLocation);
+					LoadDataFromFile(PossibleConfigLocation.FSLocation);
 				}
 				catch(Exception Ex)
 				{
@@ -578,6 +538,23 @@ namespace UnrealBuildTool
 			}
 		}
 
+		/// <summary>
+		/// Creates XML files for all known XmlConfigLocations which return IfShouldCreateFile()==true. Files
+		/// will be filled with default content chosen by the XmlConfigLocation implementation.
+		/// </summary>
+		private static void CreateUserXmlConfigTemplate()
+		{
+			foreach (var PossibleConfigLocation in ConfigLocationHierarchy)
+			{
+				if (!PossibleConfigLocation.IfShouldCreateFile())
+				{
+					continue;
+				}
+
+				PossibleConfigLocation.CreateUserXmlConfigTemplate();
+			}
+		}
+
 		/// <summary>
 		/// Reads config schema from XSD file.
 		/// </summary>
@@ -620,7 +597,7 @@ namespace UnrealBuildTool
 		/// Sets values of this class with values from given XML file.
 		/// </summary>
 		/// <param name="ConfigurationXmlPath">The path to the file with values.</param>
-		private static void Load(string ConfigurationXmlPath)
+		private static void LoadDataFromFile(string ConfigurationXmlPath)
 		{
 			var ConfigDocument = new XmlDocument();
 			var NS = new XmlNamespaceManager(ConfigDocument.NameTable);
@@ -687,7 +664,7 @@ namespace UnrealBuildTool
 
 				if (!Data.TryGetValue(ClassType, out ClassData))
 				{
-					ClassData = new XmlConfigLoaderClassData(ClassType);
+					ClassData = new XmlConfigLoaderClassData();
 					Data.Add(ClassType, ClassData);
 				}
 
