PR-1129 - RSyncUsername being overriden with empty string
UE-15307
#ios

diff --git a/Engine/Source/Programs/UnrealBuildTool/ToolChain/RemoteToolChain.cs b/Engine/Source/Programs/UnrealBuildTool/ToolChain/RemoteToolChain.cs
--- a/Engine/Source/Programs/UnrealBuildTool/ToolChain/RemoteToolChain.cs
+++ b/Engine/Source/Programs/UnrealBuildTool/ToolChain/RemoteToolChain.cs
@@ -208,7 +208,12 @@ namespace UnrealBuildTool
 			if (Ini.GetBool("/Script/IOSRuntimeSettings.IOSRuntimeSettings", "bUseRSync", out bUseRSync))
 			{
 				bUseRPCUtil = !bUseRSync;
-				Ini.GetString("/Script/IOSRuntimeSettings.IOSRuntimeSettings", "RSyncUsername", out RSyncUsername);
+				string UserName = RSyncUsername;
+
+				if (Ini.GetString("/Script/IOSRuntimeSettings.IOSRuntimeSettings", "RSyncUsername", out UserName))
+				{
+					RSyncUsername = UserName;
+				}
 				
 				if (Ini.GetString("/Script/IOSRuntimeSettings.IOSRuntimeSettings", "DeltaCopyInstallPath", out OverrideDeltaCopyInstallPath))
 				{
