Make bEnableShadowVariableWarnings a per-module setting, so we don't break reuse of shared build intermediates.

diff --git a/Engine/Config/BaseEngine.ini b/Engine/Config/BaseEngine.ini
--- a/Engine/Config/BaseEngine.ini
+++ b/Engine/Config/BaseEngine.ini
@@ -2239,6 +2239,3 @@ CrashReportReceiverIP="http://crashreporter.epicgames.com:57005"
 
 [UnrealHeaderTool]
 GeneratedCodeVersion=V1
-
-[/Script/BuildSettings.BuildSettings]
-bEnableShadowVariableWarning=true
\ No newline at end of file
diff --git a/Engine/Source/Programs/UnrealBuildTool/Configuration/BuildConfiguration.cs b/Engine/Source/Programs/UnrealBuildTool/Configuration/BuildConfiguration.cs
--- a/Engine/Source/Programs/UnrealBuildTool/Configuration/BuildConfiguration.cs
+++ b/Engine/Source/Programs/UnrealBuildTool/Configuration/BuildConfiguration.cs
@@ -90,12 +90,6 @@ namespace UnrealBuildTool
 		[XmlConfig]
 		public static bool bOmitPCDebugInfoInDevelopment;
 
-		/// <summary>
-		/// Whether shadow variable warning should be enabled
-		/// </summary>
-		[XmlConfig]
-		public static bool bEnableShadowVariableWarning;
-
 		/// <summary>
 		/// Whether PDB files should be used for Visual C++ builds.
 		/// </summary>
@@ -483,8 +477,6 @@ namespace UnrealBuildTool
 			bSupportEditAndContinue = false;
 			bUseActionHistory = true;
 
-			bEnableShadowVariableWarning = false;
-
 			// Incremental linking can yield faster iteration times when making small changes
 			// NOTE: We currently don't use incremental linking because it tends to behave a bit buggy on some computers (PDB-related compile errors)
 			bUseIncrementalLinking = false;
diff --git a/Engine/Source/Programs/UnrealBuildTool/Configuration/UEBuildModule.cs b/Engine/Source/Programs/UnrealBuildTool/Configuration/UEBuildModule.cs
--- a/Engine/Source/Programs/UnrealBuildTool/Configuration/UEBuildModule.cs
+++ b/Engine/Source/Programs/UnrealBuildTool/Configuration/UEBuildModule.cs
@@ -1047,6 +1047,9 @@ namespace UnrealBuildTool
 		/** Enable exception handling */
 		public bool bEnableExceptions = false;
 
+		/** Enable warnings for shadowed variables */
+		public bool bEnableShadowVariableWarnings;
+
 		public List<string> IncludeSearchPaths = new List<string>();
 
 		public class ProcessedDependenciesClass
@@ -1135,6 +1138,7 @@ namespace UnrealBuildTool
 			bool InFasterWithoutUnity,
 			int InMinFilesUsingPrecompiledHeaderOverride,
 			bool InEnableExceptions,
+			bool InEnableShadowVariableWarnings,
 			bool bInBuildSourceFiles,
 			string InBuildCsFilename
 			)
@@ -1187,6 +1191,7 @@ namespace UnrealBuildTool
 			bFasterWithoutUnity                    = InFasterWithoutUnity;
 			MinFilesUsingPrecompiledHeaderOverride = InMinFilesUsingPrecompiledHeaderOverride;
 			bEnableExceptions                      = InEnableExceptions;
+			bEnableShadowVariableWarnings          = InEnableShadowVariableWarnings;
 		}
 
 		// UEBuildModule interface.
@@ -1833,6 +1838,7 @@ namespace UnrealBuildTool
 			Result.Config.bFasterWithoutUnity                    = bFasterWithoutUnity;
 			Result.Config.MinFilesUsingPrecompiledHeaderOverride = MinFilesUsingPrecompiledHeaderOverride;
 			Result.Config.bEnableExceptions                      = bEnableExceptions;
+			Result.Config.bEnableShadowVariableWarning          = bEnableShadowVariableWarnings;
 			Result.Config.bUseStaticCRT							 = (Target.Rules != null && Target.Rules.bUseStaticCRT);
 			Result.Config.OutputDirectory                        = Path.Combine(Binary.Config.IntermediateDirectory, Name);
 
@@ -2119,6 +2125,7 @@ namespace UnrealBuildTool
 			bool InFasterWithoutUnity,
 			int InMinFilesUsingPrecompiledHeaderOverride,
 			bool InEnableExceptions,
+			bool InEnableShadowVariableWarnings,
 			bool bInBuildSourceFiles,
 			string InBuildCsFilename
 			)
@@ -2128,7 +2135,7 @@ namespace UnrealBuildTool
 			InPrivateIncludePaths,InPrivateIncludePathModuleNames,InPrivateDependencyModuleNames,
             InCircularlyReferencedDependentModules, InDynamicallyLoadedModuleNames, InPlatformSpecificDynamicallyLoadedModuleNames, InRuntimeDependencies, InOptimizeCode,
 			InAllowSharedPCH, InSharedPCHHeaderFile, InUseRTTI, InEnableBufferSecurityChecks, InFasterWithoutUnity, InMinFilesUsingPrecompiledHeaderOverride,
-			InEnableExceptions, bInBuildSourceFiles, InBuildCsFilename)
+			InEnableExceptions, InEnableShadowVariableWarnings, bInBuildSourceFiles, InBuildCsFilename)
 		{
 			PrivateAssemblyReferences = HashSetFromOptionalEnumerableStringParameter(InPrivateAssemblyReferences);
 		}
diff --git a/Engine/Source/Programs/UnrealBuildTool/Configuration/UEBuildPlatform.cs b/Engine/Source/Programs/UnrealBuildTool/Configuration/UEBuildPlatform.cs
--- a/Engine/Source/Programs/UnrealBuildTool/Configuration/UEBuildPlatform.cs
+++ b/Engine/Source/Programs/UnrealBuildTool/Configuration/UEBuildPlatform.cs
@@ -828,12 +828,6 @@ namespace UnrealBuildTool
 			{
 				UEBuildConfiguration.bCompileCEF3 = bValue;
 			}
-
-			bValue = BuildConfiguration.bEnableShadowVariableWarning;
-			if (Ini.GetBool("/Script/BuildSettings.BuildSettings", "bEnableShadowVariableWarning", out bValue))
-			{
-				BuildConfiguration.bEnableShadowVariableWarning = bValue;
-			}
 		}
 
 		/**
diff --git a/Engine/Source/Programs/UnrealBuildTool/Configuration/UEBuildTarget.cs b/Engine/Source/Programs/UnrealBuildTool/Configuration/UEBuildTarget.cs
--- a/Engine/Source/Programs/UnrealBuildTool/Configuration/UEBuildTarget.cs
+++ b/Engine/Source/Programs/UnrealBuildTool/Configuration/UEBuildTarget.cs
@@ -2130,6 +2130,7 @@ namespace UnrealBuildTool
 				InFasterWithoutUnity: true,
 				InMinFilesUsingPrecompiledHeaderOverride: 0,
 				InEnableExceptions: false,
+				InEnableShadowVariableWarnings: true,
 				bInBuildSourceFiles: true,
 				InBuildCsFilename: null);
 		}
@@ -3132,6 +3133,7 @@ namespace UnrealBuildTool
 							InFasterWithoutUnity: RulesObject.bFasterWithoutUnity,
 							InMinFilesUsingPrecompiledHeaderOverride: RulesObject.MinFilesUsingPrecompiledHeaderOverride,
 							InEnableExceptions: RulesObject.bEnableExceptions,
+							InEnableShadowVariableWarnings: RulesObject.bEnableShadowVariableWarnings,
 							bInBuildSourceFiles: bBuildSourceFiles,
 							InBuildCsFilename: InBuildCsFile
 						);
@@ -3174,6 +3176,7 @@ namespace UnrealBuildTool
 							InFasterWithoutUnity: RulesObject.bFasterWithoutUnity,
 							InMinFilesUsingPrecompiledHeaderOverride: RulesObject.MinFilesUsingPrecompiledHeaderOverride,
 							InEnableExceptions: RulesObject.bEnableExceptions,
+							InEnableShadowVariableWarnings: RulesObject.bEnableShadowVariableWarnings,
 							bInBuildSourceFiles : bBuildSourceFiles,
 							InBuildCsFilename: InBuildCsFile
 						);
diff --git a/Engine/Source/Programs/UnrealBuildTool/HTML5/HTML5ToolChain.cs b/Engine/Source/Programs/UnrealBuildTool/HTML5/HTML5ToolChain.cs
--- a/Engine/Source/Programs/UnrealBuildTool/HTML5/HTML5ToolChain.cs
+++ b/Engine/Source/Programs/UnrealBuildTool/HTML5/HTML5ToolChain.cs
@@ -39,7 +39,7 @@ namespace UnrealBuildTool
 			}
 		}
 
-		static string GetSharedArguments_Global(CPPTargetConfiguration TargetConfiguration, string Architecture)
+		static string GetSharedArguments_Global(CPPTargetConfiguration TargetConfiguration, string Architecture, bool bEnableShadowVariableWarning)
 		{
             string Result = " ";
 
@@ -64,7 +64,7 @@ namespace UnrealBuildTool
             Result += " -Wno-array-bounds"; // some VectorLoads go past the end of the array, but it's okay in that case
             Result += " -Wno-invalid-offsetof"; // too many warnings kills windows clang. 
 
-			if (BuildConfiguration.bEnableShadowVariableWarning)
+			if (bEnableShadowVariableWarning)
 			{
 				Result += " -Wshadow";
 			}
@@ -100,7 +100,7 @@ namespace UnrealBuildTool
 		
 		static string GetCLArguments_Global(CPPEnvironment CompileEnvironment)
 		{
-			string Result = GetSharedArguments_Global(CompileEnvironment.Config.Target.Configuration, CompileEnvironment.Config.Target.Architecture);
+			string Result = GetSharedArguments_Global(CompileEnvironment.Config.Target.Configuration, CompileEnvironment.Config.Target.Architecture, CompileEnvironment.Config.bEnableShadowVariableWarning);
 
 			if (CompileEnvironment.Config.Target.Architecture != "-win32")
 			{
@@ -169,7 +169,7 @@ namespace UnrealBuildTool
 
 		static string GetLinkArguments(LinkEnvironment LinkEnvironment)
 		{
-			string Result = GetSharedArguments_Global(LinkEnvironment.Config.Target.Configuration, LinkEnvironment.Config.Target.Architecture);
+			string Result = GetSharedArguments_Global(LinkEnvironment.Config.Target.Configuration, LinkEnvironment.Config.Target.Architecture, false);
 
 			if (LinkEnvironment.Config.Target.Architecture != "-win32")
 			{
diff --git a/Engine/Source/Programs/UnrealBuildTool/IOS/IOSToolChain.cs b/Engine/Source/Programs/UnrealBuildTool/IOS/IOSToolChain.cs
--- a/Engine/Source/Programs/UnrealBuildTool/IOS/IOSToolChain.cs
+++ b/Engine/Source/Programs/UnrealBuildTool/IOS/IOSToolChain.cs
@@ -230,7 +230,7 @@ namespace UnrealBuildTool
 
 			Result += " -Wall -Werror";
 				
-			if (BuildConfiguration.bEnableShadowVariableWarning)
+			if (CompileEnvironment.Config.bEnableShadowVariableWarning)
 			{
 				Result += " -Wshadow";
 			}
diff --git a/Engine/Source/Programs/UnrealBuildTool/Linux/LinuxToolChain.cs b/Engine/Source/Programs/UnrealBuildTool/Linux/LinuxToolChain.cs
--- a/Engine/Source/Programs/UnrealBuildTool/Linux/LinuxToolChain.cs
+++ b/Engine/Source/Programs/UnrealBuildTool/Linux/LinuxToolChain.cs
@@ -351,7 +351,7 @@ namespace UnrealBuildTool
             Result += " -Wno-unknown-pragmas";			// Slate triggers this (with its optimize on/off pragmas)
 			Result += " -Wno-invalid-offsetof"; // needed to suppress warnings about using offsetof on non-POD types.
 
-			if (BuildConfiguration.bEnableShadowVariableWarning)
+			if (CompileEnvironment.Config.bEnableShadowVariableWarning)
 			{
 				Result += " -Wshadow";
 			}
diff --git a/Engine/Source/Programs/UnrealBuildTool/Mac/MacToolChain.cs b/Engine/Source/Programs/UnrealBuildTool/Mac/MacToolChain.cs
--- a/Engine/Source/Programs/UnrealBuildTool/Mac/MacToolChain.cs
+++ b/Engine/Source/Programs/UnrealBuildTool/Mac/MacToolChain.cs
@@ -170,7 +170,7 @@ namespace UnrealBuildTool
 			Result += " -Wno-unused-private-field";
 			Result += " -Wno-invalid-offsetof"; // needed to suppress warnings about using offsetof on non-POD types.
 
-			if (BuildConfiguration.bEnableShadowVariableWarning)
+			if (CompileEnvironment.Config.bEnableShadowVariableWarning)
 			{
 				Result += " -Wshadow";
 			}
diff --git a/Engine/Source/Programs/UnrealBuildTool/System/CPPEnvironment.cs b/Engine/Source/Programs/UnrealBuildTool/System/CPPEnvironment.cs
--- a/Engine/Source/Programs/UnrealBuildTool/System/CPPEnvironment.cs
+++ b/Engine/Source/Programs/UnrealBuildTool/System/CPPEnvironment.cs
@@ -212,6 +212,9 @@ namespace UnrealBuildTool
 		/** Enable exception handling */
 		public bool bEnableExceptions = false;
 
+		/** Whether to warn about the use of shadow variables */
+		public bool bEnableShadowVariableWarning = false;
+
 		/** True if the environment contains performance critical code. */
 		public ModuleRules.CodeOptimization OptimizeCode = ModuleRules.CodeOptimization.Default;
 
@@ -270,6 +273,7 @@ namespace UnrealBuildTool
 			bFasterWithoutUnity                    = InCopyEnvironment.bFasterWithoutUnity;
 			MinFilesUsingPrecompiledHeaderOverride = InCopyEnvironment.MinFilesUsingPrecompiledHeaderOverride;
 			bEnableExceptions                      = InCopyEnvironment.bEnableExceptions;
+			bEnableShadowVariableWarning          = InCopyEnvironment.bEnableShadowVariableWarning;
 			OptimizeCode                           = InCopyEnvironment.OptimizeCode;
 			bCreateDebugInfo                       = InCopyEnvironment.bCreateDebugInfo;
 			bIsBuildingLibrary                     = InCopyEnvironment.bIsBuildingLibrary;
diff --git a/Engine/Source/Programs/UnrealBuildTool/System/RulesCompiler.cs b/Engine/Source/Programs/UnrealBuildTool/System/RulesCompiler.cs
--- a/Engine/Source/Programs/UnrealBuildTool/System/RulesCompiler.cs
+++ b/Engine/Source/Programs/UnrealBuildTool/System/RulesCompiler.cs
@@ -170,6 +170,9 @@ namespace UnrealBuildTool
 		/** Enable exception handling */
 		public bool bEnableExceptions = false;
 
+		/** Enable warnings for shadowed variables */
+		public bool bEnableShadowVariableWarnings = true;
+
 		/** If true and unity builds are enabled, this module will build without unity. */
 		public bool bFasterWithoutUnity = false;
 
diff --git a/Engine/Source/Programs/UnrealBuildTool/Windows/VCToolChain.cs b/Engine/Source/Programs/UnrealBuildTool/Windows/VCToolChain.cs
--- a/Engine/Source/Programs/UnrealBuildTool/Windows/VCToolChain.cs
+++ b/Engine/Source/Programs/UnrealBuildTool/Windows/VCToolChain.cs
@@ -675,7 +675,7 @@ namespace UnrealBuildTool
 				// @todo clang: Hack due to how we have our 'DummyPCH' wrappers setup when using unity builds.  This warning should not be disabled!!
 				Arguments.Append(" -Wno-msvc-include");
 
-				if (BuildConfiguration.bEnableShadowVariableWarning)
+				if (CompileEnvironment.Config.bEnableShadowVariableWarning)
 				{
 					Arguments.Append(" -Wshadow");
 				}
