Restore a bunch of UBT cmdline overrides that I accidentally removed in CL2526245.

diff --git a/Engine/Source/Programs/UnrealBuildTool/System/UnrealBuildTool.cs b/Engine/Source/Programs/UnrealBuildTool/System/UnrealBuildTool.cs
--- a/Engine/Source/Programs/UnrealBuildTool/System/UnrealBuildTool.cs
+++ b/Engine/Source/Programs/UnrealBuildTool/System/UnrealBuildTool.cs
@@ -736,6 +736,72 @@ namespace UnrealBuildTool
                     BuildConfiguration.bXGEExport = true;
                     BuildConfiguration.bAllowXGE = true;
                 }
+                else if (LowercaseArg == "-stresstestunity")
+                {
+                    BuildConfiguration.bStressTestUnity = true;
+                }
+                else if (LowercaseArg == "-disableunity")
+                {
+                    BuildConfiguration.bUseUnityBuild = false;
+                }
+                else if (LowercaseArg == "-forceunity")
+                {
+                    BuildConfiguration.bForceUnityBuild = true;
+                }
+                // New Monolithic Graphics drivers have optional "fast calls" replacing various D3d functions
+                else if (LowercaseArg == "-fastmonocalls")
+                {
+                    BuildConfiguration.bUseFastMonoCalls = true;
+                }
+                else if (LowercaseArg == "-nofastmonocalls")
+                {
+                    BuildConfiguration.bUseFastMonoCalls = false;
+                }
+                else if (LowercaseArg == "-uniqueintermediate")
+                {
+                    BuildConfiguration.BaseIntermediateFolder = "Intermediate/Build/" + BuildGuid + "/";
+                }
+                else if (LowercaseArg == "-nopch")
+                {
+                    BuildConfiguration.bUsePCHFiles = false;
+                }
+                else if (LowercaseArg == "-nosharedpch")
+                {
+                    BuildConfiguration.bUseSharedPCHs = false;
+                }
+                else if (LowercaseArg == "-skipActionHistory")
+                {
+                    BuildConfiguration.bUseActionHistory = false;
+                }
+                else if (LowercaseArg == "-noltcg")
+                {
+                    BuildConfiguration.bAllowLTCG = false;
+                }
+                else if (LowercaseArg == "-nopdb")
+                {
+                    BuildConfiguration.bUsePDBFiles = false;
+                }
+                else if (LowercaseArg == "-deploy")
+                {
+                    BuildConfiguration.bDeployAfterCompile = true;
+                }
+                else if (LowercaseArg == "-copyappbundlebacktodevice")
+                {
+                    BuildConfiguration.bCopyAppBundleBackToDevice = true;
+                }
+                else if (LowercaseArg == "-flushmac")
+                {
+                    BuildConfiguration.bFlushBuildDirOnRemoteMac = true;
+                }
+                else if (LowercaseArg == "-nodebuginfo")
+                {
+                    BuildConfiguration.bDisableDebugInfo = true;
+                }
+                else if (LowercaseArg == "-forcedebuginfo")
+                {
+                    BuildConfiguration.bDisableDebugInfo = false;
+                    BuildConfiguration.bOmitPCDebugInfoInDevelopment = false;
+                }
                 else if (LowercaseArg == "-noubtmakefiles")
                 {
                     BuildConfiguration.bUseUBTMakefiles = false;
