Renamed WinUAP platform to UWP
- Microsoft renamed Universal App Platform to Universal Windows Platform
- https://msdn.microsoft.com/en-us/library/dn894631.aspx

diff --git a/Engine/Source/Programs/UnrealBuildTool/Configuration/EngineConfiguration.cs b/Engine/Source/Programs/UnrealBuildTool/Configuration/EngineConfiguration.cs
--- a/Engine/Source/Programs/UnrealBuildTool/Configuration/EngineConfiguration.cs
+++ b/Engine/Source/Programs/UnrealBuildTool/Configuration/EngineConfiguration.cs
@@ -624,7 +624,7 @@ namespace UnrealBuildTool
 		/// </summary>
 		private static string GetIniPlatformName(UnrealTargetPlatform TargetPlatform)
 		{
-			if (TargetPlatform == UnrealTargetPlatform.Win32 || TargetPlatform == UnrealTargetPlatform.Win64 || TargetPlatform == UnrealTargetPlatform.WinUAP )
+			if (TargetPlatform == UnrealTargetPlatform.Win32 || TargetPlatform == UnrealTargetPlatform.Win64 || TargetPlatform == UnrealTargetPlatform.UWP )
 			{
 				return "Windows";
 			}
diff --git a/Engine/Source/Programs/UnrealBuildTool/Configuration/UEBuildPlatform.cs b/Engine/Source/Programs/UnrealBuildTool/Configuration/UEBuildPlatform.cs
--- a/Engine/Source/Programs/UnrealBuildTool/Configuration/UEBuildPlatform.cs
+++ b/Engine/Source/Programs/UnrealBuildTool/Configuration/UEBuildPlatform.cs
@@ -436,7 +436,7 @@ namespace UnrealBuildTool
 					return ":";
 				case UnrealTargetPlatform.Win32:
 				case UnrealTargetPlatform.Win64:
-				case UnrealTargetPlatform.WinUAP:
+				case UnrealTargetPlatform.UWP:
 					return ";";
 				default:
 					Log.TraceWarning("PATH var delimiter unknown for platform " + BuildHostPlatform.Current.Platform.ToString() + " using ';'");
diff --git a/Engine/Source/Programs/UnrealBuildTool/Configuration/UEBuildTarget.cs b/Engine/Source/Programs/UnrealBuildTool/Configuration/UEBuildTarget.cs
--- a/Engine/Source/Programs/UnrealBuildTool/Configuration/UEBuildTarget.cs
+++ b/Engine/Source/Programs/UnrealBuildTool/Configuration/UEBuildTarget.cs
@@ -18,7 +18,7 @@ namespace UnrealBuildTool
 		Win64,
 		WinRT,
 		WinRT_ARM,
-		WinUAP,
+		UWP,
 		Mac,
 		XboxOne,
 		PS4,
@@ -210,7 +210,7 @@ namespace UnrealBuildTool
 			{
 				case CPPTargetPlatform.Win32:			return UnrealTargetPlatform.Win32;
 				case CPPTargetPlatform.Win64:			return UnrealTargetPlatform.Win64;
-				case CPPTargetPlatform.WinUAP:          return UnrealTargetPlatform.WinUAP;
+				case CPPTargetPlatform.UWP:          return UnrealTargetPlatform.UWP;
 				case CPPTargetPlatform.Mac:				return UnrealTargetPlatform.Mac;
 				case CPPTargetPlatform.XboxOne:			return UnrealTargetPlatform.XboxOne;
 				case CPPTargetPlatform.PS4:				return UnrealTargetPlatform.PS4;
@@ -1473,7 +1473,7 @@ namespace UnrealBuildTool
                 }
                 else
                 {
-                    IsCurrentPlatform = Platform == UnrealTargetPlatform.Win64 || Platform == UnrealTargetPlatform.Win32 || Platform == UnrealTargetPlatform.WinUAP;
+                    IsCurrentPlatform = Platform == UnrealTargetPlatform.Win64 || Platform == UnrealTargetPlatform.Win32 || Platform == UnrealTargetPlatform.UWP;
 
                 }
 
diff --git a/Engine/Source/Programs/UnrealBuildTool/System/CPPEnvironment.cs b/Engine/Source/Programs/UnrealBuildTool/System/CPPEnvironment.cs
--- a/Engine/Source/Programs/UnrealBuildTool/System/CPPEnvironment.cs
+++ b/Engine/Source/Programs/UnrealBuildTool/System/CPPEnvironment.cs
@@ -14,7 +14,7 @@ namespace UnrealBuildTool
 	{
 		Win32,
 		Win64,
-		WinUAP,
+		UWP,
 		Mac,
 		XboxOne,
 		PS4,
diff --git a/Engine/Source/Programs/UnrealBuildTool/System/ProjectFileGenerator.cs b/Engine/Source/Programs/UnrealBuildTool/System/ProjectFileGenerator.cs
--- a/Engine/Source/Programs/UnrealBuildTool/System/ProjectFileGenerator.cs
+++ b/Engine/Source/Programs/UnrealBuildTool/System/ProjectFileGenerator.cs
@@ -1483,8 +1483,8 @@ namespace UnrealBuildTool
 					{
 						if (UnrealBuildTool.IsValidPlatform(Platform))
 						{
-							// @todo UAP: Why would we exclude UAP from project files based on SDK install status?  We don't do this with anything else, right?
-                            if (BuildPlatform is WinUAPPlatform && ((WinUAPPlatform)BuildPlatform).HasRequiredSDKsInstalled() != SDKStatus.Valid)
+							// @todo UWP: Why would we exclude UWP from project files based on SDK install status?  We don't do this with anything else, right?
+                            if (BuildPlatform is UWPPlatform && ((UWPPlatform)BuildPlatform).HasRequiredSDKsInstalled() != SDKStatus.Valid)
                                 continue;
 
 							SupportedPlatforms.Add(Platform);
diff --git a/Engine/Source/Programs/UnrealBuildTool/System/UnrealBuildTool.cs b/Engine/Source/Programs/UnrealBuildTool/System/UnrealBuildTool.cs
--- a/Engine/Source/Programs/UnrealBuildTool/System/UnrealBuildTool.cs
+++ b/Engine/Source/Programs/UnrealBuildTool/System/UnrealBuildTool.cs
@@ -258,7 +258,7 @@ namespace UnrealBuildTool
 		static public bool PlatformSupportsCrashReporter(UnrealTargetPlatform InPlatform)
 		{
 			return (
-				(InPlatform == UnrealTargetPlatform.WinUAP) ||
+				(InPlatform == UnrealTargetPlatform.UWP) ||
 				(InPlatform == UnrealTargetPlatform.Win64) ||
 				(InPlatform == UnrealTargetPlatform.Win32) ||
 				(InPlatform == UnrealTargetPlatform.Linux) ||
diff --git a/Engine/Source/Programs/UnrealBuildTool/System/Utils.cs b/Engine/Source/Programs/UnrealBuildTool/System/Utils.cs
--- a/Engine/Source/Programs/UnrealBuildTool/System/Utils.cs
+++ b/Engine/Source/Programs/UnrealBuildTool/System/Utils.cs
@@ -309,7 +309,7 @@ namespace UnrealBuildTool
 					}
 					StartInfo.EnvironmentVariables["PATH"] = NewPathVariable;
 
-					// @todo UAP: Win10-beta issue, sometimes the process hangs if standard I/O is redirected, and there's bytes
+					// @todo UWP: Win10-beta issue, sometimes the process hangs if standard I/O is redirected, and there's bytes
 					// in the buffer and there is no reader for those streams.  Loop back with a later Win10 build
 					if (System.Environment.OSVersion.Platform != PlatformID.Win32NT ||
 						System.Environment.OSVersion.Version.Major < 10)
diff --git a/Engine/Source/Programs/UnrealBuildTool/System/VCProject.cs b/Engine/Source/Programs/UnrealBuildTool/System/VCProject.cs
--- a/Engine/Source/Programs/UnrealBuildTool/System/VCProject.cs
+++ b/Engine/Source/Programs/UnrealBuildTool/System/VCProject.cs
@@ -331,9 +331,9 @@ namespace UnrealBuildTool
 				{
 					VCIncludeSearchPaths.Append(CurPath + ";");
 				}
-				if (InPlatforms.Contains(UnrealTargetPlatform.WinUAP))
+				if (InPlatforms.Contains(UnrealTargetPlatform.UWP))
 				{
-					VCIncludeSearchPaths.Append(WinUAPToolChain.GetVCIncludePaths(CPPTargetPlatform.WinUAP) + ";");
+					VCIncludeSearchPaths.Append(UWPToolChain.GetVCIncludePaths(CPPTargetPlatform.UWP) + ";");
 				}
 				else if (InPlatforms.Contains(UnrealTargetPlatform.Win64))
 				{
@@ -1015,7 +1015,7 @@ namespace UnrealBuildTool
 					}
 					string BatchFilesDirectoryName = Path.Combine(ProjectFileGenerator.EngineRelativePath, "Build", "BatchFiles");
 
-					// @todo UAP: For the MS toolchains, if an override was set for project generation, push that into the build strings to override the build toolchain as well
+					// @todo UWP: For the MS toolchains, if an override was set for project generation, push that into the build strings to override the build toolchain as well
 					string BuildToolOverride = "";
 					if (UnrealBuildTool.CommandLineContains("-2012"))
 					{
@@ -1059,7 +1059,7 @@ namespace UnrealBuildTool
 				{
 					TargetRules TargetRulesObject = Combination.ProjectTarget.TargetRules;
 
-					if ((Platform == UnrealTargetPlatform.Win32) || (Platform == UnrealTargetPlatform.Win64) || (Platform == UnrealTargetPlatform.WinUAP))
+					if ((Platform == UnrealTargetPlatform.Win32) || (Platform == UnrealTargetPlatform.Win64) || (Platform == UnrealTargetPlatform.UWP))
 					{
 						VCUserFileContent.Append(
 							"	<PropertyGroup " + ConditionString + ">" + ProjectFileGenerator.NewLine);
diff --git a/Engine/Source/Programs/UnrealBuildTool/UWP/UEBuildUWP.cs b/Engine/Source/Programs/UnrealBuildTool/UWP/UEBuildUWP.cs
new file mode 100644
--- /dev/null
+++ b/Engine/Source/Programs/UnrealBuildTool/UWP/UEBuildUWP.cs
@@ -0,0 +1,496 @@
+﻿// Copyright 1998-2014 Epic Games, Inc. All Rights Reserved.
+
+using System;
+using System.Collections.Generic;
+using System.Text;
+using System.Diagnostics;
+using System.IO;
+using Microsoft.Win32;
+
+namespace UnrealBuildTool
+{
+	public class UWPPlatform : UEBuildPlatform
+	{
+		/// Property caching.
+		private static WindowsCompiler? CachedCompiler;
+
+		/// Version of the compiler toolchain to use for Universal Windows Platform apps (UWP)
+		public static WindowsCompiler Compiler
+		{
+			get
+			{
+				// Cache the result because Compiler is often used.
+				if (CachedCompiler.HasValue)
+				{
+					return CachedCompiler.Value;
+				}
+
+				// First, default based on whether there is a command line override.
+				// Allows build chain to partially progress even in the absence of installed tools
+				if (UnrealBuildTool.CommandLineContains("-2015"))
+				{
+					CachedCompiler = WindowsCompiler.VisualStudio2015;
+				}
+				// Second, default based on what's installed
+				else if (!String.IsNullOrEmpty(WindowsPlatform.GetVSComnToolsPath(WindowsCompiler.VisualStudio2015)))
+				{
+					CachedCompiler = WindowsCompiler.VisualStudio2015;
+				}
+				else
+				{
+					CachedCompiler = null;
+				}
+
+				return CachedCompiler.Value;
+			}
+		}
+
+		// Enables the UWP platform and project file support in Unreal Build Tool
+		// @todo UWP: Remove this variable when UWP support is fully implemented
+		public static readonly bool bEnableUWPSupport = UnrealBuildTool.CommandLineContains("-uwp");
+
+		/// True if we should only build against the app-local CRT and /APPCONTAINER linker flag
+		public static readonly bool bBuildForStore = true;
+
+		/// True if we should only build against the app-local CRT and set the API-set flags
+		public static readonly bool bWinApiFamilyApp = true;
+
+		/** True if VS EnvDTE is available (false when building using Visual Studio Express) */
+		public static bool bHasVisualStudioDTE
+		{
+			get
+			{
+				return WindowsPlatform.bHasVisualStudioDTE;
+			}
+		}
+
+		public override bool RequiresDeployPrepAfterCompile()
+		{
+			return true;
+		}
+
+		protected override SDKStatus HasRequiredManualSDKInternal()
+		{
+			try
+			{
+				// @todo UWP: wire this up once the rest of cleanup has been resolved. right now, just leaving this to always return valid, like VC does.
+				//WinUWPToolChain.FindBaseVSToolPath();
+				return SDKStatus.Valid;
+			}
+			catch (BuildException)
+			{
+				return SDKStatus.Invalid;
+			}
+		}
+
+		/**
+		 * Returns VisualStudio common tools path for current compiler.
+		 * 
+		 * @return Common tools path.
+		 */
+		public static string GetVSComnToolsPath()
+		{
+			return GetVSComnToolsPath(Compiler);
+		}
+
+		/**
+		 * Returns VisualStudio common tools path for given compiler.
+		 * 
+		 * @param Compiler Compiler for which to return tools path.
+		 * 
+		 * @return Common tools path.
+		 */
+		public static string GetVSComnToolsPath(WindowsCompiler Compiler)
+		{
+			int VSVersion;
+
+			switch(Compiler)
+			{
+				case WindowsCompiler.VisualStudio2015:
+					VSVersion = 14;
+					break;
+				default:
+					throw new NotSupportedException("Not supported compiler.");
+			}
+
+			string[] PossibleRegPaths = new string[] {
+				@"Wow6432Node\Microsoft\VisualStudio",	// Non-express VS on 64-bit machine.
+				@"Microsoft\VisualStudio",				// Non-express VS on 32-bit machine.
+				@"Wow6432Node\Microsoft\WDExpress",		// Express VS on 64-bit machine.
+				@"Microsoft\WDExpress"					// Express VS on 32-bit machine.
+			};
+
+			string VSPath = null;
+
+			foreach(var PossibleRegPath in PossibleRegPaths)
+			{
+				VSPath = (string) Registry.GetValue(string.Format(@"HKEY_LOCAL_MACHINE\SOFTWARE\{0}\{1}.0", PossibleRegPath, VSVersion), "InstallDir", null);
+
+				if(VSPath != null)
+				{
+					break;
+				}
+			}
+
+			if(VSPath == null)
+			{
+				return null;
+			}
+
+			return new DirectoryInfo(Path.Combine(VSPath, "..", "Tools")).FullName;
+		}
+
+
+		/**
+		 *	Register the platform with the UEBuildPlatform class
+		 */
+		protected override void RegisterBuildPlatformInternal()
+		{
+			// Register this build platform for UWP
+			Log.TraceVerbose("        Registering for {0}", UnrealTargetPlatform.UWP.ToString());
+			UEBuildPlatform.RegisterBuildPlatform(UnrealTargetPlatform.UWP, this);
+			UEBuildPlatform.RegisterPlatformWithGroup(UnrealTargetPlatform.UWP, UnrealPlatformGroup.Microsoft);
+		}
+
+		/**
+		 *	Retrieve the CPPTargetPlatform for the given UnrealTargetPlatform
+		 *
+		 *	@param	InUnrealTargetPlatform		The UnrealTargetPlatform being build
+		 *	
+		 *	@return	CPPTargetPlatform			The CPPTargetPlatform to compile for
+		 */
+		public override CPPTargetPlatform GetCPPTargetPlatform(UnrealTargetPlatform InUnrealTargetPlatform)
+		{
+			switch (InUnrealTargetPlatform)
+			{
+				case UnrealTargetPlatform.UWP:
+					return CPPTargetPlatform.UWP;
+			}
+			throw new BuildException("UWPPlatform::GetCPPTargetPlatform: Invalid request for {0}", InUnrealTargetPlatform.ToString());
+		}
+
+		/**
+		 *	Get the extension to use for the given binary type
+		 *	
+		 *	@param	InBinaryType		The binrary type being built
+		 *	
+		 *	@return	string				The binary extenstion (ie 'exe' or 'dll')
+		 */
+		public override string GetBinaryExtension(UEBuildBinaryType InBinaryType)
+		{
+			switch (InBinaryType)
+			{
+				case UEBuildBinaryType.DynamicLinkLibrary: 
+					return ".dll";
+				case UEBuildBinaryType.Executable:
+					return ".exe";
+				case UEBuildBinaryType.StaticLibrary:
+					return ".lib";
+				case UEBuildBinaryType.Object:
+					return ".obj";
+				case UEBuildBinaryType.PrecompiledHeader:
+					return ".pch";
+			}
+			return base.GetBinaryExtension(InBinaryType);
+		}
+
+		
+		/// <summary>
+		/// When using a Visual Studio compiler, returns the version name as a string
+		/// </summary>
+		/// <returns>The Visual Studio compiler version name (e.g. "2012")</returns>
+		public static string GetVisualStudioCompilerVersionName()
+		{
+			switch( Compiler )
+			{ 
+				case WindowsCompiler.VisualStudio2015:
+					return "2015";
+				default:
+					throw new BuildException("Unexpected WindowsCompiler version for GetVisualStudioCompilerVersionName().  Either not using a Visual Studio compiler or switch block needs to be updated");
+			}
+		}
+
+
+		/**
+		 *	Get the extension to use for debug info for the given binary type
+		 *	
+		 *	@param	InBinaryType		The binary type being built
+		 *	
+		 *	@return	string				The debug info extension (i.e. 'pdb')
+		 */
+		public override string GetDebugInfoExtension( UEBuildBinaryType InBinaryType )
+		{
+			switch (InBinaryType)
+			{
+				case UEBuildBinaryType.DynamicLinkLibrary:
+				case UEBuildBinaryType.Executable:
+					return ".pdb";
+			}
+			return "";
+		}
+
+		/**
+		 *	Whether incremental linking should be used
+		 *	
+		 *	@param	InPlatform			The CPPTargetPlatform being built
+		 *	@param	InConfiguration		The CPPTargetConfiguration being built
+		 *	
+		 *	@return	bool	true if incremental linking should be used, false if not
+		 */
+		public override bool ShouldUseIncrementalLinking(CPPTargetPlatform Platform, CPPTargetConfiguration Configuration)
+		{
+			return (Configuration == CPPTargetConfiguration.Debug);
+		}
+
+		/**
+		 *	Whether PDB files should be used
+		 *	
+		 *	@param	InPlatform			The CPPTargetPlatform being built
+		 *	@param	InConfiguration		The CPPTargetConfiguration being built
+		 *	@param	bInCreateDebugInfo	true if debug info is getting create, false if not
+		 *	
+		 *	@return	bool	true if PDB files should be used, false if not
+		 */
+		public override bool ShouldUsePDBFiles(CPPTargetPlatform Platform, CPPTargetConfiguration Configuration, bool bCreateDebugInfo)
+		{
+			// Only supported on PC.
+			if (bCreateDebugInfo && ShouldUseIncrementalLinking(Platform, Configuration))
+			{
+				return true;
+			}
+			return false;
+		}
+
+		/**
+		 *	Whether the editor should be built for this platform or not
+		 *	
+		 *	@param	InPlatform		The UnrealTargetPlatform being built
+		 *	@param	InConfiguration	The UnrealTargetConfiguration being built
+		 *	@return	bool			true if the editor should be built, false if not
+		 */
+		public override bool ShouldNotBuildEditor(UnrealTargetPlatform InPlatform, UnrealTargetConfiguration InConfiguration)
+		{
+			return true;
+		}
+
+		public override bool BuildRequiresCookedData(UnrealTargetPlatform InPlatform, UnrealTargetConfiguration InConfiguration)
+		{
+			return false;
+		}
+
+		public override void ResetBuildConfiguration(UnrealTargetPlatform InPlatform, UnrealTargetConfiguration InConfiguration)
+		{
+			UEBuildConfiguration.bCompileICU = true;
+		}
+
+
+		public override void ModifyNewlyLoadedModule(UEBuildModule InModule, TargetInfo Target)
+		{
+			if ((Target.Platform == UnrealTargetPlatform.UWP))
+			{
+				if (InModule.ToString() == "Core")
+				{
+					InModule.AddPrivateDependencyModule("UWPSDK");
+				}
+				else if (InModule.ToString() == "Engine")
+				{
+					InModule.AddPrivateDependencyModule("zlib");
+					InModule.AddPrivateDependencyModule("UElibPNG");
+					InModule.AddPublicDependencyModule("UEOgg");
+					InModule.AddPublicDependencyModule("Vorbis");
+				}
+				else if (InModule.ToString() == "D3D11RHI")
+				{
+					InModule.AddPublicDefinition("D3D11_WITH_DWMAPI=0");
+					InModule.AddPublicDefinition("WITH_DX_PERF=0");
+				}
+				else if (InModule.ToString() == "DX11")
+				{
+					// Clear out all the Windows include paths and libraries...
+					// The UWPSDK module handles proper paths and libs for UWP.
+					// However, the D3D11RHI module will include the DX11 module.
+					InModule.ClearPublicIncludePaths();
+					InModule.ClearPublicLibraryPaths();
+					InModule.ClearPublicAdditionalLibraries();
+					InModule.RemovePublicDefinition("WITH_D3DX_LIBS=1");
+					InModule.AddPublicDefinition("WITH_D3DX_LIBS=0");
+					InModule.RemovePublicAdditionalLibrary("X3DAudio.lib");
+					InModule.RemovePublicAdditionalLibrary("XAPOFX.lib");
+				}
+				else if (InModule.ToString() == "XAudio2")
+				{
+					InModule.AddPublicDefinition("XAUDIO_SUPPORTS_XMA2WAVEFORMATEX=0");
+					InModule.AddPublicDefinition("XAUDIO_SUPPORTS_DEVICE_DETAILS=0");
+					InModule.AddPublicDefinition("XAUDIO2_SUPPORTS_MUSIC=0");
+					InModule.AddPublicDefinition("XAUDIO2_SUPPORTS_SENDLIST=0");
+					InModule.AddPublicAdditionalLibrary("XAudio2.lib");
+				}
+				else if (InModule.ToString() == "DX11Audio")
+				{
+					InModule.RemovePublicAdditionalLibrary("X3DAudio.lib");
+					InModule.RemovePublicAdditionalLibrary("XAPOFX.lib");
+				}
+			}
+		}
+
+		/**
+		 *	Setup the target environment for building
+		 *	
+		 *	@param	InBuildTarget		The target being built
+		 */
+		public override void SetUpEnvironment(UEBuildTarget InBuildTarget)
+		{
+
+			InBuildTarget.GlobalCompileEnvironment.Config.Definitions.Add("PLATFORM_DESKTOP=0");
+			InBuildTarget.GlobalCompileEnvironment.Config.Definitions.Add("PLATFORM_64BITS=1");
+			InBuildTarget.GlobalCompileEnvironment.Config.Definitions.Add("EXCEPTIONS_DISABLED=0");
+
+			InBuildTarget.GlobalCompileEnvironment.Config.Definitions.Add("_WIN32_WINNT=0x0A00");
+			InBuildTarget.GlobalCompileEnvironment.Config.Definitions.Add("WINVER=0x0A00");
+
+			InBuildTarget.GlobalCompileEnvironment.Config.Definitions.Add("PLATFORM_UWP=1");
+			InBuildTarget.GlobalCompileEnvironment.Config.Definitions.Add("UWP=1");	// @todo UWP: This used to be "WINUAP=1".  Need to verify that 'UWP' is the correct define that we want here.
+
+			if (bWinApiFamilyApp)
+			{
+				InBuildTarget.GlobalCompileEnvironment.Config.Definitions.Add("WINAPI_FAMILY=WINAPI_FAMILY_APP");
+			}
+
+			// @todo UWP: UE4 is non-compliant when it comes to use of %s and %S
+			// Previously %s meant "the current character set" and %S meant "the other one".
+			// Now %s means multibyte and %S means wide. %Ts means "natural width".
+			// Reverting this behavior until the UE4 source catches up.
+			InBuildTarget.GlobalCompileEnvironment.Config.Definitions.Add("_CRT_STDIO_LEGACY_WIDE_SPECIFIERS=1");
+
+			// No D3DX on UWP!
+			InBuildTarget.GlobalCompileEnvironment.Config.Definitions.Add("NO_D3DX_LIBS=1");
+
+
+			if (InBuildTarget.Rules != null)
+			{
+				// Explicitly exclude the MS C++ runtime libraries we're not using, to ensure other libraries we link with use the same
+				// runtime library as the engine.
+				if (InBuildTarget.Configuration == UnrealTargetConfiguration.Debug && BuildConfiguration.bDebugBuildsActuallyUseDebugCRT)
+				{
+					InBuildTarget.GlobalLinkEnvironment.Config.ExcludedLibraries.Add("MSVCRT");
+					InBuildTarget.GlobalLinkEnvironment.Config.ExcludedLibraries.Add("MSVCPRT");
+				}
+				else
+				{
+					InBuildTarget.GlobalLinkEnvironment.Config.ExcludedLibraries.Add("MSVCRTD");
+					InBuildTarget.GlobalLinkEnvironment.Config.ExcludedLibraries.Add("MSVCPRTD");
+				}
+				InBuildTarget.GlobalLinkEnvironment.Config.ExcludedLibraries.Add("LIBC");
+				InBuildTarget.GlobalLinkEnvironment.Config.ExcludedLibraries.Add("LIBCMT");
+				InBuildTarget.GlobalLinkEnvironment.Config.ExcludedLibraries.Add("LIBCPMT");
+				InBuildTarget.GlobalLinkEnvironment.Config.ExcludedLibraries.Add("LIBCP");
+				InBuildTarget.GlobalLinkEnvironment.Config.ExcludedLibraries.Add("LIBCD");
+				InBuildTarget.GlobalLinkEnvironment.Config.ExcludedLibraries.Add("LIBCMTD");
+				InBuildTarget.GlobalLinkEnvironment.Config.ExcludedLibraries.Add("LIBCPMTD");
+				InBuildTarget.GlobalLinkEnvironment.Config.ExcludedLibraries.Add("LIBCPD");
+			}
+
+			InBuildTarget.GlobalLinkEnvironment.Config.AdditionalLibraries.Add("mincore.lib");
+			InBuildTarget.GlobalLinkEnvironment.Config.AdditionalLibraries.Add("dloadhelper.lib");
+		   
+			// Disable Simplygon support if compiling against the NULL RHI.
+			if (InBuildTarget.GlobalCompileEnvironment.Config.Definitions.Contains("USE_NULL_RHI=1"))
+			{
+				UEBuildConfiguration.bCompileSimplygon = false;
+			}
+		}
+
+		/**
+		 *	Setup the configuration environment for building
+		 *	
+		 *	@param	InBuildTarget		The target being built
+		 */
+		public override void SetUpConfigurationEnvironment(UEBuildTarget InBuildTarget)
+		{
+			// Determine the C++ compile/link configuration based on the Unreal configuration.
+			CPPTargetConfiguration CompileConfiguration;
+			UnrealTargetConfiguration CheckConfig = InBuildTarget.Configuration;
+			switch (CheckConfig)
+			{
+				default:
+				case UnrealTargetConfiguration.Debug:
+					CompileConfiguration = CPPTargetConfiguration.Debug;
+					if( BuildConfiguration.bDebugBuildsActuallyUseDebugCRT )
+					{ 
+						InBuildTarget.GlobalCompileEnvironment.Config.Definitions.Add("_DEBUG=1"); // the engine doesn't use this, but lots of 3rd party stuff does
+					}
+					else
+					{
+						InBuildTarget.GlobalCompileEnvironment.Config.Definitions.Add("NDEBUG=1"); // the engine doesn't use this, but lots of 3rd party stuff does
+					}
+					InBuildTarget.GlobalCompileEnvironment.Config.Definitions.Add("UE_BUILD_DEBUG=1");
+					break;
+				case UnrealTargetConfiguration.DebugGame:
+					// Default to Development; can be overriden by individual modules.
+				case UnrealTargetConfiguration.Development:
+					CompileConfiguration = CPPTargetConfiguration.Development;
+					InBuildTarget.GlobalCompileEnvironment.Config.Definitions.Add("NDEBUG=1"); // the engine doesn't use this, but lots of 3rd party stuff does
+					InBuildTarget.GlobalCompileEnvironment.Config.Definitions.Add("UE_BUILD_DEVELOPMENT=1");
+					break;
+				case UnrealTargetConfiguration.Shipping:
+					CompileConfiguration = CPPTargetConfiguration.Shipping;
+					InBuildTarget.GlobalCompileEnvironment.Config.Definitions.Add("NDEBUG=1"); // the engine doesn't use this, but lots of 3rd party stuff does
+					InBuildTarget.GlobalCompileEnvironment.Config.Definitions.Add("UE_BUILD_SHIPPING=1");
+					break;
+				case UnrealTargetConfiguration.Test:
+					CompileConfiguration = CPPTargetConfiguration.Shipping;
+					InBuildTarget.GlobalCompileEnvironment.Config.Definitions.Add("NDEBUG=1"); // the engine doesn't use this, but lots of 3rd party stuff does
+					InBuildTarget.GlobalCompileEnvironment.Config.Definitions.Add("UE_BUILD_TEST=1");
+					break;
+			}
+
+			// Set up the global C++ compilation and link environment.
+			InBuildTarget.GlobalCompileEnvironment.Config.Target.Configuration = CompileConfiguration;
+			InBuildTarget.GlobalLinkEnvironment.Config.Target.Configuration    = CompileConfiguration;
+
+			// Create debug info based on the heuristics specified by the user.
+			InBuildTarget.GlobalCompileEnvironment.Config.bCreateDebugInfo =
+				!BuildConfiguration.bDisableDebugInfo && ShouldCreateDebugInfo(InBuildTarget.Platform, CheckConfig);
+
+			// NOTE: Even when debug info is turned off, we currently force the linker to generate debug info
+			//       anyway on Visual C++ platforms.  This will cause a PDB file to be generated with symbols
+			//       for most of the classes and function/method names, so that crashes still yield somewhat
+			//       useful call stacks, even though compiler-generate debug info may be disabled.  This gives
+			//       us much of the build-time savings of fully-disabled debug info, without giving up call
+			//       data completely.
+			InBuildTarget.GlobalLinkEnvironment.Config.bCreateDebugInfo = true;
+		}
+
+		/**
+		 *	Whether this platform should create debug information or not
+		 *	
+		 *	@param	InPlatform			The UnrealTargetPlatform being built
+		 *	@param	InConfiguration		The UnrealTargetConfiguration being built
+		 *	
+		 *	@return	bool				true if debug info should be generated, false if not
+		 */
+		public override bool ShouldCreateDebugInfo(UnrealTargetPlatform Platform, UnrealTargetConfiguration Configuration)
+		{
+			switch (Configuration)
+			{
+				case UnrealTargetConfiguration.Development: 
+				case UnrealTargetConfiguration.Shipping: 
+				case UnrealTargetConfiguration.Test: 
+					return !BuildConfiguration.bOmitPCDebugInfoInDevelopment;
+				case UnrealTargetConfiguration.DebugGame:
+				case UnrealTargetConfiguration.Debug:
+				default: 
+					return true;
+			};
+		}
+
+		/**
+		 *	Return whether this platform has uniquely named binaries across multiple games
+		 */
+		public override bool HasUniqueBinaries()
+		{
+			// Windows applications have many shared binaries between games
+			return false;
+		}
+	}
+}
diff --git a/Engine/Source/Programs/UnrealBuildTool/UWP/UWPDeploy.cs b/Engine/Source/Programs/UnrealBuildTool/UWP/UWPDeploy.cs
new file mode 100644
--- /dev/null
+++ b/Engine/Source/Programs/UnrealBuildTool/UWP/UWPDeploy.cs
@@ -0,0 +1,272 @@
+﻿// Copyright 1998-2014 Epic Games, Inc. All Rights Reserved.
+
+using System;
+using System.Collections.Generic;
+using System.Text;
+using System.IO;
+using System.Diagnostics;
+
+// @todo UWP: this file is a work in progress and is not yet complete for the F5 scenario for UWP
+
+namespace UnrealBuildTool
+{
+	/// <summary>
+	///  Base class to handle deploy of a target for a given platform
+	/// </summary>
+	class UWPDeploy : UEBuildDeploy
+	{
+		/**
+		 *	Register the platform with the UEBuildDeploy class
+		 */
+		public override void RegisterBuildDeploy()
+		{
+			// Register this deployment handle for UWP
+			Log.TraceVerbose("        Registering for {0}", UnrealTargetPlatform.UWP.ToString());
+			UEBuildDeploy.RegisterBuildDeploy(UnrealTargetPlatform.UWP, this);
+		}
+
+		/**
+		 *  Utility function to delete a file
+		 */
+		void DeployHelper_DeleteFile(string InFileToDelete)
+		{
+			Log.TraceInformation("UWPDeploy.DeployHelper_DeleteFile({0})", InFileToDelete);
+			if (File.Exists(InFileToDelete) == true)
+			{
+				FileAttributes attributes = File.GetAttributes(InFileToDelete);
+				if ((attributes & FileAttributes.ReadOnly) == FileAttributes.ReadOnly)
+				{
+					attributes &= ~FileAttributes.ReadOnly;
+					File.SetAttributes(InFileToDelete, attributes);
+				}
+				File.Delete(InFileToDelete);
+			}
+		}
+
+		/**
+		 *	Copy the contents of the given source directory to the given destination directory
+		 *	
+		 */
+		bool CopySourceToDestDir(string InSourceDirectory, string InDestinationDirectory, string InWildCard,
+			bool bInIncludeSubDirectories, bool bInRemoveDestinationOrphans)
+		{
+			Log.TraceInformation("UWPDeploy.CopySourceToDestDir({0}, {1}, {2},...)", InSourceDirectory, InDestinationDirectory, InWildCard);
+			if (Directory.Exists(InSourceDirectory) == false)
+			{
+				Log.TraceInformation("Warning: CopySourceToDestDir - SourceDirectory does not exist: {0}", InSourceDirectory);
+				return false;
+			}
+
+			// Make sure the destination directory exists!
+			Directory.CreateDirectory(InDestinationDirectory);
+
+			SearchOption OptionToSearch = bInIncludeSubDirectories ? SearchOption.AllDirectories : SearchOption.TopDirectoryOnly;
+
+			var SourceDirs = new List<string>(Directory.GetDirectories(InSourceDirectory, "*.*", OptionToSearch));
+			foreach (string SourceDir in SourceDirs)
+			{
+				string SubDir = SourceDir.Replace(InSourceDirectory, "");
+				string DestDir = InDestinationDirectory + SubDir;
+				Directory.CreateDirectory(DestDir);
+			}
+
+			var SourceFiles = new List<string>(Directory.GetFiles(InSourceDirectory, InWildCard, OptionToSearch));
+			var DestFiles = new List<string>(Directory.GetFiles(InDestinationDirectory, InWildCard, OptionToSearch));
+
+			// Keep a list of the files in the source directory... without the source path
+			List<string> FilesInSource = new List<string>();
+
+			// Copy all the source files that are newer...
+			foreach (string SourceFile in SourceFiles)
+			{
+				string Filename = SourceFile.Replace(InSourceDirectory, "");
+				FilesInSource.Add(Filename.ToUpperInvariant());
+				string DestFile = InDestinationDirectory + Filename;
+
+				System.DateTime SourceTime = File.GetLastWriteTime(SourceFile);
+				System.DateTime DestTime = File.GetLastWriteTime(DestFile);
+
+				if (SourceTime > DestTime)
+				{
+					try
+					{
+						DeployHelper_DeleteFile(DestFile);
+						File.Copy(SourceFile, DestFile, true);
+					}
+					catch (Exception exceptionMessage)
+					{
+						Log.TraceInformation("Failed to copy {0} to deployment: {1}", SourceFile, exceptionMessage);
+					}
+				}
+			}
+
+			if (bInRemoveDestinationOrphans == true)
+			{
+				// If requested, delete any destination files that do not have a corresponding
+				// file in the source directory
+				foreach (string DestFile in DestFiles)
+				{
+					string DestFilename = DestFile.Replace(InDestinationDirectory, "");
+					if (FilesInSource.Contains(DestFilename.ToUpperInvariant()) == false)
+					{
+						Log.TraceInformation("Destination file does not exist in Source - DELETING: {0}", DestFile);
+						FileAttributes attributes = File.GetAttributes(DestFile);
+						try
+						{
+							DeployHelper_DeleteFile(DestFile);
+						}
+						catch (Exception exceptionMessage)
+						{
+							Log.TraceInformation("Failed to delete {0} from deployment: {1}", DestFile, exceptionMessage);
+						}
+					}
+				}
+			}
+
+			return true;
+		}
+
+
+		/**
+		 *	Helper function for copying files
+		 */
+		void CopyFile(string InSource, string InDest, bool bForce)
+		{
+			Log.TraceInformation("UWPDeploy.CopyFile({0}, {1}, {2},...)", InSource, InDest, bForce);
+			if (File.Exists(InSource) == true)
+			{
+				if (bForce == true)
+				{
+					if (File.Exists(InDest) == true)
+					{
+						DeployHelper_DeleteFile(InDest);
+					}
+				}
+				File.Copy(InSource, InDest, true);
+			}
+			else
+			{
+				Log.TraceInformation("UWPDeploy: File didn't exist - {0}", InSource);
+			}
+		}
+
+		/**
+		 *	Helper function for copying a tree files
+		 */
+		//void CopyDirectory(string InSource, string InDest, bool bForce, bool bRecurse)
+		//{
+		//    if (Directory.Exists(InSource))
+		//    {
+		//        if (!Directory.Exists(InDest))
+		//        {
+		//            Directory.CreateDirectory(InDest);
+		//        }
+
+		//        // Copy all files
+		//        string[] FilesInDir = Directory.GetFiles(InSource);
+		//        foreach (string FileSourcePath in FilesInDir)
+		//        {
+		//            string FileDestPath = Path.Combine(InDest, Path.GetFileName(FileSourcePath));
+		//            if (bForce == true)
+		//            {
+		//                if (File.Exists(FileDestPath) == true)
+		//                {
+		//                    DeployHelper_DeleteFile(FileDestPath);
+		//                }
+		//            }
+		//            File.Copy(FileSourcePath, FileDestPath, true);
+		//        }
+
+		//        // Recurse sub directories
+		//        string[] DirsInDir = Directory.GetDirectories(InSource);
+		//        foreach (string DirSourcePath in DirsInDir)
+		//        {
+		//            string DirName = Path.GetFileName(DirSourcePath);
+		//            string DirDestPath = Path.Combine(InDest, DirName);
+		//            CopyDirectory(DirSourcePath, DirDestPath, bForce, bRecurse);
+		//        }
+		//    }
+		//}
+
+
+		public bool PrepForUATPackageOrDeploy(string InProjectName, string InProjectDirectory, string InExecutablePath, string InEngineDir, bool bForDistribution, string CookFlavor)
+		{
+			// Much of the UWPProjectGenerator code assumes that the path will be set to Engine/Source
+			string PreviousDirectory = Directory.GetCurrentDirectory();
+			string PreviousRelativeEnginePath = BuildConfiguration.RelativeEnginePath;
+
+			Directory.SetCurrentDirectory(Path.Combine(InEngineDir, "Binaries"));
+			BuildConfiguration.RelativeEnginePath = "../../Engine";
+
+			string ProjectBinaryFolder = new FileInfo(InExecutablePath).DirectoryName;
+			string EngineSourceRelativeBinaryPath = Utils.MakePathRelativeTo(ProjectBinaryFolder, Path.Combine(InEngineDir, "Source"));
+			bool bIsUE4Game = new FileInfo(InExecutablePath).Name.StartsWith("UE4Game", StringComparison.InvariantCultureIgnoreCase);
+
+			UWPProjectGenerator ProjectGen = new UWPProjectGenerator();
+			ProjectGen.GenerateDeployFiles(InProjectName, ProjectBinaryFolder, bIsUE4Game, EngineSourceRelativeBinaryPath);
+
+			// back to root for UAT
+			Directory.SetCurrentDirectory(PreviousDirectory);
+			BuildConfiguration.RelativeEnginePath = PreviousRelativeEnginePath;
+
+			return true;
+		}
+
+
+
+		public override bool PrepTargetForDeployment(UEBuildTarget InTarget)
+		{
+			string InAppName = InTarget.AppName;
+			Log.TraceInformation("Prepping {0} for deployment to {1}", InAppName, InTarget.Platform.ToString());
+			System.DateTime PrepDeployStartTime = DateTime.UtcNow;
+
+			string TargetFilename = RulesCompiler.GetTargetFilename(InAppName);
+			string ProjectSourceFolder = new FileInfo(TargetFilename).DirectoryName + "/";
+
+			string RelativeTargetDirectory;
+			string EngineSourceRelativeBinaryPath;
+			UWPProjectGenerator.GetTargetUWPPaths(InAppName, InTarget.Rules, out EngineSourceRelativeBinaryPath, out RelativeTargetDirectory);
+
+			PrepForUATPackageOrDeploy(InAppName, InTarget.ProjectDirectory, InTarget.OutputPath, BuildConfiguration.RelativeEnginePath, false, "");
+
+			// TODO - richiem - restore this if we find that it's needed.
+			//Log.TraceInformation("...copying the CELL dll...");
+			//string CELLPath = "../../../Engine/Source/" + UEBuildConfiguration.UEThirdPartySourceDirectory + "CELL/lib/win64/";
+			//string CELLPathRelease = CELLPath + "Release/";
+			//string CELLDllRelease = "CommonEventLoggingLibrary.dll";
+			//string CELLPathDebug = CELLPath + "Debug/";
+			//string CELLDllDebug = "CommonEventLoggingLibraryd.dll";
+
+			//CopyFile(EngineSourceRelativeBinaryPath + CELLPathRelease + CELLDllRelease, EngineSourceRelativeBinaryPath + CELLDllRelease, true);
+			//CopyFile(EngineSourceRelativeBinaryPath + CELLPathDebug + CELLDllDebug, EngineSourceRelativeBinaryPath + CELLDllDebug, true);
+
+			//string XSAPIPath            = EngineSourceRelativeBinaryPath + "../../../Engine/Source/ThirdParty/XSAPI/lib/";
+			//string XboxServicesConfig   = "xboxservices.config";
+			//string DesktopLogin         = "DesktopLogin.exe";
+
+			//Log.TraceInformation("...copying xboxservices.config");
+			//CopyFile(XSAPIPath + XboxServicesConfig, EngineSourceRelativeBinaryPath + XboxServicesConfig, true);
+			//Log.TraceInformation("...copying DesktopLogin.exe...");
+			//CopyFile(XSAPIPath + DesktopLogin, EngineSourceRelativeBinaryPath + DesktopLogin, true);
+
+			// TODO - richiem - restore this if we find that it's needed.
+			//if (InTarget.Configuration == UnrealTargetConfiguration.Development)
+			//{
+			//    Log.TraceInformation("...copying AutoLogin...");
+			//    string AutoLoginPath    = EngineSourceRelativeBinaryPath + "../../../Engine/Source/Tools/AutoLogin/";
+			//    string AutoLoginDLL     = "AutoLogin.dll";
+			//    string AutoLoginBat     = "AutoLogin.bat";
+			//    CopyFile(AutoLoginPath + "bin/Debug/" + AutoLoginDLL, EngineSourceRelativeBinaryPath + AutoLoginDLL, true);
+			//    CopyFile(AutoLoginPath + AutoLoginBat, EngineSourceRelativeBinaryPath + AutoLoginBat, true);
+			//}
+
+
+
+			// Log out the time taken to deploy...
+			double PrepDeployDuration = (DateTime.UtcNow - PrepDeployStartTime).TotalSeconds;
+			Log.TraceInformation("UWP deployment preparation took {0:0.00} seconds", PrepDeployDuration);
+
+			return true;
+		}
+	}
+}
diff --git a/Engine/Source/Programs/UnrealBuildTool/UWP/UWPProjectGenerator.cs b/Engine/Source/Programs/UnrealBuildTool/UWP/UWPProjectGenerator.cs
new file mode 100644
--- /dev/null
+++ b/Engine/Source/Programs/UnrealBuildTool/UWP/UWPProjectGenerator.cs
@@ -0,0 +1,408 @@
+// Copyright 1998-2014 Epic Games, Inc. All Rights Reserved.
+
+using System;
+using System.Collections.Generic;
+using System.Text;
+using System.Diagnostics;
+using System.IO;
+
+// @todo UWP: this file is a work in progress and is not yet complete for the F5 scenario for UWP
+
+namespace UnrealBuildTool
+{
+	/**
+	 *	Base class for platform-specific project generators 
+	 */
+	public class UWPProjectGenerator : UEPlatformProjectGenerator
+	{
+		/**
+		 *	Register the platform with the UEPlatformProjectGenerator class
+		 */
+		public override void RegisterPlatformProjectGenerator()
+		{
+			if( UWPPlatform.bEnableUWPSupport )
+			{ 
+				// Register this project generator for UWP
+				Log.TraceVerbose("        Registering for {0}", UnrealTargetPlatform.UWP.ToString());
+			
+				UEPlatformProjectGenerator.RegisterPlatformProjectGenerator(UnrealTargetPlatform.UWP, this);
+			}
+		}
+
+		public static void GetTargetUWPPaths(string InTargetName, TargetRules InTargetRules,
+			out string OutEngineSourceRelativeBinaryPath, out string OutRelativeTargetPath)
+		{
+			OutEngineSourceRelativeBinaryPath = "";
+			OutRelativeTargetPath = "";
+
+			string TargetFilename = RulesCompiler.GetTargetFilename(InTargetName);
+
+			string ProjectSourceFolder = new FileInfo(TargetFilename).DirectoryName;
+
+			string EnginePath = Path.Combine(ProjectFileGenerator.EngineRelativePath);
+			string EngineSourcePath = Path.Combine(EnginePath, "Source");
+			string RelativeTargetFilename = Utils.MakePathRelativeTo(TargetFilename, EngineSourcePath);
+
+			if ((RelativeTargetFilename.StartsWith("..") == false) && (RelativeTargetFilename.Contains(":") == false))
+			{
+				// This target must be UNDER Engine/Source...
+				RelativeTargetFilename = Path.Combine(EngineSourcePath, RelativeTargetFilename);
+			}
+			RelativeTargetFilename = RelativeTargetFilename.Replace("\\", "/");
+			EnginePath = EnginePath.Replace("\\", "/");
+
+			Int32 LastSourceIdx = RelativeTargetFilename.LastIndexOf("Source");
+			if (LastSourceIdx != -1)
+			{
+				RelativeTargetFilename = RelativeTargetFilename.Substring(0, LastSourceIdx);
+			}
+			else
+			{
+				RelativeTargetFilename = "";
+			}
+
+			OutRelativeTargetPath = RelativeTargetFilename;
+
+			if (InTargetRules.bOutputToEngineBinaries)
+			{
+				RelativeTargetFilename = EnginePath;
+			}
+			OutEngineSourceRelativeBinaryPath = Path.Combine(RelativeTargetFilename, "Binaries/UWP/");
+			OutEngineSourceRelativeBinaryPath = OutEngineSourceRelativeBinaryPath.Replace("\\", "/");
+		}
+
+		///
+		///	VisualStudio project generation functions
+		///	
+		/**
+		 *	Whether this build platform has native support for VisualStudio
+		 *	
+		 *	@param	InPlatform			The UnrealTargetPlatform being built
+		 *	@param	InConfiguration		The UnrealTargetConfiguration being built
+		 *	
+		 *	@return	bool				true if native VisualStudio support (or custom VSI) is available
+		 */
+		public override bool HasVisualStudioSupport(UnrealTargetPlatform InPlatform, UnrealTargetConfiguration InConfiguration)
+		{
+			return false;
+		}
+
+		/**
+		 * Get whether this platform deploys 
+		 * 
+		 * @return	bool		true if the 'Deploy' option should be enabled
+		 */
+		public override bool GetVisualStudioDeploymentEnabled(UnrealTargetPlatform InPlatform, UnrealTargetConfiguration InConfiguration)
+		{
+			return true;
+		}
+
+		private Dictionary<string, string> GameToBuildFolder = new Dictionary<string, string>();
+
+
+		protected void GeneratePackageAppXManifest(string InGameName, bool bIsUE4Game, string InManifestType, string InGameBinaries, string InGameIntermediates, string InOutputPath)
+		{
+			string BuildFolder = Path.Combine(InOutputPath, "../../Build/UWP");
+			string CheckPath;
+			if (GameToBuildFolder.TryGetValue(InGameName, out CheckPath) == false)
+			{
+				GameToBuildFolder.Add(InGameName, BuildFolder);
+			}
+
+			string OutputPackageAppXManifestPath = Path.Combine(BuildFolder, InManifestType + ".AppxManifest");
+
+
+			// first look for a per-game template
+			string PackageAppXManifestPath = OutputPackageAppXManifestPath + ".template";
+			if (!File.Exists(PackageAppXManifestPath))
+			{
+				// then fallback to shared template
+				PackageAppXManifestPath = "../../Engine/Build/UWP/Templates/" + InManifestType + ".appxmanifest.template";
+			}
+
+			GenerateFileFromTemplate(PackageAppXManifestPath, OutputPackageAppXManifestPath, bIsUE4Game, InGameName, "", InGameBinaries, InGameIntermediates);
+
+			// Generate the package name...
+			// todo - richiem - this is next to fix
+			string PackageName = "foobar"; //= XboxOneDeploy.XboxOnePackageNameHelper_Run(OutputPackageAppXManifestPath);
+			if (string.IsNullOrEmpty(PackageName) == false)
+			{
+				if (GameToBuildFolder.TryGetValue(InGameName, out CheckPath) == false)
+				{
+					GameToBuildFolder.Add(InGameName, PackageName);
+				}
+			}
+			else
+			{
+				throw new BuildException("Error: Failed to find package name for {0}", OutputPackageAppXManifestPath);
+			}
+		}
+
+		protected void GeneratePackagePriConfig(string InGameName, bool bIsUE4Game, string InGameBinaries, string InGameIntermediates, string InOutputPath)
+		{
+			string BuildFolder = Path.Combine(InOutputPath, "../../Build/UWP");
+
+			string OutputPackagePriConfigPath = Path.Combine(BuildFolder, "Package.PriConfig");
+
+			// first look for a per-game template
+			string PackagePriConfigPath = OutputPackagePriConfigPath + ".template";
+			if (!File.Exists(PackagePriConfigPath))
+			{
+				// then fallback to shared template
+				PackagePriConfigPath = "../../Engine/Build/UWP/Templates/Package.priconfig.template";
+			}
+
+			GenerateFileFromTemplate(PackagePriConfigPath, OutputPackagePriConfigPath, bIsUE4Game, InGameName, "", InGameBinaries, InGameIntermediates);
+		}
+
+		protected void GeneratePackageResources(string InGameName, bool bIsUE4Game, string InGameBinaries, string InGameIntermediates, string InOutputPath)
+		{
+			string BuildFolder = Path.Combine(InOutputPath, "../../Build/UWP");
+
+			string OutputPackageResourcesPath = Path.Combine(BuildFolder, "Resources", "Resources.resw");
+
+			// first look for a per-game template
+			string PackageResourcesPath = Path.Combine(BuildFolder, "Resources", "Package.resources.template");
+			if (!File.Exists(PackageResourcesPath))
+			{
+				// then fallback to shared template
+				PackageResourcesPath = "../../Engine/Build/UWP/Templates/Package.resources.template";
+			}
+
+			// Only generate the example language if the resources haven't been setup at all previously
+			if (!File.Exists(OutputPackageResourcesPath))
+			{
+				GenerateFileFromTemplate(PackageResourcesPath, Path.Combine(BuildFolder, "Resources", "en-US", "Resources.resw"), bIsUE4Game, InGameName, "", InGameBinaries, InGameIntermediates);
+			}
+			GenerateFileFromTemplate(PackageResourcesPath, OutputPackageResourcesPath, bIsUE4Game, InGameName, "", InGameBinaries, InGameIntermediates);
+		}
+
+		protected void GenerateFileFromTemplate(string InTemplateFile, string InOutputFile, bool bIsUE4Game, string InGameName, string InGameConfiguration, string InGameBinaries, string InGameIntermediates, string InGuidString = "", Dictionary<string, string> InExtraReplacements = null)
+		{
+			String ExecutableName = bIsUE4Game ? "UE4Game" : InGameName;
+
+			StringBuilder outputContents = new StringBuilder();
+			using (StreamReader reader = new StreamReader(InTemplateFile))
+			{
+				string LineStr = null;
+				while (reader.Peek() != -1)
+				{
+					LineStr = reader.ReadLine();
+					if (LineStr.Contains("%GAME%"))
+					{
+						if (LineStr.Contains("Executable") || LineStr.Contains("EntryPoint"))
+						{
+							LineStr = LineStr.Replace("%GAME%", ExecutableName);
+						}
+						else
+						{
+							LineStr = LineStr.Replace("%GAME%", InGameName);
+						}
+					}
+					if (LineStr.Contains("%PROJECT%"))
+					{
+						if (bIsUE4Game)
+						{
+							LineStr = LineStr.Replace("%PROJECT%", "Engine");
+						}
+						else
+						{
+							LineStr = LineStr.Replace("%PROJECT%", InGameName);
+						}
+					}
+
+					if (LineStr.Contains("%CONFIG%"))
+					{
+						LineStr = LineStr.Replace("%CONFIG%", InGameConfiguration);
+					}
+
+					if (LineStr.Contains("%DISPLAYCONFIG%"))
+					{
+						if (string.IsNullOrEmpty(InGameConfiguration))
+						{
+							LineStr = LineStr.Replace("%DISPLAYCONFIG%", "");
+						}
+						else
+						{
+							LineStr = LineStr.Replace("%DISPLAYCONFIG%", " (" + InGameConfiguration + ")");
+						}
+					}
+
+					if (LineStr.Contains("%GAMEBINARIES%"))
+					{
+						LineStr = LineStr.Replace("%GAMEBINARIES%", InGameBinaries);
+					}
+
+					if (LineStr.Contains("%GAMEBINARIESFOLDER%"))
+					{
+						LineStr = LineStr.Replace("%GAMEBINARIESFOLDER%", InGameBinaries);
+					}
+
+					if (LineStr.Contains("%GAMEINTERMEDIATESFOLDER%"))
+					{
+						LineStr = LineStr.Replace("%GAMEINTERMEDIATESFOLDER%", InGameIntermediates);
+					}
+
+					if (LineStr.Contains("%GUID%"))
+					{
+						LineStr = LineStr.Replace("%GUID%", InGuidString);
+					}
+
+					if (InExtraReplacements != null)
+					{
+						foreach (KeyValuePair<string, string> Replacement in InExtraReplacements)
+						{
+							if (LineStr.Contains(Replacement.Key))
+							{
+								LineStr = LineStr.Replace(Replacement.Key, Replacement.Value);
+							}
+						}
+					}
+
+					outputContents.AppendLine(LineStr);
+				}
+			}
+
+			if (outputContents.Length > 0)
+			{
+				bool bFileNeedsSave = true;
+
+				if (File.Exists(InOutputFile))
+				{
+					// Read in the original file completely
+					string LoadedFileContent = null;
+					var FileAlreadyExists = File.Exists(InOutputFile);
+					if (FileAlreadyExists)
+					{
+						try
+						{
+							LoadedFileContent = File.ReadAllText(InOutputFile);
+						}
+						catch (Exception)
+						{
+							Log.TraceInformation("Error while trying to load existing file {0}.  Ignored.", InOutputFile);
+						}
+					}
+
+					// Don't bother saving anything out if the new file content is the same as the old file's content
+					if (LoadedFileContent != null)
+					{
+						var bIgnoreProjectFileWhitespaces = true;
+						if (ProjectFileComparer.CompareOrdinalIgnoreCase(LoadedFileContent, outputContents.ToString(), bIgnoreProjectFileWhitespaces) == 0)
+						{
+							// Exact match!
+							bFileNeedsSave = false;
+						}
+
+						if (!bFileNeedsSave)
+						{
+							Log.TraceVerbose("Skipped saving {0} because contents haven't changed.", Path.GetFileName(InOutputFile));
+						}
+					}
+				}
+
+				if (bFileNeedsSave)
+				{
+					// Save the file
+					try
+					{
+						Directory.CreateDirectory(Path.GetDirectoryName(InOutputFile));
+						File.WriteAllText(InOutputFile, outputContents.ToString(), Encoding.UTF8);
+						Log.TraceVerbose("Saved {0}", Path.GetFileName(InOutputFile));
+					}
+					catch (Exception)
+					{
+						// Unable to write to the project file.
+						Log.TraceInformation("Error while trying to write file {0}.  The file is probably read-only.", InOutputFile);
+					}
+				}
+			}
+		}
+
+
+		protected void CopyDefaultImageFiles(string InOutputPath)
+		{
+			string OutputImageFolder = Path.Combine(InOutputPath, "../../Build/UWP/Resources");
+			if (Directory.Exists(OutputImageFolder) == false)
+			{
+				Directory.CreateDirectory(OutputImageFolder);
+			}
+
+			string[] ImageList = new string[] { "Logo.png", "SmallLogo.png", "SplashScreen.png", "StoreLogo.png", "WideLogo.png" };
+
+			foreach (string ImageName in ImageList)
+			{
+				string OutputImageFilePath = Path.Combine(OutputImageFolder, ImageName);
+				if (File.Exists(OutputImageFilePath) == false)
+				{
+					string SourceImagePath = "../../Engine/Build/UWP/DefaultImages/" + ImageName;
+					File.Copy(SourceImagePath, OutputImageFilePath);
+
+					string OutputEnglishImageFilePath = Path.Combine(OutputImageFolder, "en-US", ImageName);
+					if (Directory.Exists(OutputImageFolder + "/en-US") == false)
+					{
+						Directory.CreateDirectory(OutputImageFolder + "/en-US");
+					}
+					if (File.Exists(OutputEnglishImageFilePath) == false)
+					{
+						File.Copy(SourceImagePath, OutputEnglishImageFilePath);
+					}
+				}
+			}
+		}
+
+
+		public void GenerateDeployFiles(string InTargetName, string InTargetFilepath, bool bIsUE4Game, string EngineSourceRelativeBinaryPath)
+		{
+			const string SourceFolderString = "/Source/";
+			string OutTargetPath = InTargetFilepath;
+			OutTargetPath = OutTargetPath.Replace("\\", "/");
+			if (bIsUE4Game == false)
+			{
+				int LastSourceIndex = OutTargetPath.LastIndexOf(SourceFolderString);
+				if (LastSourceIndex != -1)
+				{
+					OutTargetPath = OutTargetPath.Substring(0, LastSourceIndex + SourceFolderString.Length);
+					OutTargetPath = Path.Combine(OutTargetPath, "UWP");
+				}
+				else
+				{
+					// @todo UWP: This is bad... throw a build exception?
+				}
+			}
+			else
+			{
+				OutTargetPath = Path.Combine((new FileInfo(InTargetFilepath)).DirectoryName, "UWP");
+			}
+
+			string GameBinariesFolder = Utils.MakePathRelativeTo(EngineSourceRelativeBinaryPath, OutTargetPath);
+			GameBinariesFolder = GameBinariesFolder.Replace("\\", "/");
+			if (GameBinariesFolder.EndsWith("/") == false)
+			{
+				GameBinariesFolder += "/";
+			}
+			string GameIntermediatesFolder = GameBinariesFolder.Replace("Binaries/WinAUP", "Intermediate/Build/Unused");
+
+			// Make sure the output directory exists...
+			if (Directory.Exists(OutTargetPath) == false)
+			{
+				Log.TraceVerbose("Generating output directory " + OutTargetPath);
+				Directory.CreateDirectory(OutTargetPath);
+			}
+
+			GeneratePackageAppXManifest(InTargetName, bIsUE4Game, "Package", GameBinariesFolder, GameIntermediatesFolder, OutTargetPath);
+			GeneratePackageAppXManifest(InTargetName, bIsUE4Game, "Deploy", GameBinariesFolder, GameIntermediatesFolder, OutTargetPath);
+			GeneratePackagePriConfig(InTargetName, bIsUE4Game, GameBinariesFolder, GameIntermediatesFolder, OutTargetPath);
+			GeneratePackageResources(InTargetName, bIsUE4Game, GameBinariesFolder, GameIntermediatesFolder, OutTargetPath);
+			CopyDefaultImageFiles(OutTargetPath);
+		}
+
+
+		public override void GenerateGameProperties(UnrealTargetConfiguration Configuration, StringBuilder VCProjectFileContent, TargetRules.TargetType TargetType, string RootDirectory, string TargetFilePath)
+		{
+			// @todo UWP: This used to be "WINUAP=1".  Need to verify that 'UWP' is the correct define that we want here.
+			VCProjectFileContent.Append("		<NMakePreprocessorDefinitions>$(NMakePreprocessorDefinitions);PLATFORM_UWP=1;UWP=1;</NMakePreprocessorDefinitions>" + ProjectFileGenerator.NewLine);
+		}
+
+	}
+}
diff --git a/Engine/Source/Programs/UnrealBuildTool/UWP/UWPToolChain.cs b/Engine/Source/Programs/UnrealBuildTool/UWP/UWPToolChain.cs
new file mode 100644
--- /dev/null
+++ b/Engine/Source/Programs/UnrealBuildTool/UWP/UWPToolChain.cs
@@ -0,0 +1,1046 @@
+// Copyright 1998-2014 Epic Games, Inc. All Rights Reserved.
+
+using System;
+using System.Collections.Generic;
+using System.Diagnostics;
+using System.IO;
+using System.Text;
+using System.Text.RegularExpressions;
+
+using Microsoft.Win32;
+
+namespace UnrealBuildTool
+{
+	public class UWPToolChain : UEToolChain
+	{
+		public override void RegisterToolChain()
+		{
+			if( UWPPlatform.bEnableUWPSupport )
+			{ 
+				// Register this tool chain for UWP
+				Log.TraceVerbose("        Registered for {0}", CPPTargetPlatform.UWP.ToString());
+				UEToolChain.RegisterPlatformToolChain(CPPTargetPlatform.UWP, this);
+			}
+		}
+
+		static void AppendCLArguments_Global(CPPEnvironment CompileEnvironment, VCEnvironment EnvVars, StringBuilder Arguments)
+		{
+			//Arguments.Append(" /showIncludes");
+
+			if( BuildConfiguration.bEnableCodeAnalysis )
+			{
+				Arguments.Append(" /analyze");
+
+				// Don't cause analyze warnings to be errors
+				Arguments.Append(" /analyze:WX-");
+
+				// Report functions that use a LOT of stack space.  You can lower this value if you
+				// want more aggressive checking for functions that use a lot of stack memory.
+				Arguments.Append(" /analyze:stacksize81940");
+
+				// Don't bother generating code, only analyze code (may report fewer warnings though.)
+				//Arguments.Append(" /analyze:only");
+			}
+
+			// Prevents the compiler from displaying its logo for each invocation.
+			Arguments.Append(" /nologo");
+
+			// Enable intrinsic functions.
+			Arguments.Append(" /Oi");
+
+			// Pack struct members on 8-byte boundaries.
+			Arguments.Append(" /Zp8");
+
+			// Separate functions for linker.
+			Arguments.Append(" /Gy");
+
+			// Relaxes floating point precision semantics to allow more optimization.
+			Arguments.Append(" /fp:fast");
+
+			// Compile into an .obj file, and skip linking.
+			Arguments.Append(" /c");
+
+			// Allow 800% of the default memory allocation limit.
+			Arguments.Append(" /Zm800");
+
+			// Allow large object files to avoid hitting the 2^16 section limit when running with -StressTestUnity.
+			Arguments.Append(" /bigobj");
+
+			// Disable "The file contains a character that cannot be represented in the current code page" warning for non-US windows.
+			Arguments.Append(" /wd4819");
+
+			// @todo UWP: Disable "unreachable code" warning since auto-included vccorlib.h triggers it
+			Arguments.Append(" /wd4702");
+
+			// @todo UWP: Silence the hash_map deprecation errors for now. This should be replaced with unordered_map for the real fix.
+			if (WindowsPlatform.Compiler == WindowsCompiler.VisualStudio2015)
+			{
+				Arguments.Append(" /D_SILENCE_STDEXT_HASH_DEPRECATION_WARNINGS");
+			}
+
+			if ( BuildConfiguration.bUseSharedPCHs )
+			{
+				// @todo SharedPCH: Disable warning about PCH defines not matching .cpp defines.  We "cheat" these defines a little
+				// bit to make shared PCHs work.  But it's totally safe.  Trust us.
+				Arguments.Append(" /wd4651");
+
+				// @todo SharedPCH: Disable warning about redefining *API macros.  The PCH header is compiled with various DLLIMPORTs, but
+				// when a module that uses that PCH header *IS* one of those imports, that module is compiled with EXPORTS, so the macro
+				// is redefined on the command-line.  We need to clobber those defines to make shared PCHs work properly!
+				Arguments.Append(" /wd4005");
+			}
+
+			// If compiling as a DLL, set the relevant defines
+			if (CompileEnvironment.Config.bIsBuildingDLL)
+			{
+				Arguments.Append(" /D_WINDLL");
+			}
+
+			// Handle Common Language Runtime support (C++/CLI)
+			if (CompileEnvironment.Config.CLRMode == CPPCLRMode.CLREnabled)
+			{
+				Arguments.Append(" /clr");
+
+				// Don't use default lib path, we override it explicitly to use the 4.0 reference assemblies.
+				Arguments.Append(" /clr:nostdlib");
+			}
+
+			//
+			//	Debug
+			//
+			if (CompileEnvironment.Config.Target.Configuration == CPPTargetConfiguration.Debug)
+			{
+				// Disable compiler optimization.
+				Arguments.Append(" /Od");
+
+				// Favor code size (especially useful for embedded platforms).
+				Arguments.Append(" /Os");
+
+				// Allow inline method expansion unless E&C support is requested
+				if( !BuildConfiguration.bSupportEditAndContinue )
+				{
+					// @todo UWP: No inlining in Debug builds except in the editor where DLL exports/imports aren't handled properly in module _API macros.
+					if (UEBuildConfiguration.bBuildEditor)
+					{
+						Arguments.Append(" /Ob2");
+					}
+				}
+
+				// Runtime stack checks are not allowed when compiling for CLR
+				if (CompileEnvironment.Config.CLRMode == CPPCLRMode.CLRDisabled)
+				{
+					Arguments.Append(" /RTCs");
+				}
+			}
+			//
+			//	Development and LTCG
+			//
+			else
+			{
+				// Maximum optimizations if desired.
+				if (CompileEnvironment.Config.OptimizeCode >= ModuleRules.CodeOptimization.InNonDebugBuilds)
+				{
+					Arguments.Append(" /Ox");
+
+					// Allow optimized code to be debugged more easily.  This makes PDBs a bit larger, but doesn't noticeably affect
+					// compile times.  The executable code is not affected at all by this switch, only the debugging information.
+					Arguments.Append(" /Zo");
+				}
+
+				// Favor code speed.
+				Arguments.Append(" /Ot");
+
+				// Only omit frame pointers on the PC (which is implied by /Ox) if wanted.
+				if ( BuildConfiguration.bOmitFramePointers == false
+				&& (CompileEnvironment.Config.Target.Platform == CPPTargetPlatform.UWP))
+				{
+					Arguments.Append(" /Oy-");
+				}
+
+				// Allow inline method expansion
+				Arguments.Append(" /Ob2");
+
+				//
+				// LTCG
+				//
+				if (CompileEnvironment.Config.Target.Configuration == CPPTargetConfiguration.Shipping)
+				{
+					if( BuildConfiguration.bAllowLTCG )
+					{
+						// Enable link-time code generation.
+						Arguments.Append(" /GL");
+					}
+				}
+			}
+
+			//
+			//	PC
+			//
+
+			// Prompt the user before reporting internal errors to Microsoft.
+			Arguments.Append(" /errorReport:prompt");
+
+			// Enable C++ exception handling, but not C exceptions.
+			Arguments.Append(" /EHsc");
+
+			// If enabled, create debug information.
+			if (CompileEnvironment.Config.bCreateDebugInfo)
+			{
+				// Store debug info in .pdb files.
+				if (BuildConfiguration.bUsePDBFiles)
+				{
+					// Create debug info suitable for E&C if wanted.
+					if ( BuildConfiguration.bSupportEditAndContinue &&
+						// We only need to do this in debug as that's the only configuration that supports E&C.
+						CompileEnvironment.Config.Target.Configuration == CPPTargetConfiguration.Debug)
+					{
+						Arguments.Append(" /ZI");
+					}
+					// Regular PDB debug information.
+					else
+					{
+						Arguments.Append(" /Zi");
+					}
+					// We need to add this so VS won't lock the PDB file and prevent synchronous updates. This forces serialization through MSPDBSRV.exe.
+					// See http://msdn.microsoft.com/en-us/library/dn502518.aspx for deeper discussion of /FS switch.
+					if (BuildConfiguration.bUseIncrementalLinking && WindowsPlatform.Compiler >= WindowsCompiler.VisualStudio2013)
+					{
+						Arguments.Append(" /FS");
+					}
+				}
+				// Store C7-format debug info in the .obj files, which is faster.
+				else
+				{
+					Arguments.Append(" /Z7");
+				}
+			}
+
+			// Specify the appropriate runtime library based on the platform and config.
+			if (CompileEnvironment.Config.bUseStaticCRT)
+			{
+				if (CompileEnvironment.Config.Target.Configuration == CPPTargetConfiguration.Debug && BuildConfiguration.bDebugBuildsActuallyUseDebugCRT)
+				{
+					Arguments.Append(" /MTd");
+				}
+				else
+				{
+					Arguments.Append(" /MT");
+				}
+			}
+			else
+			{
+				if (CompileEnvironment.Config.Target.Configuration == CPPTargetConfiguration.Debug && BuildConfiguration.bDebugBuildsActuallyUseDebugCRT)
+				{
+					Arguments.Append(" /MDd");
+				}
+				else
+				{
+					Arguments.Append(" /MD");
+				}
+			}
+
+			if (UWPPlatform.bBuildForStore)
+			{
+				Arguments.Append(" /D_BUILD_FOR_STORE=1");
+			}
+
+			if (WindowsPlatform.Compiler == WindowsCompiler.VisualStudio2015)
+			{
+				// @todo UWP: These must be appended to the end of the system include list, lest they override some of the third party sources includes
+				if (Directory.Exists(EnvVars.WindowsSDKExtensionDir))
+				{
+					CompileEnvironment.Config.CPPIncludeInfo.SystemIncludePaths.Add(string.Format(@"{0}\Include\{1}\ucrt", EnvVars.WindowsSDKExtensionDir, EnvVars.WindowsSDKExtensionHeaderLibVersion));
+					CompileEnvironment.Config.CPPIncludeInfo.SystemIncludePaths.Add(string.Format(@"{0}\Include\{1}\um", EnvVars.WindowsSDKExtensionDir, EnvVars.WindowsSDKExtensionHeaderLibVersion));
+					CompileEnvironment.Config.CPPIncludeInfo.SystemIncludePaths.Add(string.Format(@"{0}\Include\{1}\shared", EnvVars.WindowsSDKExtensionDir, EnvVars.WindowsSDKExtensionHeaderLibVersion));
+					CompileEnvironment.Config.CPPIncludeInfo.SystemIncludePaths.Add(string.Format(@"{0}\Include\{1}\winrt", EnvVars.WindowsSDKExtensionDir, EnvVars.WindowsSDKExtensionHeaderLibVersion));
+				}
+			}
+
+			if (CompileEnvironment.Config.bIsBuildingLibrary == false) // will put in a config option, but for now...
+			{
+				// Enable Windows Runtime extensions
+				Arguments.Append(" /ZW");
+				Arguments.Append(" /DUSE_WINRT_MAIN=1");
+				// TODO - richiem - this will have to be updated when final layout SDKs are available
+				if (WindowsPlatform.Compiler == WindowsCompiler.VisualStudio2015 && 
+					Directory.Exists(Path.Combine(EnvVars.WindowsSDKExtensionDir, "UnionMetadata")))
+				{
+					Arguments.AppendFormat(@" /AI""{0}\UnionMetadata""", EnvVars.WindowsSDKExtensionDir);
+					Arguments.AppendFormat(@" /FU""{0}\UnionMetadata\windows.winmd""", EnvVars.WindowsSDKExtensionDir);
+				}
+				Arguments.AppendFormat(@" /AI""{0}\..\..\VC\vcpackages""", EnvVars.BaseVSToolPath);
+				Arguments.AppendFormat(@" /FU""{0}\..\..\VC\vcpackages\platform.winmd""", EnvVars.BaseVSToolPath);
+			}
+
+		}
+
+		static void AppendCLArguments_CPP( CPPEnvironment CompileEnvironment, StringBuilder Arguments)
+		{
+			// Explicitly compile the file as C++.
+			Arguments.Append(" /TP");
+
+			if (!CompileEnvironment.Config.bEnableBufferSecurityChecks)
+			{
+				// This will disable buffer security checks (which are enabled by default) that the MS compiler adds around arrays on the stack,
+				// Which can add some performance overhead, especially in performance intensive code
+				// Only disable this if you know what you are doing, because it will be disabled for the entire module!
+				Arguments.Append(" /GS-");
+			}
+
+			// C++/CLI requires that RTTI is left enabled
+			if (CompileEnvironment.Config.CLRMode == CPPCLRMode.CLRDisabled)
+			{
+				if (CompileEnvironment.Config.bUseRTTI)
+				{
+					// Enable C++ RTTI.
+					Arguments.Append(" /GR");
+				}
+				else
+				{
+					// Disable C++ RTTI.
+					Arguments.Append(" /GR-");
+				}
+			}
+
+			// Level 4 warnings.
+			Arguments.Append(" /W3");
+		}
+		
+		static void AppendCLArguments_C(StringBuilder Arguments)
+		{
+			// Explicitly compile the file as C.
+			Arguments.Append(" /TC");
+		
+			// Level 0 warnings.  Needed for external C projects that produce warnings at higher warning levels.
+			Arguments.Append(" /W0");
+		}
+
+		static void AppendLinkArguments(LinkEnvironment LinkEnvironment, StringBuilder Arguments)
+		{
+			// Don't create a side-by-side manifest file for the executable.
+			Arguments.Append(" /MANIFEST:NO");
+
+			// Prevents the linker from displaying its logo for each invocation.
+			Arguments.Append(" /NOLOGO");
+
+			if (LinkEnvironment.Config.bCreateDebugInfo)
+			{
+				// Output debug info for the linked executable.
+				Arguments.Append(" /DEBUG");
+			}
+
+			// Prompt the user before reporting internal errors to Microsoft.
+			Arguments.Append(" /errorReport:prompt");
+
+			//
+			//	PC
+			//
+			// Set machine type/ architecture to be 64 bit, and set as a store app
+			if (UWPPlatform.bBuildForStore)
+			{
+				Arguments.Append(" /MACHINE:x64");
+				Arguments.Append(" /APPCONTAINER");
+				// this helps with store API compliance validation tools, adding additional pdb info
+				Arguments.Append(" /PROFILE");
+			}
+
+			if (LinkEnvironment.Config.bIsBuildingConsoleApplication)
+			{
+				Arguments.Append(" /SUBSYSTEM:CONSOLE");
+			}
+			else
+			{
+				Arguments.Append(" /SUBSYSTEM:WINDOWS");
+			}
+
+			if (LinkEnvironment.Config.bIsBuildingConsoleApplication && !LinkEnvironment.Config.bIsBuildingDLL && !String.IsNullOrEmpty( LinkEnvironment.Config.WindowsEntryPointOverride ) )
+			{
+				// Use overridden entry point
+				Arguments.Append(" /ENTRY:" + LinkEnvironment.Config.WindowsEntryPointOverride);
+			}
+
+			// Allow the OS to load the EXE at different base addresses than its preferred base address.
+			Arguments.Append(" /FIXED:No");
+
+			// Explicitly declare that the executable is compatible with Data Execution Prevention.
+			Arguments.Append(" /NXCOMPAT");
+
+			// Set the default stack size.
+			Arguments.Append(" /STACK:5000000");
+
+			// Allow delay-loaded DLLs to be explicitly unloaded.
+			Arguments.Append(" /DELAY:UNLOAD");
+
+			if (LinkEnvironment.Config.bIsBuildingDLL)
+			{
+				Arguments.Append(" /DLL");
+			}
+
+			if (LinkEnvironment.Config.CLRMode == CPPCLRMode.CLREnabled)
+			{
+				// DLLs built with managed code aren't allowed to have entry points as they will try to initialize
+				// complex static variables.  Managed code isn't allowed to run during DLLMain, we can't allow
+				// these variables to be initialized here.
+				if (LinkEnvironment.Config.bIsBuildingDLL)
+				{
+					// NOTE: This appears to only be needed if we want to get call stacks for debugging exit crashes related to the above
+					//		Arguments.Append(" /NOENTRY /NODEFAULTLIB:nochkclr.obj");
+				}
+			}
+
+			// Don't embed the full PDB path; we want to be able to move binaries elsewhere. They will always be side by side.
+			Arguments.Append(" /PDBALTPATH:%_PDB%");
+
+			//
+			//	Shipping & LTCG
+			//
+			if( BuildConfiguration.bAllowLTCG &&
+				LinkEnvironment.Config.Target.Configuration == CPPTargetConfiguration.Shipping)
+			{
+				// Use link-time code generation.
+				Arguments.Append(" /LTCG");
+
+				// This is where we add in the PGO-Lite linkorder.txt if we are using PGO-Lite
+				//Arguments.Append(" /ORDER:@linkorder.txt");
+				//Arguments.Append(" /VERBOSE");
+			}
+
+			//
+			//	Shipping binary
+			//
+			if (LinkEnvironment.Config.Target.Configuration == CPPTargetConfiguration.Shipping)
+			{
+				// Generate an EXE checksum.
+				Arguments.Append(" /RELEASE");
+
+				// Eliminate unreferenced symbols.
+				Arguments.Append(" /OPT:REF");
+
+				// Remove redundant COMDATs.
+				Arguments.Append(" /OPT:ICF");
+			}
+			//
+			//	Regular development binary. 
+			//
+			else
+			{
+				// Keep symbols that are unreferenced.
+				Arguments.Append(" /OPT:NOREF");
+
+				// Disable identical COMDAT folding.
+				Arguments.Append(" /OPT:NOICF");
+			}
+
+			// Enable incremental linking if wanted.
+			// NOTE: Don't bother using incremental linking for C++/CLI projects, as that's not supported and the option
+			//		 will silently be ignored anyway
+			if (BuildConfiguration.bUseIncrementalLinking && LinkEnvironment.Config.CLRMode != CPPCLRMode.CLREnabled)
+			{
+				Arguments.Append(" /INCREMENTAL");
+			}
+			else
+			{
+				Arguments.Append(" /INCREMENTAL:NO");
+			}
+
+			// Disable 
+			//LINK : warning LNK4199: /DELAYLOAD:nvtt_64.dll ignored; no imports found from nvtt_64.dll
+			// type warning as we leverage the DelayLoad option to put third-party DLLs into a 
+			// non-standard location. This requires the module(s) that use said DLL to ensure that it 
+			// is loaded prior to using it.
+			Arguments.Append(" /ignore:4199");
+
+			// Suppress warnings about missing PDB files for statically linked libraries.  We often don't want to distribute
+			// PDB files for these libraries.
+			Arguments.Append(" /ignore:4099");		// warning LNK4099: PDB '<file>' was not found with '<file>'
+		}
+
+		static void AppendLibArguments(LinkEnvironment LinkEnvironment, StringBuilder Arguments)
+		{
+			// Prevents the linker from displaying its logo for each invocation.
+			Arguments.Append(" /NOLOGO");
+
+			// Prompt the user before reporting internal errors to Microsoft.
+			Arguments.Append(" /errorReport:prompt");
+
+			//
+			//	PC
+			//
+			// Set machine type/ architecture to be 64 bit.
+			Arguments.Append(" /MACHINE:x64");
+
+			if (LinkEnvironment.Config.bIsBuildingConsoleApplication)
+			{
+				Arguments.Append(" /SUBSYSTEM:CONSOLE");
+			}
+			else
+			{
+				Arguments.Append(" /SUBSYSTEM:WINDOWS");
+			}
+
+			//
+			//	Shipping & LTCG
+			//
+			if (LinkEnvironment.Config.Target.Configuration == CPPTargetConfiguration.Shipping)
+			{
+				// Use link-time code generation.
+				Arguments.Append(" /LTCG");
+			}
+		}
+
+		public override CPPOutput CompileCPPFiles(UEBuildTarget Target, CPPEnvironment CompileEnvironment, List<FileItem> SourceFiles, string ModuleName)
+		{
+			var EnvVars = VCEnvironment.SetEnvironment(CompileEnvironment.Config.Target.Platform);
+
+			StringBuilder Arguments = new StringBuilder();
+			AppendCLArguments_Global(CompileEnvironment, EnvVars, Arguments);
+
+			// Add include paths to the argument list.
+			foreach (string IncludePath in CompileEnvironment.Config.CPPIncludeInfo.IncludePaths)
+			{
+				string IncludePathRelative = Utils.CleanDirectorySeparators(Utils.MakePathRelativeTo(IncludePath, Path.Combine(ProjectFileGenerator.RootRelativePath, "Engine/Source")), '/');
+				Arguments.AppendFormat(" /I \"{0}\"", IncludePathRelative);
+			}
+			foreach (string IncludePath in CompileEnvironment.Config.CPPIncludeInfo.SystemIncludePaths)
+			{
+				string IncludePathRelative = Utils.CleanDirectorySeparators(Utils.MakePathRelativeTo(IncludePath, Path.Combine(ProjectFileGenerator.RootRelativePath, "Engine/Source")), '/');
+				Arguments.AppendFormat(" /I \"{0}\"", IncludePathRelative);
+			}
+
+
+			if (CompileEnvironment.Config.CLRMode == CPPCLRMode.CLREnabled)
+			{
+				// Add .NET framework assembly paths.  This is needed so that C++/CLI projects
+				// can reference assemblies with #using, without having to hard code a path in the
+				// .cpp file to the assembly's location.				
+				foreach (string AssemblyPath in CompileEnvironment.Config.SystemDotNetAssemblyPaths)
+				{
+					Arguments.AppendFormat(" /AI \"{0}\"", AssemblyPath);
+				}
+
+				// Add explicit .NET framework assembly references				
+				foreach (string AssemblyName in CompileEnvironment.Config.FrameworkAssemblyDependencies)
+				{
+					Arguments.AppendFormat( " /FU \"{0}\"", AssemblyName );
+				}
+
+				// Add private assembly references				
+				foreach( PrivateAssemblyInfo CurAssemblyInfo in CompileEnvironment.PrivateAssemblyDependencies )
+				{
+					Arguments.AppendFormat( " /FU \"{0}\"", CurAssemblyInfo.FileItem.AbsolutePath );
+				}
+			}
+
+
+			// Add preprocessor definitions to the argument list.
+			foreach (string Definition in CompileEnvironment.Config.Definitions)
+			{
+				// Escape all quotation marks so that they get properly passed with the command line.
+				var DefinitionArgument = Definition.Contains("\"") ? Definition.Replace("\"", "\\\"") : Definition;
+				Arguments.AppendFormat(" /D\"{0}\"", DefinitionArgument);
+			}
+
+			var BuildPlatform = UEBuildPlatform.GetBuildPlatformForCPPTargetPlatform(CompileEnvironment.Config.Target.Platform);
+
+			// Create a compile action for each source file.
+			CPPOutput Result = new CPPOutput();
+			foreach (FileItem SourceFile in SourceFiles)
+			{
+				Action CompileAction = new Action(ActionType.Compile);
+				CompileAction.CommandDescription = "Compile";
+
+				StringBuilder FileArguments = new StringBuilder();
+				bool bIsPlainCFile = Path.GetExtension(SourceFile.AbsolutePath).ToUpperInvariant() == ".C";
+
+				// Add the C++ source file and its included files to the prerequisite item list.
+				AddPrerequisiteSourceFile(Target, BuildPlatform, CompileEnvironment, SourceFile, CompileAction.PrerequisiteItems);
+
+				// If this is a CLR file then make sure our dependent assemblies are added as prerequisites
+				if (CompileEnvironment.Config.CLRMode == CPPCLRMode.CLREnabled)
+				{
+					foreach( PrivateAssemblyInfo CurPrivateAssemblyDependency in CompileEnvironment.PrivateAssemblyDependencies )
+					{
+						CompileAction.PrerequisiteItems.Add( CurPrivateAssemblyDependency.FileItem );
+					}
+				}
+
+				bool bEmitsObjectFile = true;
+				if (CompileEnvironment.Config.PrecompiledHeaderAction == PrecompiledHeaderAction.Create)
+				{
+					// Generate a CPP File that just includes the precompiled header.
+					string PCHCPPFilename = "PCH." + ModuleName + "." + Path.GetFileName(CompileEnvironment.Config.PrecompiledHeaderIncludeFilename) + ".cpp";
+					string PCHCPPPath = Path.Combine(CompileEnvironment.Config.OutputDirectory, PCHCPPFilename);
+					FileItem PCHCPPFile = FileItem.CreateIntermediateTextFile(
+						PCHCPPPath,
+						string.Format("#include \"{0}\"\r\n", CompileEnvironment.Config.PrecompiledHeaderIncludeFilename)
+						);
+
+					// Make sure the original source directory the PCH header file existed in is added as an include
+					// path -- it might be a private PCH header and we need to make sure that its found!
+					string OriginalPCHHeaderDirectory = Path.GetDirectoryName( SourceFile.AbsolutePath );
+					FileArguments.AppendFormat(" /I \"{0}\"", OriginalPCHHeaderDirectory);
+
+					var PrecompiledFileExtension = UEBuildPlatform.BuildPlatformDictionary[UnrealTargetPlatform.UWP].GetBinaryExtension(UEBuildBinaryType.PrecompiledHeader);
+					// Add the precompiled header file to the produced items list.
+					FileItem PrecompiledHeaderFile = FileItem.GetItemByPath(
+						Path.Combine(
+							CompileEnvironment.Config.OutputDirectory,
+						Path.GetFileName(SourceFile.AbsolutePath) + PrecompiledFileExtension
+							)
+						);
+					CompileAction.ProducedItems.Add(PrecompiledHeaderFile);
+					Result.PrecompiledHeaderFile = PrecompiledHeaderFile;
+
+					// Add the parameters needed to compile the precompiled header file to the command-line.
+					FileArguments.AppendFormat(" /Yc\"{0}\"", CompileEnvironment.Config.PrecompiledHeaderIncludeFilename);
+					FileArguments.AppendFormat(" /Fp\"{0}\"", PrecompiledHeaderFile.AbsolutePath);
+
+					// If we're creating a PCH that will be used to compile source files for a library, we need
+					// the compiled modules to retain a reference to PCH's module, so that debugging information
+					// will be included in the library.  This is also required to avoid linker warning "LNK4206"
+					// when linking an application that uses this library.
+					if (CompileEnvironment.Config.bIsBuildingLibrary)
+					{
+						// NOTE: The symbol name we use here is arbitrary, and all that matters is that it is
+						// unique per PCH module used in our library
+						string FakeUniquePCHSymbolName = Path.GetFileNameWithoutExtension(CompileEnvironment.Config.PrecompiledHeaderIncludeFilename);
+						FileArguments.AppendFormat(" /Yl{0}", FakeUniquePCHSymbolName);
+					}
+
+					FileArguments.AppendFormat(" \"{0}\"", PCHCPPFile.AbsolutePath);
+
+					CompileAction.StatusDescription = PCHCPPFilename;
+				}
+				else
+				{
+					if (CompileEnvironment.Config.PrecompiledHeaderAction == PrecompiledHeaderAction.Include)
+					{
+						CompileAction.bIsUsingPCH = true;
+						CompileAction.PrerequisiteItems.Add(CompileEnvironment.PrecompiledHeaderFile);
+
+						FileArguments.AppendFormat(" /Yu\"{0}\"", CompileEnvironment.Config.PCHHeaderNameInCode);
+						FileArguments.AppendFormat(" /Fp\"{0}\"", CompileEnvironment.PrecompiledHeaderFile.AbsolutePath);
+
+						// Is it unsafe to always force inclusion?  Clang is doing it, and .generated.cpp files
+						// won't work otherwise, because they're not located in the context of the module,
+						// so they can't access the module's PCH without an absolute path.
+						//if( CompileEnvironment.Config.bForceIncludePrecompiledHeader )
+						{
+							// Force include the precompiled header file.  This is needed because we may have selected a
+							// precompiled header that is different than the first direct include in the C++ source file, but
+							// we still need to make sure that our precompiled header is the first thing included!
+							FileArguments.AppendFormat( " /FI\"{0}\"", CompileEnvironment.Config.PCHHeaderNameInCode);
+						}
+					}
+
+					// Add the source file path to the command-line.
+					FileArguments.AppendFormat(" \"{0}\"", SourceFile.AbsolutePath);
+
+					CompileAction.StatusDescription = Path.GetFileName( SourceFile.AbsolutePath );
+				}
+
+				if( bEmitsObjectFile )
+				{
+					var ObjectFileExtension = UEBuildPlatform.BuildPlatformDictionary[UnrealTargetPlatform.UWP].GetBinaryExtension(UEBuildBinaryType.Object);
+					// Add the object file to the produced item list.
+					FileItem ObjectFile = FileItem.GetItemByPath(
+						Path.Combine(
+							CompileEnvironment.Config.OutputDirectory,
+							Path.GetFileName(SourceFile.AbsolutePath) + ObjectFileExtension
+							)
+						);
+					CompileAction.ProducedItems.Add(ObjectFile);
+					Result.ObjectFiles.Add(ObjectFile);
+					FileArguments.AppendFormat(" /Fo\"{0}\"", ObjectFile.AbsolutePath);
+				}
+
+				// Create PDB files if we were configured to do that.
+				//
+				// Also, when debug info is off and XGE is enabled, force PDBs, otherwise it will try to share
+				// a PDB file, which causes PCH creation to be serial rather than parallel (when debug info is disabled)
+				//		--> See https://udn.epicgames.com/lists/showpost.php?id=50619&list=unprog3
+				if (BuildConfiguration.bUsePDBFiles ||
+					(BuildConfiguration.bAllowXGE && !CompileEnvironment.Config.bCreateDebugInfo))
+				{
+					string PDBFileName;
+					bool bActionProducesPDB = false;
+
+					// All files using the same PCH are required to share the same PDB that was used when compiling the PCH
+					if (CompileEnvironment.Config.PrecompiledHeaderAction == PrecompiledHeaderAction.Include)
+					{
+						PDBFileName = "PCH." + ModuleName + "." + Path.GetFileName(CompileEnvironment.Config.PrecompiledHeaderIncludeFilename);
+					}
+					// Files creating a PCH use a PDB per file.
+					else if (CompileEnvironment.Config.PrecompiledHeaderAction == PrecompiledHeaderAction.Create)
+					{
+						PDBFileName = "PCH." + ModuleName + "." + Path.GetFileName(CompileEnvironment.Config.PrecompiledHeaderIncludeFilename);
+						bActionProducesPDB = true;
+					}
+					// Ungrouped C++ files use a PDB per file.
+					else if( !bIsPlainCFile )
+					{
+						PDBFileName = Path.GetFileName( SourceFile.AbsolutePath );
+						bActionProducesPDB = true;
+					}
+					// Group all plain C files that doesn't use PCH into the same PDB
+					else
+					{
+						PDBFileName = "MiscPlainC";
+					}
+
+					// Specify the PDB file that the compiler should write to.
+					FileItem PDBFile = FileItem.GetItemByPath(
+							Path.Combine(
+								CompileEnvironment.Config.OutputDirectory,
+								PDBFileName + ".pdb"
+								)
+							);
+					FileArguments.AppendFormat(" /Fd\"{0}\"", PDBFile.AbsolutePath);
+
+					// Only use the PDB as an output file if we want PDBs and this particular action is
+					// the one that produces the PDB (as opposed to no debug info, where the above code
+					// is needed, but not the output PDB, or when multiple files share a single PDB, so
+					// only the action that generates it should count it as output directly)
+					if( BuildConfiguration.bUsePDBFiles && bActionProducesPDB )
+					{
+						CompileAction.ProducedItems.Add(PDBFile);
+						Result.DebugDataFiles.Add(PDBFile);
+					}
+				}
+
+				// Add C or C++ specific compiler arguments.
+				if (bIsPlainCFile)
+				{
+					AppendCLArguments_C(FileArguments);
+				}
+				else
+				{
+					AppendCLArguments_CPP(CompileEnvironment, FileArguments);
+				}
+
+				CompileAction.WorkingDirectory = Path.GetFullPath(".");
+				CompileAction.CommandPath      = EnvVars.CompilerPath;
+				CompileAction.CommandArguments = Arguments.ToString() + FileArguments.ToString() + CompileEnvironment.Config.AdditionalArguments;
+
+				if( CompileEnvironment.Config.PrecompiledHeaderAction == PrecompiledHeaderAction.Create )
+				{
+					Log.TraceVerbose("Creating PCH: " + CompileEnvironment.Config.PrecompiledHeaderIncludeFilename);
+					Log.TraceVerbose("     Command: " + CompileAction.CommandArguments);
+				}
+				else
+				{
+					Log.TraceVerbose("   Compiling: " + CompileAction.StatusDescription);
+					Log.TraceVerbose("     Command: " + CompileAction.CommandArguments);
+				}
+
+				// VC++ always outputs the source file name being compiled, so we don't need to emit this ourselves
+				CompileAction.bShouldOutputStatusDescription = false;
+
+				// Don't farm out creation of precompiled headers as it is the critical path task.
+				CompileAction.bCanExecuteRemotely =
+					CompileEnvironment.Config.PrecompiledHeaderAction != PrecompiledHeaderAction.Create ||
+					BuildConfiguration.bAllowRemotelyCompiledPCHs
+					;
+
+				// @todo: XGE has problems remote compiling C++/CLI files that use .NET Framework 4.0
+				if (CompileEnvironment.Config.CLRMode == CPPCLRMode.CLREnabled)
+				{
+					CompileAction.bCanExecuteRemotely = false;
+				}
+			}
+			return Result;
+		}
+
+		public override CPPOutput CompileRCFiles(UEBuildTarget Target, CPPEnvironment Environment, List<FileItem> RCFiles)
+		{
+			var EnvVars = VCEnvironment.SetEnvironment(Environment.Config.Target.Platform);
+
+			CPPOutput Result = new CPPOutput();
+
+			var BuildPlatform = UEBuildPlatform.GetBuildPlatformForCPPTargetPlatform(Environment.Config.Target.Platform);
+
+			foreach (FileItem RCFile in RCFiles)
+			{
+				Action CompileAction = new Action(ActionType.Compile);
+				CompileAction.CommandDescription = "Resource";
+				CompileAction.WorkingDirectory = Path.GetFullPath(".");
+				CompileAction.CommandPath = EnvVars.ResourceCompilerPath;
+				CompileAction.StatusDescription = Path.GetFileName(RCFile.AbsolutePath);
+
+				// Suppress header spew
+				CompileAction.CommandArguments += " /nologo";
+
+				// If we're compiling for 64-bit Windows, also add the _WIN64 definition to the resource
+				// compiler so that we can switch on that in the .rc file using #ifdef.
+				CompileAction.CommandArguments += " /D_WIN64";
+
+				// Language
+				CompileAction.CommandArguments += " /l 0x409";
+
+				// Include paths.
+				foreach (string IncludePath in Environment.Config.CPPIncludeInfo.IncludePaths)
+				{
+					CompileAction.CommandArguments += string.Format( " /i \"{0}\"", IncludePath );
+				}
+
+				// System include paths.
+				foreach( var SystemIncludePath in Environment.Config.CPPIncludeInfo.SystemIncludePaths)
+				{
+					CompileAction.CommandArguments += string.Format( " /i \"{0}\"", SystemIncludePath );
+				}
+
+				// Preprocessor definitions.
+				foreach (string Definition in Environment.Config.Definitions)
+				{
+					CompileAction.CommandArguments += string.Format(" /d \"{0}\"", Definition);
+				}
+
+				// Add the RES file to the produced item list.
+				FileItem CompiledResourceFile = FileItem.GetItemByPath(
+					Path.Combine(
+						Environment.Config.OutputDirectory,
+						Path.GetFileName(RCFile.AbsolutePath) + ".res"
+						)
+					);
+				CompileAction.ProducedItems.Add(CompiledResourceFile);
+				CompileAction.CommandArguments += string.Format(" /fo \"{0}\"", CompiledResourceFile.AbsolutePath);
+				Result.ObjectFiles.Add(CompiledResourceFile);
+
+				// Add the RC file as a prerequisite of the action.
+				CompileAction.CommandArguments += string.Format(" \"{0}\"", RCFile.AbsolutePath);
+
+				// Add the C++ source file and its included files to the prerequisite item list.
+				AddPrerequisiteSourceFile(Target, BuildPlatform, Environment, RCFile, CompileAction.PrerequisiteItems);
+			}
+
+			return Result;
+		}
+
+		public override FileItem LinkFiles(LinkEnvironment LinkEnvironment, bool bBuildImportLibraryOnly)
+		{
+			var EnvVars = VCEnvironment.SetEnvironment(LinkEnvironment.Config.Target.Platform);
+
+			if (LinkEnvironment.Config.bIsBuildingDotNetAssembly)
+			{
+				return FileItem.GetItemByPath(LinkEnvironment.Config.OutputFilePath);
+			}
+
+			bool bIsBuildingLibrary = LinkEnvironment.Config.bIsBuildingLibrary || bBuildImportLibraryOnly;
+			bool bIncludeDependentLibrariesInLibrary = bIsBuildingLibrary && LinkEnvironment.Config.bIncludeDependentLibrariesInLibrary;
+
+			// Get link arguments.
+			StringBuilder Arguments = new StringBuilder();
+			if (bIsBuildingLibrary)
+			{
+				AppendLibArguments(LinkEnvironment, Arguments);
+			}
+			else
+			{
+				AppendLinkArguments(LinkEnvironment, Arguments);
+			}
+
+
+
+			// If we're only building an import library, add the '/DEF' option that tells the LIB utility
+			// to simply create a .LIB file and .EXP file, and don't bother validating imports
+			if (bBuildImportLibraryOnly)
+			{
+				Arguments.Append(" /DEF");
+
+				// Ensure that the import library references the correct filename for the linked binary.
+				Arguments.AppendFormat(" /NAME:\"{0}\"", Path.GetFileName(LinkEnvironment.Config.OutputFilePath));
+			}
+
+
+			// Add delay loaded DLLs.
+			if (!bIsBuildingLibrary)
+			{
+				// Delay-load these DLLs.
+				foreach (string DelayLoadDLL in LinkEnvironment.Config.DelayLoadDLLs)
+				{
+					Arguments.AppendFormat(" /DELAYLOAD:\"{0}\"", DelayLoadDLL);
+				}
+			}
+
+			// @todo UE4 DLL: Why do I need LIBPATHs to build only export libraries with /DEF? (tbbmalloc.lib)
+			if (!LinkEnvironment.Config.bIsBuildingLibrary || (LinkEnvironment.Config.bIsBuildingLibrary && bIncludeDependentLibrariesInLibrary))
+			{
+				// Add the library paths to the argument list.
+				foreach (string LibraryPath in LinkEnvironment.Config.LibraryPaths)
+				{
+					Arguments.AppendFormat(" /LIBPATH:\"{0}\"", LibraryPath);
+				}
+
+				// Add the excluded default libraries to the argument list.
+				foreach (string ExcludedLibrary in LinkEnvironment.Config.ExcludedLibraries)
+				{
+					Arguments.AppendFormat(" /NODEFAULTLIB:\"{0}\"", ExcludedLibrary);
+				}
+			}
+
+			// For targets that are cross-referenced, we don't want to write a LIB file during the link step as that
+			// file will clobber the import library we went out of our way to generate during an earlier step.  This
+			// file is not needed for our builds, but there is no way to prevent MSVC from generating it when
+			// linking targets that have exports.  We don't want this to clobber our LIB file and invalidate the
+			// existing timstamp, so instead we simply emit it with a different name
+			string ImportLibraryFilePath = Path.Combine( LinkEnvironment.Config.IntermediateDirectory, 
+														 Path.GetFileNameWithoutExtension(LinkEnvironment.Config.OutputFilePath) + ".lib" );
+
+			if (LinkEnvironment.Config.bIsCrossReferenced && !bBuildImportLibraryOnly)
+			{
+				ImportLibraryFilePath += ".suppressed";
+			}
+
+			FileItem OutputFile;
+			if (bBuildImportLibraryOnly)
+			{
+				OutputFile = FileItem.GetItemByPath(ImportLibraryFilePath);
+			}
+			else
+			{
+				OutputFile = FileItem.GetItemByPath(LinkEnvironment.Config.OutputFilePath);
+				OutputFile.bNeedsHotReloadNumbersDLLCleanUp = LinkEnvironment.Config.bIsBuildingDLL;
+			}
+
+			List<FileItem> ProducedItems = new List<FileItem>();
+			ProducedItems.Add(OutputFile);
+
+			List<FileItem> PrerequisiteItems = new List<FileItem>();
+
+			// Add the input files to a response file, and pass the response file on the command-line.
+			List<string> InputFileNames = new List<string>();
+			foreach (FileItem InputFile in LinkEnvironment.InputFiles)
+			{
+				InputFileNames.Add(string.Format("\"{0}\"", InputFile.AbsolutePath));
+				PrerequisiteItems.Add(InputFile);
+			}
+
+			if (!bBuildImportLibraryOnly)
+			{
+				// Add input libraries as prerequisites, too!
+				foreach (FileItem InputLibrary in LinkEnvironment.InputLibraries)
+				{
+					InputFileNames.Add(string.Format("\"{0}\"", InputLibrary.AbsolutePath));
+					PrerequisiteItems.Add(InputLibrary);
+				}
+			}
+
+			if (!bIsBuildingLibrary || (LinkEnvironment.Config.bIsBuildingLibrary && bIncludeDependentLibrariesInLibrary))
+			{
+				foreach (string AdditionalLibrary in LinkEnvironment.Config.AdditionalLibraries)
+				{
+					InputFileNames.Add(string.Format("\"{0}\"", AdditionalLibrary));
+
+					// If the library file name has a relative path attached (rather than relying on additional
+					// lib directories), then we'll add it to our prerequisites list.  This will allow UBT to detect
+					// when the binary needs to be relinked because a dependent external library has changed.
+					//if( !String.IsNullOrEmpty( Path.GetDirectoryName( AdditionalLibrary ) ) )
+					{
+						PrerequisiteItems.Add(FileItem.GetItemByPath(AdditionalLibrary));
+					}
+				}
+			}
+
+			// Create a response file for the linker
+			string ResponseFileName = GetResponseFileName( LinkEnvironment, OutputFile );
+
+			// Never create response files when we are only generating IntelliSense data
+			if( !ProjectFileGenerator.bGenerateProjectFiles )
+			{
+				ResponseFile.Create( ResponseFileName, InputFileNames );
+			}
+			Arguments.AppendFormat(" @\"{0}\"", ResponseFileName);
+
+			// Add the output file to the command-line.
+			Arguments.AppendFormat(" /OUT:\"{0}\"", OutputFile.AbsolutePath);
+
+			if (bBuildImportLibraryOnly || (LinkEnvironment.Config.bHasExports && !bIsBuildingLibrary))
+			{
+				// An export file is written to the output directory implicitly; add it to the produced items list.
+				string ExportFilePath = Path.ChangeExtension(ImportLibraryFilePath, ".exp");
+				FileItem ExportFile = FileItem.GetItemByPath(ExportFilePath);
+				ProducedItems.Add(ExportFile);
+			}
+
+			if (!bIsBuildingLibrary)
+			{
+				// Xbox 360 LTCG does not seem to produce those.
+				if (LinkEnvironment.Config.bHasExports &&
+					(LinkEnvironment.Config.Target.Configuration != CPPTargetConfiguration.Shipping))
+				{
+					// Write the import library to the output directory for nFringe support.
+					FileItem ImportLibraryFile = FileItem.GetItemByPath(ImportLibraryFilePath);
+					Arguments.AppendFormat(" /IMPLIB:\"{0}\"", ImportLibraryFilePath);
+					ProducedItems.Add(ImportLibraryFile);
+				}
+
+				if (LinkEnvironment.Config.bCreateDebugInfo)
+				{
+					// Write the PDB file to the output directory.			
+					string PDBFilePath = Path.Combine(LinkEnvironment.Config.OutputDirectory, Path.GetFileNameWithoutExtension(OutputFile.AbsolutePath) + ".pdb");
+					FileItem PDBFile = FileItem.GetItemByPath(PDBFilePath);
+					Arguments.AppendFormat(" /PDB:\"{0}\"", PDBFilePath);
+					ProducedItems.Add(PDBFile);
+
+					// Write the MAP file to the output directory.			
+#if false					
+					if (true)
+					{
+						string MAPFilePath = Path.Combine(LinkEnvironment.Config.OutputDirectory, Path.GetFileNameWithoutExtension(OutputFile.AbsolutePath) + ".map");
+						FileItem MAPFile = FileItem.GetItemByPath(MAPFilePath);
+						LinkAction.CommandArguments += string.Format(" /MAP:\"{0}\"", MAPFilePath);
+						LinkAction.ProducedItems.Add(MAPFile);
+					}
+#endif
+				}
+
+				// Add the additional arguments specified by the environment.
+				Arguments.Append(LinkEnvironment.Config.AdditionalArguments);
+			}
+
+			// Create an action that invokes the linker.
+			Action LinkAction = new Action(ActionType.Link);
+			LinkAction.CommandDescription = "Link";
+			LinkAction.WorkingDirectory   = Path.GetFullPath(".");
+			LinkAction.CommandPath        = bIsBuildingLibrary ? EnvVars.LibraryLinkerPath : EnvVars.LinkerPath;
+			LinkAction.CommandArguments   = Arguments.ToString();
+			LinkAction.ProducedItems    .AddRange(ProducedItems);
+			LinkAction.PrerequisiteItems.AddRange(PrerequisiteItems);
+			LinkAction.StatusDescription = Path.GetFileName(OutputFile.AbsolutePath);
+
+			// Tell the action that we're building an import library here and it should conditionally be
+			// ignored as a prerequisite for other actions
+			LinkAction.bProducesImportLibrary = bBuildImportLibraryOnly || LinkEnvironment.Config.bIsBuildingDLL;
+
+			// Only execute linking on the local PC.
+			LinkAction.bCanExecuteRemotely = false;
+
+			Log.TraceVerbose("     Linking: " + LinkAction.StatusDescription);
+			Log.TraceVerbose("     Command: " + LinkAction.CommandArguments);
+
+			return OutputFile;
+		}
+
+  
+		/** Gets the default include paths for the given platform. */
+		public static string GetVCIncludePaths(CPPTargetPlatform Platform)
+		{
+			Debug.Assert(Platform == CPPTargetPlatform.UWP);
+
+			// Make sure we've got the environment variables set up for this target
+			VCEnvironment.SetEnvironment(Platform);
+
+			// Also add any include paths from the INCLUDE environment variable.  MSVC is not necessarily running with an environment that
+			// matches what UBT extracted from the vcvars*.bat using SetEnvironmentVariablesFromBatchFile().  We'll use the variables we
+			// extracted to populate the project file's list of include paths
+			// @todo projectfiles: Should we only do this for VC++ platforms?
+			var IncludePaths = Environment.GetEnvironmentVariable("INCLUDE");
+			if (!String.IsNullOrEmpty(IncludePaths) && !IncludePaths.EndsWith(";"))
+			{
+				IncludePaths += ";";
+			}
+			
+			return IncludePaths;
+		}
+	};
+}
diff --git a/Engine/Source/Programs/UnrealBuildTool/UnrealBuildTool.csproj b/Engine/Source/Programs/UnrealBuildTool/UnrealBuildTool.csproj
--- a/Engine/Source/Programs/UnrealBuildTool/UnrealBuildTool.csproj
+++ b/Engine/Source/Programs/UnrealBuildTool/UnrealBuildTool.csproj
@@ -130,17 +130,15 @@
     <Compile Include="Android\AndroidProjectGenerator.cs" />
   </ItemGroup>
   <ItemGroup>
+    <Compile Include="UWP\UEBuildUWP.cs" />
+    <Compile Include="UWP\UWPDeploy.cs" />
+    <Compile Include="UWP\UWPProjectGenerator.cs" />
+    <Compile Include="UWP\UWPToolChain.cs" />
     <Compile Include="WinRT\UEBuildWinRT.cs" />
     <Compile Include="WinRT\WinRTToolChain.cs" />
     <Compile Include="WinRT\WinRTDeploy.cs" />
     <Compile Include="WinRT\WinRTProjectGenerator.cs" />
   </ItemGroup>
-  <ItemGroup>
-    <Compile Include="WinUAP\WinUAPProjectGenerator.cs" />
-    <Compile Include="WinUAP\WinUAPToolChain.cs" />
-    <Compile Include="WinUAP\UEBuildWinUAP.cs" />
-    <Compile Include="WinUAP\WinUAPDeploy.cs" />
-  </ItemGroup>
   <ItemGroup>
     <Compile Include="IOS\IOSProjectGenerator.cs" />
     <Compile Include="IOS\IOSToolChain.cs" />
@@ -299,4 +297,4 @@
     <PreBuildEvent>
     </PreBuildEvent>
   </PropertyGroup>
-</Project>
+</Project>
\ No newline at end of file
diff --git a/Engine/Source/Programs/UnrealBuildTool/WinRT/WinRTToolChain.cs b/Engine/Source/Programs/UnrealBuildTool/WinRT/WinRTToolChain.cs
--- a/Engine/Source/Programs/UnrealBuildTool/WinRT/WinRTToolChain.cs
+++ b/Engine/Source/Programs/UnrealBuildTool/WinRT/WinRTToolChain.cs
@@ -66,7 +66,7 @@ namespace UnrealBuildTool
 			// Disable "The file contains a character that cannot be represented in the current code page" warning for non-US windows.
 			Result += " /wd4819";
 
-			// @todo UAP: UE4 is non-compliant when it comes to use of %s and %S
+			// @todo UWP: UE4 is non-compliant when it comes to use of %s and %S
 			// Previously %s meant "the current character set" and %S meant "the other one".
 			// Now %s means multibyte and %S means wide. %Ts means "natural width".
 			// Reverting this behaviour until the UE4 source catches up.
diff --git a/Engine/Source/Programs/UnrealBuildTool/WinUAP/UEBuildWinUAP.cs b/Engine/Source/Programs/UnrealBuildTool/WinUAP/UEBuildWinUAP.cs
deleted file mode 100644
diff --git a/Engine/Source/Programs/UnrealBuildTool/WinUAP/WinUAPDeploy.cs b/Engine/Source/Programs/UnrealBuildTool/WinUAP/WinUAPDeploy.cs
deleted file mode 100644
diff --git a/Engine/Source/Programs/UnrealBuildTool/WinUAP/WinUAPProjectGenerator.cs b/Engine/Source/Programs/UnrealBuildTool/WinUAP/WinUAPProjectGenerator.cs
deleted file mode 100644
diff --git a/Engine/Source/Programs/UnrealBuildTool/WinUAP/WinUAPToolChain.cs b/Engine/Source/Programs/UnrealBuildTool/WinUAP/WinUAPToolChain.cs
deleted file mode 100644
diff --git a/Engine/Source/Programs/UnrealBuildTool/Windows/UEBuildWindows.cs b/Engine/Source/Programs/UnrealBuildTool/Windows/UEBuildWindows.cs
--- a/Engine/Source/Programs/UnrealBuildTool/Windows/UEBuildWindows.cs
+++ b/Engine/Source/Programs/UnrealBuildTool/Windows/UEBuildWindows.cs
@@ -58,7 +58,7 @@ namespace UnrealBuildTool
 				}
 
 				// Second, default based on what's installed, from newest to oldest
-				// @todo UAP: Currently we default to VS 2013, even if VS 2015 is installed, until the engine is fully updated to
+				// @todo UWP: Currently we default to VS 2013, even if VS 2015 is installed, until the engine is fully updated to
 				// support compiling with VS 2015.  You can override this with the "-2015" command-line switch.
 // 				else if (!String.IsNullOrEmpty(WindowsPlatform.GetVSComnToolsPath(WindowsCompiler.VisualStudio2015)))
 // 				{
@@ -95,7 +95,7 @@ namespace UnrealBuildTool
 
 		/// Whether to compile against the Windows 10 SDK, instead of the Windows 8.1 SDK.  This requires the Visual Studio 2015
 		/// compiler or later, and the Windows 10 SDK must be installed.  The application will require at least Windows 8.x to run.
-		// @todo UAP: Expose this to be enabled more easily for building Windows 10 desktop apps
+		// @todo UWP: Expose this to be enabled more easily for building Windows 10 desktop apps
 		public static readonly bool bUseWindowsSDK10 = false;
 
 
diff --git a/Engine/Source/Programs/UnrealBuildTool/Windows/VCEnvironment.cs b/Engine/Source/Programs/UnrealBuildTool/Windows/VCEnvironment.cs
--- a/Engine/Source/Programs/UnrealBuildTool/Windows/VCEnvironment.cs
+++ b/Engine/Source/Programs/UnrealBuildTool/Windows/VCEnvironment.cs
@@ -76,8 +76,8 @@ namespace UnrealBuildTool
 			ResourceCompilerPath = GetResourceCompilerToolPath(Platform);
 
 			var VCVarsBatchFile = Path.Combine(BaseVSToolPath, 
-				(Platform == CPPTargetPlatform.Win64 || Platform == CPPTargetPlatform.WinUAP) ? "../../VC/bin/x86_amd64/vcvarsx86_amd64.bat" : "vsvars32.bat");
-			if (Platform == CPPTargetPlatform.WinUAP && WinUAPPlatform.bBuildForStore)
+				(Platform == CPPTargetPlatform.Win64 || Platform == CPPTargetPlatform.UWP) ? "../../VC/bin/x86_amd64/vcvarsx86_amd64.bat" : "vsvars32.bat");
+			if (Platform == CPPTargetPlatform.UWP && UWPPlatform.bBuildForStore)
 			{
 				Utils.SetEnvironmentVariablesFromBatchFile(VCVarsBatchFile, "store");
 			}
@@ -263,7 +263,7 @@ namespace UnrealBuildTool
 		{
 			// Regardless of the target, if we're linking on a 64 bit machine, we want to use the 64 bit linker (it's faster than the 32 bit linker)
 			[REDACTED_EMAIL]: Using the 64-bit linker appears to be broken at the moment.
-			if (Platform == CPPTargetPlatform.Win64 || Platform == CPPTargetPlatform.WinUAP)
+			if (Platform == CPPTargetPlatform.Win64 || Platform == CPPTargetPlatform.UWP)
 			{
 				// Use the native 64-bit compiler if present, otherwise use the amd64-on-x86 compiler. VS2012 Express only includes the latter.
 				var Result = Path.Combine(BaseVSToolPath, "../../VC/bin/amd64");
@@ -348,7 +348,7 @@ namespace UnrealBuildTool
 		string GetResourceCompilerToolPath(CPPTargetPlatform Platform)
 		{
 			// 64 bit -- we can use the 32 bit version to target 64 bit on 32 bit OS.
-			if (Platform == CPPTargetPlatform.Win64 || Platform == CPPTargetPlatform.WinUAP)
+			if (Platform == CPPTargetPlatform.Win64 || Platform == CPPTargetPlatform.UWP)
 			{
 				if (WindowsPlatform.Compiler == WindowsCompiler.VisualStudio2015 && WindowsPlatform.bUseWindowsSDK10)
 				{
@@ -360,7 +360,7 @@ namespace UnrealBuildTool
 				}
 			}
 
-			// @todo UAP: Verify that Windows XP will compile using VS 2015 (it should be supported)
+			// @todo UWP: Verify that Windows XP will compile using VS 2015 (it should be supported)
 			if (!WindowsPlatform.IsWindowsXPSupported())	// Windows XP requires use to force Windows SDK 7.1 even on the newer compiler, so we need the old path RC.exe
 			{
 				if (WindowsPlatform.Compiler == WindowsCompiler.VisualStudio2015 && WindowsPlatform.bUseWindowsSDK10)
diff --git a/Engine/Source/Programs/UnrealBuildTool/Windows/VCToolChain.cs b/Engine/Source/Programs/UnrealBuildTool/Windows/VCToolChain.cs
--- a/Engine/Source/Programs/UnrealBuildTool/Windows/VCToolChain.cs
+++ b/Engine/Source/Programs/UnrealBuildTool/Windows/VCToolChain.cs
@@ -91,8 +91,8 @@ namespace UnrealBuildTool
 
 		static void AppendCLArguments_Global(CPPEnvironment CompileEnvironment, VCEnvironment EnvVars, StringBuilder Arguments)
 		{
-			// @todo UAP: Why do we ever need WinRT headers when building regular Win32?  Is this just needed for the Windows 10 SDK?
-			// @todo UAP: These include paths should be added in SetUpEnvironment(), not here.  Do they need to be the last includes or something?
+			// @todo UWP: Why do we ever need WinRT headers when building regular Win32?  Is this just needed for the Windows 10 SDK?
+			// @todo UWP: These include paths should be added in SetUpEnvironment(), not here.  Do they need to be the last includes or something?
 			if (WindowsPlatform.Compiler == WindowsCompiler.VisualStudio2015 && WindowsPlatform.bUseWindowsSDK10)
 			{
 				if (Directory.Exists(EnvVars.WindowsSDKExtensionDir))
@@ -220,7 +220,7 @@ namespace UnrealBuildTool
 				Arguments.Append(" /wd4819");
 			}
 
-			// @todo UAP: UE4 is non-compliant when it comes to use of %s and %S
+			// @todo UWP: UE4 is non-compliant when it comes to use of %s and %S
 			// Previously %s meant "the current character set" and %S meant "the other one".
 			// Now %s means multibyte and %S means wide. %Ts means "natural width".
 			// Reverting this behaviour until the UE4 source catches up.
@@ -229,7 +229,7 @@ namespace UnrealBuildTool
 				Arguments.Append(" /D_CRT_STDIO_LEGACY_WIDE_SPECIFIERS=1");
 			}
 
-			// @todo UAP: Silence the hash_map deprecation errors for now. This should be replaced with unordered_map for the real fix.
+			// @todo UWP: Silence the hash_map deprecation errors for now. This should be replaced with unordered_map for the real fix.
 			if (WindowsPlatform.Compiler == WindowsCompiler.VisualStudio2015)
 			{
 				Arguments.Append(" /D_SILENCE_STDEXT_HASH_DEPRECATION_WARNINGS");
@@ -397,8 +397,8 @@ namespace UnrealBuildTool
 						(CompileEnvironment.Config.bEnableExceptions || UEBuildConfiguration.bBuildEditor || UEBuildConfiguration.bForceEnableExceptions))
 					{
 						// Enable C++ exception handling, but not C exceptions.
-						Arguments.Append(" /EHsc");
-					}
+							Arguments.Append(" /EHsc");
+						}
 					else
 					{
 						// This is required to disable exception handling in VC platform headers.
@@ -1440,7 +1440,7 @@ namespace UnrealBuildTool
 		{
 			var EnvVars = VCEnvironment.SetEnvironment(LinkEnvironment.Config.Target.Platform);
 
-			// @todo UAP: These paths should be added in SetUpEnvironment(), not here.  Also is this actually needed for classic desktop targets or only UAP?
+			// @todo UWP: These paths should be added in SetUpEnvironment(), not here.  Also is this actually needed for classic desktop targets or only UWP?
 			if (WindowsPlatform.Compiler == WindowsCompiler.VisualStudio2015 && WindowsPlatform.bUseWindowsSDK10)
 			{
 				if (LinkEnvironment.Config.Target.Platform == CPPTargetPlatform.Win64)
