fix for trying to create an SSH key when the remote server name was unset
UE-15438
#ios

diff --git a/Engine/Source/Programs/UnrealBuildTool/ToolChain/RemoteToolChain.cs b/Engine/Source/Programs/UnrealBuildTool/ToolChain/RemoteToolChain.cs
--- a/Engine/Source/Programs/UnrealBuildTool/ToolChain/RemoteToolChain.cs
+++ b/Engine/Source/Programs/UnrealBuildTool/ToolChain/RemoteToolChain.cs
@@ -332,6 +332,7 @@ namespace UnrealBuildTool
 				{
 					Log.TraceError("Remote compiling requires a server name. Use the editor to set up your remote compilation settings.");
 					InitializationErrorCode = 99;
+					return InitializationErrorCode;
 				}
 
 				if (!bUseRPCUtil)
