Include native Android glue in Launch module
#ue4
#android
#ue-15473

diff --git a/Engine/Source/Programs/UnrealBuildTool/Android/AndroidToolChain.cs b/Engine/Source/Programs/UnrealBuildTool/Android/AndroidToolChain.cs
--- a/Engine/Source/Programs/UnrealBuildTool/Android/AndroidToolChain.cs
+++ b/Engine/Source/Programs/UnrealBuildTool/Android/AndroidToolChain.cs
@@ -11,8 +11,6 @@ namespace UnrealBuildTool
 {
 	public class AndroidToolChain : UEToolChain
 	{
-		static private string ModuleHasNDKExtensionsCompiled = null;
-
 		// the number of the clang version being used to compile 
 		static private float ClangVersionFloat = 0;
 
@@ -568,15 +566,11 @@ namespace UnrealBuildTool
 
 		static void ConditionallyAddNDKSourceFiles(List<FileItem> SourceFiles, string ModuleName)
 		{
-			// We need to add the extra glue and cpu code only once to the first module. But since
-			// this is called multiple times for the same module we remember the module we
-			// add the extra code to and add it to that always. I.e. it makes this call
-			// consistent within one UBT invocation.
-			if (null == ModuleHasNDKExtensionsCompiled || ModuleHasNDKExtensionsCompiled.Equals(ModuleName))
+			// We need to add the extra glue and cpu code only to Launch module.
+			if (ModuleName.Equals("Launch"))
 			{
 				SourceFiles.Add(FileItem.GetItemByPath(Environment.GetEnvironmentVariable("NDKROOT") + "/sources/android/native_app_glue/android_native_app_glue.c"));
 				SourceFiles.Add(FileItem.GetItemByPath(Environment.GetEnvironmentVariable("NDKROOT") + "/sources/android/cpufeatures/cpu-features.c"));
-				ModuleHasNDKExtensionsCompiled = ModuleName;
 			}
 		}
 
@@ -610,10 +604,7 @@ namespace UnrealBuildTool
 			}
 
 			// Directly added NDK files for NDK extensions
-			if (!UnrealBuildTool.RunningRocket())
-			{
-				ConditionallyAddNDKSourceFiles(SourceFiles, ModuleName);
-			}
+			ConditionallyAddNDKSourceFiles(SourceFiles, ModuleName);
 
 			// Add preprocessor definitions to the argument list.
 			foreach (string Definition in CompileEnvironment.Config.Definitions)
