Use a clean compile environment to build shared PCHs, instead of reusing the current module compile environment.

diff --git a/Engine/Source/Programs/UnrealBuildTool/Configuration/UEBuildModule.cs b/Engine/Source/Programs/UnrealBuildTool/Configuration/UEBuildModule.cs
--- a/Engine/Source/Programs/UnrealBuildTool/Configuration/UEBuildModule.cs
+++ b/Engine/Source/Programs/UnrealBuildTool/Configuration/UEBuildModule.cs
@@ -1490,30 +1490,44 @@ namespace UnrealBuildTool
 					{
 						if( SharedPCHHeaderFile != null || CPPFilesToBuild.Count >= MinFilesUsingPrecompiledHeader )
 						{
-							bool bAllowDLLExports = true;
-							var PCHOutputDirectory = ModuleCompileEnvironment.Config.OutputDirectory;
-							var PCHModuleName = this.Name;
-
-							if( SharedPCHHeaderFile != null )
+							CPPOutput PCHOutput;
+							if (SharedPCHHeaderFile == null)
 							{
-								// Disallow DLLExports when generating shared PCHs.  These headers aren't able to export anything, because they're potentially shared between many modules.
-								bAllowDLLExports = false;
-
-								// Save shared PCHs to a specific folder
-								PCHOutputDirectory = Path.Combine( CompileEnvironment.Config.OutputDirectory, "SharedPCHs" );
-
-								// Use a fake module name for "shared" PCHs.  It may be used by many modules, so we don't want to use this module's name.
-								PCHModuleName = "Shared";
+								PCHOutput = PrecompileHeaderEnvironment.GeneratePCHCreationAction( 
+									Target,
+									CPPFilesToBuild[0].PCHHeaderNameInCode,
+									ModulePCHEnvironment.PrecompiledHeaderIncludeFilename,
+									ModuleCompileEnvironment, 
+									ModuleCompileEnvironment.Config.OutputDirectory,
+									Name, 
+									true );
+							}
+							else
+							{
+								UEBuildModuleCPP SharedPCHModule = (UEBuildModuleCPP)Target.FindOrCreateModuleByName(SharedPCHModuleName);
+
+								CPPEnvironment SharedPCHCompileEnvironment = GlobalCompileEnvironment.DeepCopy();
+								SharedPCHCompileEnvironment.Config.bEnableShadowVariableWarning = SharedPCHModule.bEnableShadowVariableWarnings;
+
+								SharedPCHModule.SetupPublicCompileEnvironment(
+									Binary, 
+									false, 
+									SharedPCHCompileEnvironment.Config.CPPIncludeInfo.IncludePaths, 
+									SharedPCHCompileEnvironment.Config.CPPIncludeInfo.SystemIncludePaths, 
+									SharedPCHCompileEnvironment.Config.Definitions, 
+									SharedPCHCompileEnvironment.Config.AdditionalFrameworks, 
+									new Dictionary<UEBuildModule,bool>());
+
+								PCHOutput = PrecompileHeaderEnvironment.GeneratePCHCreationAction( 
+									Target,
+									CPPFilesToBuild[0].PCHHeaderNameInCode,
+									ModulePCHEnvironment.PrecompiledHeaderIncludeFilename,
+									SharedPCHCompileEnvironment,
+									Path.Combine( CompileEnvironment.Config.OutputDirectory, "SharedPCHs" ),
+									"Shared", 
+									false );
 							}
 
-							var PCHOutput = PrecompileHeaderEnvironment.GeneratePCHCreationAction( 
-								Target,
-								CPPFilesToBuild[0].PCHHeaderNameInCode,
-								ModulePCHEnvironment.PrecompiledHeaderIncludeFilename,
-								ModuleCompileEnvironment, 
-								PCHOutputDirectory,
-								PCHModuleName, 
-								bAllowDLLExports );
 							ModulePCHEnvironment.PrecompiledHeaderFile = PCHOutput.PrecompiledHeaderFile;
 							
 							ModulePCHEnvironment.OutputObjectFiles.Clear();
