fix for trying to generate an SSH key when the user name has not been set
UE-15518
#ios

diff --git a/Engine/Source/Programs/UnrealBuildTool/ToolChain/RemoteToolChain.cs b/Engine/Source/Programs/UnrealBuildTool/ToolChain/RemoteToolChain.cs
--- a/Engine/Source/Programs/UnrealBuildTool/ToolChain/RemoteToolChain.cs
+++ b/Engine/Source/Programs/UnrealBuildTool/ToolChain/RemoteToolChain.cs
@@ -399,7 +399,7 @@ namespace UnrealBuildTool
 				InitializationErrorCode = RPCUtilHelper.Initialize(RemoteServerName);
 
 				// allow user to set up
-				if (InitializationErrorCode == 100)
+				if (InitializationErrorCode == 100 && !string.IsNullOrEmpty(ResolvedRSyncUsername))
 				{
 					Process KeyProcess = new Process();
 					KeyProcess.StartInfo.WorkingDirectory = Path.GetFullPath(Path.Combine(BuildConfiguration.RelativeEnginePath, "Build", "BatchFiles"));
@@ -423,6 +423,12 @@ namespace UnrealBuildTool
 						InitializeRemoteExecution();
 					}
 				}
+				else if (string.IsNullOrEmpty(ResolvedRSyncUsername))
+				{
+					Log.TraceError("Remote compiling requires a user name. Use the editor to set up your remote compilation settings.");
+					InitializationErrorCode = 98;
+					return InitializationErrorCode;
+				}
 			}
             else
 			{
