UE-15116: Engine *.generated.cpp files change during hot-reload

diff --git a/Engine/Source/Programs/UnrealBuildTool/Configuration/UEBuildModule.cs b/Engine/Source/Programs/UnrealBuildTool/Configuration/UEBuildModule.cs
--- a/Engine/Source/Programs/UnrealBuildTool/Configuration/UEBuildModule.cs
+++ b/Engine/Source/Programs/UnrealBuildTool/Configuration/UEBuildModule.cs
@@ -997,6 +997,32 @@ namespace UnrealBuildTool
 					       OtherFiles  .Count;
 				}
 			}
+
+			/// <summary>
+			/// Copy from list to list helper.
+			/// </summary>
+			/// <param name="From">Source list.</param>
+			/// <param name="To">Destination list.</param>
+			private static void CopyFromListToList(List<FileItem> From, List<FileItem> To)
+			{
+				To.Clear();
+				To.AddRange(From);
+			}
+
+			/// <summary>
+			/// Copies file lists from other SourceFilesClass to this.
+			/// </summary>
+			/// <param name="Other">Source object.</param>
+			public void CopyFrom(SourceFilesClass Other)
+			{
+				CopyFromListToList(Other.MissingFiles, MissingFiles);
+				CopyFromListToList(Other.CPPFiles, CPPFiles);
+				CopyFromListToList(Other.CFiles, CFiles);
+				CopyFromListToList(Other.CCFiles, CCFiles);
+				CopyFromListToList(Other.MMFiles, MMFiles);
+				CopyFromListToList(Other.RCFiles, RCFiles);
+				CopyFromListToList(Other.OtherFiles, OtherFiles);
+			}
 		}
 
 		/**
@@ -1012,6 +1038,9 @@ namespace UnrealBuildTool
 		/** A list of the absolute paths of source files to be built in this module. */
 		public readonly SourceFilesClass SourceFilesToBuild = new SourceFilesClass();
 
+		/** A list of the source files that were found for the module. */
+		public readonly SourceFilesClass SourceFilesFound = new SourceFilesClass();
+
 		/** The directory for this module's generated code */
 		public readonly string GeneratedCodeDirectory;
 
@@ -1173,9 +1202,10 @@ namespace UnrealBuildTool
 			GeneratedCodeDirectory = InGeneratedCodeDirectory;
 			IntelliSenseGatherer = InIntelliSenseGatherer;
 
+			CategorizeSourceFiles(InSourceFiles, SourceFilesFound);
 			if (bInBuildSourceFiles)
 			{
-				CategorizeSourceFiles(InSourceFiles, SourceFilesToBuild);
+				SourceFilesToBuild.CopyFrom(SourceFilesFound);
 			}
 
 			Definitions = HashSetFromOptionalEnumerableStringParameter(InDefinitions);
@@ -1692,7 +1722,7 @@ namespace UnrealBuildTool
 				bool bFoundAProblemWithPCHs = false;
 
 				FileItem UniquePCH = null;
-				foreach( var CPPFile in SourceFilesToBuild.CPPFiles )	// @todo ubtmake: We're not caching CPPEnvironments for .c/.mm files, etc.  Even though they don't use PCHs, they still have #includes!  This can break dependency checking!
+				foreach( var CPPFile in SourceFilesFound.CPPFiles )	// @todo ubtmake: We're not caching CPPEnvironments for .c/.mm files, etc.  Even though they don't use PCHs, they still have #includes!  This can break dependency checking!
 				{
 					// Build a single list of include paths to search.
 					var IncludePathsToSearch = ModuleCompileEnvironment.Config.CPPIncludeInfo.GetIncludesPathsToSearch( CPPFile );
