Remove receipts from project files.

diff --git a/Engine/Source/Programs/UnrealBuildTool/System/ProjectFileGenerator.cs b/Engine/Source/Programs/UnrealBuildTool/System/ProjectFileGenerator.cs
--- a/Engine/Source/Programs/UnrealBuildTool/System/ProjectFileGenerator.cs
+++ b/Engine/Source/Programs/UnrealBuildTool/System/ProjectFileGenerator.cs
@@ -1904,10 +1904,8 @@ namespace UnrealBuildTool
 			var DirectoriesToSearch = new List<string>();
 			DirectoriesToSearch.Add( BuildDirectory );
 			var SubdirectoryNamesToExclude = new List<string>();
-			{
-				// Nothing to exclude, yet!
-				// SubdirectoryNamesToExclude.Add( "DirectoryName" );
-			}
+			SubdirectoryNamesToExclude.Add("Receipts");
+
 			EngineProject.AddFilesToProject( SourceFileSearch.FindFiles(
 				DirectoriesToSearch: DirectoriesToSearch,
 				ExcludeNoRedistFiles: bExcludeNoRedistFiles,
