Windows XP games now run without a 'not a valid win32 application' error. BootstrapPackagedGame had lost it's xp support.

UE-14290

diff --git a/Engine/Source/Programs/UnrealBuildTool/System/RulesCompiler.cs b/Engine/Source/Programs/UnrealBuildTool/System/RulesCompiler.cs
--- a/Engine/Source/Programs/UnrealBuildTool/System/RulesCompiler.cs
+++ b/Engine/Source/Programs/UnrealBuildTool/System/RulesCompiler.cs
@@ -443,6 +443,13 @@ namespace UnrealBuildTool
 		/// </summary>
         public bool bUseStaticCRT = false;
 
+		
+        /// <summary>
+		/// Allow a target to specify a preferred sub-platform.
+		/// Can be used to target a build using sub platform specifics.
+		/// </summary>
+		public string PreferredSubPlatform = String.Empty;
+
         /// <summary>
         /// By default we use the Release C++ Runtime (CRT), even when compiling Debug builds.  This is because the Debug C++
         /// Runtime isn't very useful when debugging Unreal Engine projects, and linking against the Debug CRT libraries forces
diff --git a/Engine/Source/Programs/UnrealBuildTool/Windows/UEBuildWindows.cs b/Engine/Source/Programs/UnrealBuildTool/Windows/UEBuildWindows.cs
--- a/Engine/Source/Programs/UnrealBuildTool/Windows/UEBuildWindows.cs
+++ b/Engine/Source/Programs/UnrealBuildTool/Windows/UEBuildWindows.cs
@@ -102,10 +102,7 @@ namespace UnrealBuildTool
         /// True if we're targeting Windows XP as a minimum spec.  In Visual Studio 2012 and higher, this may change how
         /// we compile and link the application (http://blogs.msdn.com/b/vcblog/archive/2012/10/08/10357555.aspx)
 		/// This is a flag to determine we should support XP if possible from XML
-        [XmlConfig]
-		public static bool SupportWindowsXPIfAvailable = false;
-
-		private static bool SupportWindowsXP;
+		public static bool SupportWindowsXP;
 		public static bool IsWindowsXPSupported()
 		{
 			return SupportWindowsXP;
@@ -382,36 +379,38 @@ namespace UnrealBuildTool
             return false;
         }
 
-        public override void ResetBuildConfiguration(UnrealTargetPlatform InPlatform, UnrealTargetConfiguration InConfiguration)
-        {
-			UEBuildConfiguration.bCompileICU = true;
+		private void SetupXPSupportFromConfiguration()
+		{
+			string[] CmdLine = Environment.GetCommandLineArgs();
 
-			// Should we enable Windows XP support
+			bool SupportWindowsXPIfAvailable = false;
+			SupportWindowsXPIfAvailable = CmdLine.Contains("-winxp", StringComparer.InvariantCultureIgnoreCase);
+
+			// ...check if it was supported from a config.
+			if (!SupportWindowsXPIfAvailable)
 			{
-				// If it wasnt set as an XML flag, check if it is requested from the commandline.
-				if (!SupportWindowsXPIfAvailable)
+				ConfigCacheIni Ini = new ConfigCacheIni(UnrealTargetPlatform.Win64, "Engine", UnrealBuildTool.GetUProjectPath());
+				string MinimumOS;
+				if (Ini.GetString("/Script/WindowsTargetPlatform.WindowsTargetSettings", "MinimumOSVersion", out MinimumOS))
 				{
-					string[] CmdLine = Environment.GetCommandLineArgs();
-					SupportWindowsXPIfAvailable = CmdLine.Contains("-winxp", StringComparer.InvariantCultureIgnoreCase);
-
-					// ...check if it was supported from a config.
-					if( !SupportWindowsXPIfAvailable )
+					if (string.IsNullOrEmpty(MinimumOS) == false)
 					{
-						ConfigCacheIni Ini = new ConfigCacheIni(UnrealTargetPlatform.Win64, "Engine", UnrealBuildTool.GetUProjectPath());
-						string MinimumOS;
-						if (Ini.GetString("/Script/WindowsTargetPlatform.WindowsTargetSettings", "MinimumOSVersion", out MinimumOS))
-						{
-							if (string.IsNullOrEmpty(MinimumOS) == false)
-							{
-								SupportWindowsXPIfAvailable = MinimumOS == "MSOS_XP";
-							}
-						}
+						SupportWindowsXPIfAvailable = MinimumOS == "MSOS_XP";
 					}
 				}
+			}
+
+			SupportWindowsXP = SupportWindowsXPIfAvailable;
+		}
 
-				// Win32 XP is only supported at this time.
-				SupportWindowsXP = SupportWindowsXPIfAvailable && (GetCPPTargetPlatform(InPlatform) == CPPTargetPlatform.Win32);
+        public override void ResetBuildConfiguration(UnrealTargetPlatform InPlatform, UnrealTargetConfiguration InConfiguration)
+        {
+			UEBuildConfiguration.bCompileICU = true;
+			if (InPlatform == UnrealTargetPlatform.Win32)
+			{
+				SetupXPSupportFromConfiguration();
 			}
+
         }
 
         public override void ValidateBuildConfiguration(CPPTargetConfiguration Configuration, CPPTargetPlatform Platform, bool bCreateDebugInfo)
@@ -487,6 +486,16 @@ namespace UnrealBuildTool
         {
             InBuildTarget.GlobalCompileEnvironment.Config.Definitions.Add("WIN32=1");
 
+			// Should we enable Windows XP support
+			if ((InBuildTarget.TargetType == TargetRules.TargetType.Program) && (GetCPPTargetPlatform(InBuildTarget.Platform) == CPPTargetPlatform.Win32))
+			{
+				// Check if the target has requested XP support.
+				if (String.Equals(InBuildTarget.Rules.PreferredSubPlatform, "WindowsXP", StringComparison.InvariantCultureIgnoreCase))
+				{
+					SupportWindowsXP = true;
+				}
+			}
+
 			if (WindowsPlatform.bUseWindowsSDK10 && WindowsPlatform.Compiler == WindowsCompiler.VisualStudio2015)
 			{
 				if (SupportWindowsXP)
@@ -499,7 +508,6 @@ namespace UnrealBuildTool
 			}
 			else
 			{
-				SupportWindowsXP = SupportWindowsXP && (GetCPPTargetPlatform(InBuildTarget.Platform) == CPPTargetPlatform.Win32);
 				if (SupportWindowsXP)
 				{
 					// Windows XP SP3 or higher required
diff --git a/Engine/Source/Programs/Windows/BootstrapPackagedGame/BootstrapPackagedGame.Target.cs b/Engine/Source/Programs/Windows/BootstrapPackagedGame/BootstrapPackagedGame.Target.cs
--- a/Engine/Source/Programs/Windows/BootstrapPackagedGame/BootstrapPackagedGame.Target.cs
+++ b/Engine/Source/Programs/Windows/BootstrapPackagedGame/BootstrapPackagedGame.Target.cs
@@ -10,6 +10,11 @@ public class BootstrapPackagedGameTarget : TargetRules
 		Type = TargetType.Program;
 
 		bUseStaticCRT = true;
+
+		if (Target.Platform == UnrealTargetPlatform.Win32)
+		{
+			PreferredSubPlatform = "WindowsXP";
+		}
 	}
 
 	//
@@ -49,11 +54,6 @@ public class BootstrapPackagedGameTarget : TargetRules
 		UEBuildConfiguration.bBuildWithEditorOnlyData = false;
 		UEBuildConfiguration.bCompileAgainstEngine = false;
 		UEBuildConfiguration.bCompileAgainstCoreUObject = false;
-
-		if (Target.Platform == UnrealTargetPlatform.Win32)
-		{
-			WindowsPlatform.SupportWindowsXPIfAvailable = true;
-		}
 	}
 	
     public override bool GUBP_AlwaysBuildWithTools(UnrealTargetPlatform InHostPlatform, out bool bInternalToolOnly, out bool SeparateNode, out bool CrossCompile)
