Added application display name (app_name) setting for Android
#ue-15652
#ue4
#android

diff --git a/Engine/Source/Programs/UnrealBuildTool/Android/UEDeployAndroid.cs b/Engine/Source/Programs/UnrealBuildTool/Android/UEDeployAndroid.cs
--- a/Engine/Source/Programs/UnrealBuildTool/Android/UEDeployAndroid.cs
+++ b/Engine/Source/Programs/UnrealBuildTool/Android/UEDeployAndroid.cs
@@ -567,6 +567,62 @@ namespace UnrealBuildTool.Android
 			}
 		}
 
+		private bool CheckApplicationName(string UE4BuildPath, string ProjectName, out string ApplicationDisplayName)
+		{
+			string StringsXMLPath = Path.Combine(UE4BuildPath, "res/values/strings.xml");
+
+			ApplicationDisplayName = null;
+			ConfigCacheIni Ini = GetConfigCacheIni("Engine");
+			Ini.GetString("/Script/AndroidRuntimeSettings.AndroidRuntimeSettings", "ApplicationDisplayName", out ApplicationDisplayName);
+
+			// use project name if display name is left blank
+			if (String.IsNullOrWhiteSpace(ApplicationDisplayName))
+			{
+				ApplicationDisplayName = ProjectName;
+			}
+
+			// make sure name does not have < or >
+			ApplicationDisplayName = ApplicationDisplayName.Replace("<", "(").Replace(">",")");
+
+			// if it doesn't exist, need to repackage
+			if (!File.Exists(StringsXMLPath))
+			{
+				return true;
+			}
+
+			// read it and see if needs to be updated
+			string Contents = File.ReadAllText(StringsXMLPath);
+
+			// find the key
+			string AppNameTag = "<string name=\"app_name\">";
+			int KeyIndex = Contents.IndexOf(AppNameTag);
+
+			// if doesn't exist, need to repackage
+			if (KeyIndex < 0)
+			{
+				return true;
+			}
+
+			// get the current value
+			KeyIndex += AppNameTag.Length;
+			int TagEnd = Contents.IndexOf("</string>", KeyIndex);
+			if (TagEnd < 0)
+			{
+				return true;
+			}
+			string CurrentApplicationName = Contents.Substring(KeyIndex, TagEnd - KeyIndex);
+
+			// no need to do anything if matches
+			if (CurrentApplicationName == ApplicationDisplayName)
+			{
+				// name matches, no need to force a repackage
+				return false;
+			}
+
+			// need to repackage
+			return true;
+		}
+
 		private void UpdateProjectProperties(string UE4BuildPath, string ProjectName)
 		{
 			Log.TraceInformation("\n===={0}====UPDATING BUILD CONFIGURATION FILES====================================================", DateTime.Now.ToString());
@@ -1126,6 +1182,14 @@ namespace UnrealBuildTool.Android
 				Log.TraceInformation("AndroidManifest.xml that was generated is different than last build, forcing repackage.");
 			}
 
+			// get application name and whether it changed, needing to force repackage
+			string ApplicationDisplayName;
+			if (CheckApplicationName(UE4BuildPath, ProjectName, out ApplicationDisplayName))
+			{
+				bBuildSettingsMatch = false;
+				Log.TraceInformation("Application display name is different than last build, forcing repackage.");
+			}
+
 			// if the manifest matches, look at other settings stored in a file
 			if (bBuildSettingsMatch)
 			{
@@ -1165,15 +1229,9 @@ namespace UnrealBuildTool.Android
 
 			// Once for all arches code:
 
-			// make up a dictionary of strings to replace in the Manifest file
+			// make up a dictionary of strings to replace in xml files (strings.xml)
 			Dictionary<string, string> Replacements = new Dictionary<string, string>();
-			Replacements.Add("${EXECUTABLE_NAME}", ProjectName);
-
-			// distribution apps can't be debuggable, so if it was set to true, set it to false:
-			if (bForDistribution)
-			{
-				Replacements.Add("android:debuggable=\"true\"", "android:debuggable=\"false\"");
-			}
+			Replacements.Add("${EXECUTABLE_NAME}", ApplicationDisplayName);
 
 			if (!bIncrementalPackage)
 			{
diff --git a/Engine/Source/Runtime/Android/AndroidRuntimeSettings/Classes/AndroidRuntimeSettings.h b/Engine/Source/Runtime/Android/AndroidRuntimeSettings/Classes/AndroidRuntimeSettings.h
--- a/Engine/Source/Runtime/Android/AndroidRuntimeSettings/Classes/AndroidRuntimeSettings.h
+++ b/Engine/Source/Runtime/Android/AndroidRuntimeSettings/Classes/AndroidRuntimeSettings.h
@@ -112,6 +112,10 @@ public:
 	UPROPERTY(GlobalConfig, EditAnywhere, Category = APKPackaging, Meta = (DisplayName = "Store Version (1-65535)", ClampMin="1", ClampMax="65535"))
 	int32 StoreVersion;
 
+	// The visual application name displayed for end users
+	UPROPERTY(GlobalConfig, EditAnywhere, Category = APKPackaging, Meta = (DisplayName = "Application Display Name (app_name), project name if blank"))
+	FString ApplicationDisplayName;
+
 	// The visual version displayed for end users
 	UPROPERTY(GlobalConfig, EditAnywhere, Category = APKPackaging, Meta = (DisplayName = "Version Display Name (usually x.y)"))
 	FString VersionDisplayName;
