Fix some optional modules not being hooked up to their precompiled binaries when running in Rocket. Fixes OnlineSubsystemNull/OnlineSubsystemSteam link errors for ShooterGame.

diff --git a/Engine/Source/Programs/UnrealBuildTool/Configuration/UEBuildTarget.cs b/Engine/Source/Programs/UnrealBuildTool/Configuration/UEBuildTarget.cs
--- a/Engine/Source/Programs/UnrealBuildTool/Configuration/UEBuildTarget.cs
+++ b/Engine/Source/Programs/UnrealBuildTool/Configuration/UEBuildTarget.cs
@@ -2303,6 +2303,11 @@ namespace UnrealBuildTool
 							TargetModule.RecursivelyAddPrecompiledModules(TargetModules);
 						}
 					}
+					foreach(string ExtraModuleName in ExtraModuleNames)
+					{
+						UEBuildModule ExtraModule = FindOrCreateModuleByName(ExtraModuleName);
+						ExtraModule.RecursivelyAddPrecompiledModules(TargetModules);
+					}
 					foreach(UEBuildModuleCPP TargetModule in TargetModules)
 					{
 						if(TargetModule is UEBuildModuleCPP && !RulesCompiler.IsGameModule(TargetModule.Name))
