Removed unnecessary Rocket checks from code used by Analytics

diff --git a/Engine/Source/Programs/UnrealBuildTool/System/UnrealBuildTool.cs b/Engine/Source/Programs/UnrealBuildTool/System/UnrealBuildTool.cs
--- a/Engine/Source/Programs/UnrealBuildTool/System/UnrealBuildTool.cs
+++ b/Engine/Source/Programs/UnrealBuildTool/System/UnrealBuildTool.cs
@@ -1179,7 +1179,6 @@ namespace UnrealBuildTool
                                                 		: "Build",
                         "Platform", CheckPlatform.ToString(),
                         "Configuration", CheckConfiguration.ToString(),
-                        "IsRocket", bRunningRocket.ToString(),
                         "EngineVersion", Utils.GetEngineVersionFromObjVersionCPP().ToString()
                         );
 
diff --git a/Engine/Source/Programs/UnrealBuildTool/System/Utils.cs b/Engine/Source/Programs/UnrealBuildTool/System/Utils.cs
--- a/Engine/Source/Programs/UnrealBuildTool/System/Utils.cs
+++ b/Engine/Source/Programs/UnrealBuildTool/System/Utils.cs
@@ -932,22 +932,19 @@ namespace UnrealBuildTool
 		/// <returns></returns>
 		public static int GetEngineVersionFromObjVersionCPP()
 		{
-			if(UnrealBuildTool.RunningRocket() == false)
+			try
 			{
-				try
-				{
-					return
-						(from line in File.ReadLines("Runtime/Core/Private/UObject/ObjectVersion.cpp", Encoding.ASCII)
-						 where line.StartsWith("#define	ENGINE_VERSION")
-						 select int.Parse(line.Split()[2])).Single();
-				}
-				catch (Exception ex)
-				{
-					// Don't do a stack trace so we don't pollute the logs with spurious exception data, as we don't crash on this case.
-					Log.TraceWarning("Could not parse Engine Version from ObjectVersion.cpp: {0}", ex.Message);
-				}
+				return
+					(from line in File.ReadLines("Runtime/Core/Private/UObject/ObjectVersion.cpp", Encoding.ASCII)
+						where line.StartsWith("#define	ENGINE_VERSION")
+						select int.Parse(line.Split()[2])).Single();
+			}
+			catch (Exception ex)
+			{
+				// Don't do a stack trace so we don't pollute the logs with spurious exception data, as we don't crash on this case.
+				Log.TraceWarning("Could not parse Engine Version from ObjectVersion.cpp: {0}", ex.Message);
+				return 0;
 			}
-			return 0;
 		}
 
 		/// <summary>
