UBT cleanups:
* Removal of some null checks.
* Some arrays changed to Lists.
* Some lists changed to IEnumerables.
* Some loops changed to Linq actions.
* FilterOnlyModules and FilterGameModules renamed to GetFiltered* functions, which better represent what they do.
* Some indentation fixes.
* Some thrown exceptions moved around to more appropriate locations.

diff --git a/Engine/Source/Programs/UnrealBuildTool/Android/AndroidToolChain.cs b/Engine/Source/Programs/UnrealBuildTool/Android/AndroidToolChain.cs
--- a/Engine/Source/Programs/UnrealBuildTool/Android/AndroidToolChain.cs
+++ b/Engine/Source/Programs/UnrealBuildTool/Android/AndroidToolChain.cs
@@ -811,7 +811,7 @@ namespace UnrealBuildTool
 					int OutputPathIndex = ArchIndex * GPUArchitectures.Length + GPUArchIndex;
 
 					// Android will have an array of outputs
-					if (LinkEnvironment.Config.OutputFilePaths.Length < OutputPathIndex ||
+					if (LinkEnvironment.Config.OutputFilePaths.Count < OutputPathIndex ||
 						!Path.GetFileNameWithoutExtension(LinkEnvironment.Config.OutputFilePaths[OutputPathIndex]).EndsWith(Arch + GPUArchitecture))
 					{
 						throw new BuildException("The OutputFilePaths array didn't match the Arches array in AndroidToolChain.LinkAllFiles");
diff --git a/Engine/Source/Programs/UnrealBuildTool/Android/UEBuildAndroid.cs b/Engine/Source/Programs/UnrealBuildTool/Android/UEBuildAndroid.cs
--- a/Engine/Source/Programs/UnrealBuildTool/Android/UEBuildAndroid.cs
+++ b/Engine/Source/Programs/UnrealBuildTool/Android/UEBuildAndroid.cs
@@ -428,7 +428,7 @@ namespace UnrealBuildTool
 			};
 		}
 
-		public override string[] FinalizeBinaryPaths(string BinaryName)
+		public override List<string> FinalizeBinaryPaths(string BinaryName)
 		{
 			string[] Architectures = AndroidToolChain.GetAllArchitectures();
 			string[] GPUArchitectures = AndroidToolChain.GetAllGPUArchitectures();
@@ -443,7 +443,7 @@ namespace UnrealBuildTool
 				}
 			}
 
-			return AllBinaries.ToArray();
+			return AllBinaries;
 		}
 	}
 }
diff --git a/Engine/Source/Programs/UnrealBuildTool/Android/UEDeployAndroid.cs b/Engine/Source/Programs/UnrealBuildTool/Android/UEDeployAndroid.cs
--- a/Engine/Source/Programs/UnrealBuildTool/Android/UEDeployAndroid.cs
+++ b/Engine/Source/Programs/UnrealBuildTool/Android/UEDeployAndroid.cs
@@ -1476,7 +1476,7 @@ namespace UnrealBuildTool.Android
 				bMakeSeparateApks:ShouldMakeSeparateApks(), bIncrementalPackage:true, bDisallowPackagingDataInApk:false);
 
 			// if we made any non-standard .apk files, the generated debugger settings may be wrong
-			if (ShouldMakeSeparateApks() && (InTarget.OutputPaths.Length > 1 || !InTarget.OutputPaths[0].Contains("-armv7-es2")))
+			if (ShouldMakeSeparateApks() && (InTarget.OutputPaths.Count > 1 || !InTarget.OutputPaths[0].Contains("-armv7-es2")))
 			{
 				Console.WriteLine("================================================================================================================================");
 				Console.WriteLine("Non-default apk(s) have been made: If you are debugging, you will need to manually select one to run in the debugger properties!");
diff --git a/Engine/Source/Programs/UnrealBuildTool/Configuration/UEBuildBinary.cs b/Engine/Source/Programs/UnrealBuildTool/Configuration/UEBuildBinary.cs
--- a/Engine/Source/Programs/UnrealBuildTool/Configuration/UEBuildBinary.cs
+++ b/Engine/Source/Programs/UnrealBuildTool/Configuration/UEBuildBinary.cs
@@ -38,7 +38,7 @@ namespace UnrealBuildTool
 		/// <summary>
 		/// The output file path. This must be set before a binary can be built using it.
 		/// </summary>
-		public string[] OutputFilePaths;
+		public List<string> OutputFilePaths = new List<string>();
 
 		/// <summary>
 		/// Returns the OutputFilePath if there is only one entry in OutputFilePaths
@@ -47,9 +47,9 @@ namespace UnrealBuildTool
 		{
 			get
 			{
-				if (OutputFilePaths.Length != 1)
+				if (OutputFilePaths.Count != 1)
 				{
-					throw new BuildException("Attempted to use UEBuildBinaryConfiguration.OutputFilePath property, but there are multiple (or no) OutputFilePaths. You need to handle multiple in the code that called this (size = {0})", OutputFilePaths.Length);
+					throw new BuildException("Attempted to use UEBuildBinaryConfiguration.OutputFilePath property, but there are multiple (or no) OutputFilePaths. You need to handle multiple in the code that called this (size = {0})", OutputFilePaths.Count);
 				}
 				return OutputFilePaths[0];
 			}
@@ -58,7 +58,7 @@ namespace UnrealBuildTool
 		/// <summary>
 		/// Original output filepath. This is the original binary name before hot-reload suffix has been appended to it.
 		/// </summary>
-		public string[] OriginalOutputFilePaths;
+		public List<string> OriginalOutputFilePaths;
 
 		/// <summary>
 		/// Returns the OriginalOutputFilePath if there is only one entry in OriginalOutputFilePaths
@@ -67,9 +67,9 @@ namespace UnrealBuildTool
 		{
 			get
 			{
-				if (OriginalOutputFilePaths.Length != 1)
+				if (OriginalOutputFilePaths.Count != 1)
 				{
-					throw new BuildException("Attempted to use UEBuildBinaryConfiguration.OriginalOutputFilePath property, but there are multiple (or no) OriginalOutputFilePaths. You need to handle multiple in the code that called this (size = {0})", OriginalOutputFilePaths.Length);
+					throw new BuildException("Attempted to use UEBuildBinaryConfiguration.OriginalOutputFilePath property, but there are multiple (or no) OriginalOutputFilePaths. You need to handle multiple in the code that called this (size = {0})", OriginalOutputFilePaths.Count);
 				}
 				return OriginalOutputFilePaths[0];
 			}
@@ -138,7 +138,7 @@ namespace UnrealBuildTool
 		/// <param name="InModuleNames"></param>
 		public UEBuildBinaryConfiguration(
 				UEBuildBinaryType InType,
-				string[] InOutputFilePaths = null,
+				IEnumerable<string> InOutputFilePaths = null,
 				string InIntermediateDirectory = null,
 				bool bInAllowExports = false,
 				bool bInCreateImportLibrarySeparately = false,
@@ -147,11 +147,14 @@ namespace UnrealBuildTool
 				bool bInHasModuleRules = true,
                 bool bInIsCrossTarget = false,
 				string InProjectFilePath = "",
-				List<string> InModuleNames = null
+				IEnumerable<string> InModuleNames = null
 			)
 		{
 			Type = InType;
-			OutputFilePaths = InOutputFilePaths != null ? (string[])InOutputFilePaths.Clone() : null;
+			if (InOutputFilePaths != null)
+			{
+				OutputFilePaths.AddRange(InOutputFilePaths);
+			}
 			IntermediateDirectory = InIntermediateDirectory;
 			bAllowExports = bInAllowExports;
 			bCreateImportLibrarySeparately = bInCreateImportLibrarySeparately;
@@ -745,7 +748,7 @@ namespace UnrealBuildTool
 			}
 
 			// Set the link output file.
-			BinaryLinkEnvironment.Config.OutputFilePaths = Config.OutputFilePaths != null ? (string[])Config.OutputFilePaths.Clone() : null;
+			BinaryLinkEnvironment.Config.OutputFilePaths = Config.OutputFilePaths.ToList();
 
 			// Set whether the link is allowed to have exports.
 			BinaryLinkEnvironment.Config.bHasExports = Config.bAllowExports;
@@ -817,10 +820,7 @@ namespace UnrealBuildTool
 				var ConsoleAppLinkEvironment = BinaryLinkEnvironment.DeepCopy();
 				ConsoleAppLinkEvironment.Config.bIsBuildingConsoleApplication = true;
 				ConsoleAppLinkEvironment.Config.WindowsEntryPointOverride = "WinMainCRTStartup";		// For WinMain() instead of "main()" for Launch module
-				for (int Index = 0; Index < Config.OutputFilePaths.Length; Index++)
-				{
-					ConsoleAppLinkEvironment.Config.OutputFilePaths[Index] = GetAdditionalConsoleAppPath(ConsoleAppLinkEvironment.Config.OutputFilePaths[Index]);
-				}
+				ConsoleAppLinkEvironment.Config.OutputFilePaths = ConsoleAppLinkEvironment.Config.OutputFilePaths.Select(Path => GetAdditionalConsoleAppPath(Path)).ToList();
 
 				// Link the console app executable
 				OutputFiles.AddRange(ConsoleAppLinkEvironment.LinkExecutable(false));
diff --git a/Engine/Source/Programs/UnrealBuildTool/Configuration/UEBuildClient.cs b/Engine/Source/Programs/UnrealBuildTool/Configuration/UEBuildClient.cs
--- a/Engine/Source/Programs/UnrealBuildTool/Configuration/UEBuildClient.cs
+++ b/Engine/Source/Programs/UnrealBuildTool/Configuration/UEBuildClient.cs
@@ -3,6 +3,7 @@
 using System;
 using System.Collections.Generic;
 using System.IO;
+using System.Linq;
 using System.Text;
 
 namespace UnrealBuildTool
@@ -15,18 +16,13 @@ namespace UnrealBuildTool
         {
             if (ShouldCompileMonolithic())
             {
-                if ((UnrealBuildTool.IsDesktopPlatform(Platform) == false) ||
-                    (Platform == UnrealTargetPlatform.WinRT) ||
-                    (Platform == UnrealTargetPlatform.WinRT_ARM))
+                if (!UnrealBuildTool.IsDesktopPlatform(Platform) || Platform == UnrealTargetPlatform.WinRT || Platform == UnrealTargetPlatform.WinRT_ARM)
                 {
                     // We are compiling for a console...
                     // We want the output to go into the <GAME>\Binaries folder
-                    if (InRulesObject.bOutputToEngineBinaries == false)
+                    if (!InRulesObject.bOutputToEngineBinaries)
                     {
-						for (int Index = 0; Index < OutputPaths.Length; Index++ )
-						{
-							OutputPaths[Index] = OutputPaths[Index].Replace("Engine\\Binaries", InDesc.TargetName + "\\Binaries");
-						}
+						OutputPaths = OutputPaths.Select(Path => Path.Replace("Engine\\Binaries", InDesc.TargetName + "\\Binaries")).ToList();
                     }
                 }
             }
diff --git a/Engine/Source/Programs/UnrealBuildTool/Configuration/UEBuildGame.cs b/Engine/Source/Programs/UnrealBuildTool/Configuration/UEBuildGame.cs
--- a/Engine/Source/Programs/UnrealBuildTool/Configuration/UEBuildGame.cs
+++ b/Engine/Source/Programs/UnrealBuildTool/Configuration/UEBuildGame.cs
@@ -3,6 +3,7 @@
 using System;
 using System.Collections.Generic;
 using System.IO;
+using System.Linq;
 using System.Text;
 
 namespace UnrealBuildTool
@@ -15,18 +16,13 @@ namespace UnrealBuildTool
 		{
 			if (ShouldCompileMonolithic())
 			{
-				if ((UnrealBuildTool.IsDesktopPlatform(Platform) == false) ||
-					(Platform == UnrealTargetPlatform.WinRT) ||
-					(Platform == UnrealTargetPlatform.WinRT_ARM))
+				if (!UnrealBuildTool.IsDesktopPlatform(Platform) || Platform == UnrealTargetPlatform.WinRT || Platform == UnrealTargetPlatform.WinRT_ARM)
 				{
 					// We are compiling for a console...
 					// We want the output to go into the <GAME>\Binaries folder
-					if (InRulesObject.bOutputToEngineBinaries == false)
+					if (!InRulesObject.bOutputToEngineBinaries)
 					{
-						for (int Index = 0; Index < OutputPaths.Length; Index++)
-						{
-							OutputPaths[Index] = OutputPaths[Index].Replace("Engine\\Binaries", InDesc.TargetName + "\\Binaries");
-						}
+						OutputPaths = OutputPaths.Select(Path => Path.Replace("Engine\\Binaries", InDesc.TargetName + "\\Binaries")).ToList();
 					}
 				}
 			}
diff --git a/Engine/Source/Programs/UnrealBuildTool/Configuration/UEBuildPlatform.cs b/Engine/Source/Programs/UnrealBuildTool/Configuration/UEBuildPlatform.cs
--- a/Engine/Source/Programs/UnrealBuildTool/Configuration/UEBuildPlatform.cs
+++ b/Engine/Source/Programs/UnrealBuildTool/Configuration/UEBuildPlatform.cs
@@ -180,7 +180,7 @@ namespace UnrealBuildTool
 		 * For platforms that need to output multiple files per binary (ie Android "fat" binaries)
 		 * this will emit multiple paths. By default, it simply makes an array from the input
 		 */
-		string[] FinalizeBinaryPaths(string BinaryName);
+		List<string> FinalizeBinaryPaths(string BinaryName);
 
 		/**
 		 * Setup the configuration environment for building
@@ -682,10 +682,10 @@ namespace UnrealBuildTool
 		 * For platforms that need to output multiple files per binary (ie Android "fat" binaries)
 		 * this will emit multiple paths. By default, it simply makes an array from the input
 		 */
-		public virtual string[] FinalizeBinaryPaths(string BinaryName)
+		public virtual List<string> FinalizeBinaryPaths(string BinaryName)
 		{
 			List<string> TempList = new List<string>() { BinaryName };
-			return TempList.ToArray();
+			return TempList;
 		}
 
 		/**
diff --git a/Engine/Source/Programs/UnrealBuildTool/Configuration/UEBuildTarget.cs b/Engine/Source/Programs/UnrealBuildTool/Configuration/UEBuildTarget.cs
--- a/Engine/Source/Programs/UnrealBuildTool/Configuration/UEBuildTarget.cs
+++ b/Engine/Source/Programs/UnrealBuildTool/Configuration/UEBuildTarget.cs
@@ -703,16 +703,16 @@ namespace UnrealBuildTool
 		public string EngineIntermediateDirectory;
 
 		/** Output paths of final executable. */
-		public string[] OutputPaths;
+		public List<string> OutputPaths;
 
 		/** Returns the OutputPath is there is only one entry in OutputPaths */
 		public string OutputPath
 		{
 			get
 			{
-				if (OutputPaths.Length != 1)
+				if (OutputPaths.Count != 1)
 				{
-					throw new BuildException("Attempted to use UEBuildTarget.OutputPath property, but there are multiple (or no) OutputPaths. You need to handle multiple in the code that called this (size = {0})", OutputPaths.Length);
+					throw new BuildException("Attempted to use UEBuildTarget.OutputPath property, but there are multiple (or no) OutputPaths. You need to handle multiple in the code that called this (size = {0})", OutputPaths.Count);
 				}
 				return OutputPaths[0];
 			}
@@ -1755,11 +1755,19 @@ namespace UnrealBuildTool
 			// If we're building a single module, then find the binary for that module and add it to our target
 			if (OnlyModules.Count > 0)
 			{
-				AppBinaries = FilterOnlyModules();
+				AppBinaries = GetFilteredOnlyModules(AppBinaries, OnlyModules);
+				if (AppBinaries.Count == 0)
+				{
+					throw new BuildException("One or more of the modules specified using the '-module' argument could not be found.");
+				}
 			}
 			else if (UEBuildConfiguration.bHotReloadFromIDE)
 			{
-				AppBinaries = FilterGameModules();
+				AppBinaries = GetFilteredGameModules(AppBinaries);
+				if (AppBinaries.Count == 0)
+				{
+					throw new BuildException("One or more of the modules specified using the '-module' argument could not be found.");
+				}
 			}
 
 			// Filter out binaries that were already built and are just used for linking. We will not build these binaries but we need them for link information
@@ -1865,68 +1873,47 @@ namespace UnrealBuildTool
 			return FilePath.Insert(MatchPos + ModuleName.Length, Appendage);
 		}
 
-		private List<UEBuildBinary> FilterOnlyModules()
+		private static List<UEBuildBinary> GetFilteredOnlyModules(List<UEBuildBinary> Binaries, List<OnlyModule> OnlyModules)
 		{
-			var FilteredBinaries = new List<UEBuildBinary>();
+			var Result = new List<UEBuildBinary>();
 
-			var AnyBinariesAdded = false;
-			foreach (var DLLBinary in AppBinaries)
+			foreach (var DLLBinary in Binaries)
 			{
 				var FoundOnlyModule = DLLBinary.FindOnlyModule(OnlyModules);
 				if (FoundOnlyModule != null)
 				{
-					FilteredBinaries.Add(DLLBinary);
-					AnyBinariesAdded = true;
+					Result.Add(DLLBinary);
 
 					if (!String.IsNullOrEmpty(FoundOnlyModule.OnlyModuleSuffix))
 					{
-						DLLBinary.Config.OriginalOutputFilePaths = DLLBinary.Config.OutputFilePaths != null ? (string[])DLLBinary.Config.OutputFilePaths.Clone() : null;
-						for (int Index = 0; Index < DLLBinary.Config.OutputFilePaths.Length; Index++ )
-						{
-							DLLBinary.Config.OutputFilePaths[Index] = AddModuleFilenameSuffix(FoundOnlyModule.OnlyModuleName, DLLBinary.Config.OutputFilePaths[Index], FoundOnlyModule.OnlyModuleSuffix);
-						}
+						DLLBinary.Config.OriginalOutputFilePaths = DLLBinary.Config.OutputFilePaths;
+						DLLBinary.Config.OutputFilePaths         = DLLBinary.Config.OutputFilePaths.Select(Path => AddModuleFilenameSuffix(FoundOnlyModule.OnlyModuleName, Path, FoundOnlyModule.OnlyModuleSuffix)).ToList();
 					}
 				}
 			}
 
-			if (!AnyBinariesAdded)
-		{
-				throw new BuildException("One or more of the modules specified using the '-module' argument could not be found.");
-			}
-
-			// Copy the result into the final list
-			return FilteredBinaries;			
+			return Result;
 		}
 
-		private List<UEBuildBinary> FilterGameModules()
+		private static List<UEBuildBinary> GetFilteredGameModules(List<UEBuildBinary> Binaries)
 		{
-			var FilteredBinaries = new List<UEBuildBinary>();
+			var Result = new List<UEBuildBinary>();
 
-			var AnyBinariesAdded = false;
-			foreach (var DLLBinary in AppBinaries)
+			foreach (var DLLBinary in Binaries)
 			{
 				var GameModules = DLLBinary.FindGameModules();
 				if (GameModules != null && GameModules.Count > 0)
 				{
+					Result.Add(DLLBinary);
+
 					var UniqueSuffix = (new Random((int)(DateTime.Now.Ticks % Int32.MaxValue)).Next(10000)).ToString();
-					DLLBinary.Config.OriginalOutputFilePaths = DLLBinary.Config.OutputFilePaths != null ? (string[])DLLBinary.Config.OutputFilePaths.Clone() : null;
-					for (int Index = 0; Index < DLLBinary.Config.OutputFilePaths.Length; Index++)
-					{
-						DLLBinary.Config.OutputFilePaths[Index] = AddModuleFilenameSuffix(GameModules[0].Name, DLLBinary.Config.OutputFilePaths[Index], UniqueSuffix);
-					}
 
-					FilteredBinaries.Add(DLLBinary);
-					AnyBinariesAdded = true;
+					DLLBinary.Config.OriginalOutputFilePaths = DLLBinary.Config.OutputFilePaths;
+					DLLBinary.Config.OutputFilePaths         = DLLBinary.Config.OutputFilePaths.Select(Path => AddModuleFilenameSuffix(GameModules[0].Name, Path, UniqueSuffix)).ToList();
 				}
-				}
-
-			if (!AnyBinariesAdded)
-			{
-				throw new BuildException("One or more of the modules specified using the '-module' argument could not be found.");
 			}
 
-			// Copy the result into the final list
-			return FilteredBinaries;
+			return Result;
 		}
 
 		/// <summary>
@@ -1947,7 +1934,7 @@ namespace UnrealBuildTool
 				{
 					string UnusedFilename;
 					ModuleRules CheckRules = RulesCompiler.CreateModuleRules(TargetModuleName, TargetInfo, out UnusedFilename);
-					if ((CheckRules != null) && (CheckRules.Type != ModuleRules.ModuleType.External))
+					if (CheckRules.Type != ModuleRules.ModuleType.External)
 					{
 						PrivateDependencyModuleNames.Add(TargetModuleName);
 					}
@@ -2419,17 +2406,19 @@ namespace UnrealBuildTool
 
 			// Get the output filenames
 			string BaseBinaryPath = Path.Combine(OutputDirectory, MakeBinaryFileName(AppName + "-" + ModuleName, Platform, ModuleConfiguration, Rules.UndecoratedConfiguration, BinaryType));
-			string[] OutputFilePaths = UEBuildPlatform.GetBuildPlatform(Platform).FinalizeBinaryPaths(BaseBinaryPath);
+			List<string> OutputFilePaths = UEBuildPlatform.GetBuildPlatform(Platform).FinalizeBinaryPaths(BaseBinaryPath);
 
 			// Prepare the configuration object
-			UEBuildBinaryConfiguration Config = new UEBuildBinaryConfiguration(BinaryType);
-			Config.OutputFilePaths = OutputFilePaths;
-			Config.IntermediateDirectory = IntermediateDirectory;
-			Config.bHasModuleRules = (ModuleRules != null);
-			Config.bAllowExports = (BinaryType == UEBuildBinaryType.DynamicLinkLibrary);
-			Config.bAllowCompilation = bAllowCompilation;
-			Config.bIsCrossTarget = bIsCrossTarget;
-			Config.ModuleNames.Add(ModuleName);
+			UEBuildBinaryConfiguration Config = new UEBuildBinaryConfiguration(
+				InType: BinaryType,
+				InOutputFilePaths: OutputFilePaths,
+				InIntermediateDirectory: IntermediateDirectory,
+				bInHasModuleRules: ModuleRules != null,
+				bInAllowExports: BinaryType == UEBuildBinaryType.DynamicLinkLibrary,
+				bInAllowCompilation: bAllowCompilation,
+				bInIsCrossTarget: bIsCrossTarget,
+				InModuleNames: new string[]{ ModuleName }
+			);
 
 			// Create the new binary
 			UEBuildBinaryCPP Binary = new UEBuildBinaryCPP(this, Config);
@@ -2523,7 +2512,7 @@ namespace UnrealBuildTool
 		/// <param name="bIncludesGameModules">Whether this executable contains game modules</param>
 		/// <param name="ExeSubFolder">Subfolder for executables. May be null.</param>
 		/// <returns>List of executable paths for this target</returns>
-		public static string[] MakeExecutablePaths(string BaseDirectory, string BinaryName, UnrealTargetPlatform Platform, UnrealTargetConfiguration Configuration, UnrealTargetConfiguration UndecoratedConfiguration, bool bIncludesGameModules, string ExeSubFolder)
+		public static List<string> MakeExecutablePaths(string BaseDirectory, string BinaryName, UnrealTargetPlatform Platform, UnrealTargetConfiguration Configuration, UnrealTargetConfiguration UndecoratedConfiguration, bool bIncludesGameModules, string ExeSubFolder)
 		{
 			// Get the configuration for the executable. If we're building DebugGame, and this executable only contains engine modules, use the same name as development.
 			UnrealTargetConfiguration ExeConfiguration = Configuration;
@@ -2636,7 +2625,7 @@ namespace UnrealBuildTool
 
 				foreach (UEBuildBinaryConfiguration BinaryConfig in RulesBuildBinaryConfigurations)
 				{
-					BinaryConfig.OutputFilePaths = OutputPaths != null ? (string[])OutputPaths.Clone() : null;
+					BinaryConfig.OutputFilePaths = OutputPaths.ToList();
 					BinaryConfig.IntermediateDirectory = ProjectIntermediateDirectory;
 
 					if (BinaryConfig.ModuleNames.Count > 0)
diff --git a/Engine/Source/Programs/UnrealBuildTool/System/LinkEnvironment.cs b/Engine/Source/Programs/UnrealBuildTool/System/LinkEnvironment.cs
--- a/Engine/Source/Programs/UnrealBuildTool/System/LinkEnvironment.cs
+++ b/Engine/Source/Programs/UnrealBuildTool/System/LinkEnvironment.cs
@@ -2,6 +2,7 @@
 
 using System;
 using System.Collections.Generic;
+using System.Linq;
 using System.Text;
 
 namespace UnrealBuildTool
@@ -21,16 +22,16 @@ namespace UnrealBuildTool
 		public string LocalShadowDirectory = null;
 
 		/** The file path for the executable file that is output by the linker. */
-		public string[] OutputFilePaths = null;
+		public List<string> OutputFilePaths = new List<string>();
 
 		/** Returns the OutputFilePath is there is only one entry in OutputFilePaths */
 		public string OutputFilePath
 		{
 			get
 			{
-				if (OutputFilePaths.Length != 1)
+				if (OutputFilePaths.Count != 1)
 				{
-					throw new BuildException("Attempted to use LinkEnvironmentConfiguration.OutputFilePath property, but there are multiple (or no) OutputFilePaths. You need to handle multiple in the code that called this (size = {0})", OutputFilePaths.Length);
+					throw new BuildException("Attempted to use LinkEnvironmentConfiguration.OutputFilePath property, but there are multiple (or no) OutputFilePaths. You need to handle multiple in the code that called this (size = {0})", OutputFilePaths.Count);
 				}
 				return OutputFilePaths[0];
 			}
@@ -118,7 +119,7 @@ namespace UnrealBuildTool
 			OutputDirectory = InCopyEnvironment.OutputDirectory;
 			IntermediateDirectory = InCopyEnvironment.IntermediateDirectory;
 			LocalShadowDirectory = InCopyEnvironment.LocalShadowDirectory;
-			OutputFilePaths = InCopyEnvironment.OutputFilePaths != null ? (string[])InCopyEnvironment.OutputFilePaths.Clone() : null;
+			OutputFilePaths = InCopyEnvironment.OutputFilePaths.ToList();
 			LibraryPaths.AddRange(InCopyEnvironment.LibraryPaths);
 			ExcludedLibraries.AddRange(InCopyEnvironment.ExcludedLibraries);
 			AdditionalLibraries.AddRange(InCopyEnvironment.AdditionalLibraries);
diff --git a/Engine/Source/Programs/UnrealBuildTool/System/UnrealBuildTool.cs b/Engine/Source/Programs/UnrealBuildTool/System/UnrealBuildTool.cs
--- a/Engine/Source/Programs/UnrealBuildTool/System/UnrealBuildTool.cs
+++ b/Engine/Source/Programs/UnrealBuildTool/System/UnrealBuildTool.cs
@@ -2087,7 +2087,7 @@ namespace UnrealBuildTool
 			if (!ProjectFileGenerator.bGenerateProjectFiles && !UEBuildConfiguration.bGenerateManifest && bIsEditorTarget )
 			{
 				var EditorProcessFilenames = UEBuildTarget.MakeExecutablePaths("..", "UE4Editor", TargetDesc.Platform, TargetDesc.Configuration, UnrealTargetConfiguration.Development, false, null);
-				if (EditorProcessFilenames.Length != 1)
+				if (EditorProcessFilenames.Count != 1)
 				{
 					throw new BuildException("ShouldDoHotReload cannot handle multiple binaries returning from UEBuildTarget.MakeExecutablePaths");
 				}
@@ -2349,8 +2349,6 @@ namespace UnrealBuildTool
 				ReasonNotLoaded = "no existing makefile";
 			}
 
-
-
             UBTMakefile LoadedUBTMakefile = null;
             if( !bForceOutOfDate )
             { 
