More robust filtering of platform-specific plugins; only removes platform folders which are above the engine or project root.

diff --git a/Engine/Source/Programs/UnrealBuildTool/Configuration/UEBuildTarget.cs b/Engine/Source/Programs/UnrealBuildTool/Configuration/UEBuildTarget.cs
--- a/Engine/Source/Programs/UnrealBuildTool/Configuration/UEBuildTarget.cs
+++ b/Engine/Source/Programs/UnrealBuildTool/Configuration/UEBuildTarget.cs
@@ -2549,14 +2549,16 @@ namespace UnrealBuildTool
 			ValidPlugins = Plugins.ReadAvailablePlugins(UnrealBuildTool.GetUProjectFile());
 
 			// Remove any plugins for platforms we don't have
+			List<string> ExcludeFolders = new List<string>();
 			foreach (UnrealTargetPlatform TargetPlatform in Enum.GetValues(typeof(UnrealTargetPlatform)))
 			{
-				if (TargetPlatform != UnrealTargetPlatform.Desktop && UEBuildPlatform.GetBuildPlatform(TargetPlatform, true) == null)
+				if (UEBuildPlatform.GetBuildPlatform(TargetPlatform, true) == null)
 				{
 					string DirectoryFragment = String.Format("/{0}/", TargetPlatform.ToString());
-					ValidPlugins.RemoveAll(x => x.Directory.Replace('\\', '/').Contains(DirectoryFragment));
+					ExcludeFolders.Add(DirectoryFragment);
 				}
 			}
+			ValidPlugins.RemoveAll(x => ShouldExcludePlugin(x, ExcludeFolders));
 
 			// Build a list of enabled plugins
 			EnabledPlugins = new List<PluginInfo>();
@@ -2610,6 +2612,21 @@ namespace UnrealBuildTool
 			}
 		}
 
+		/** Checks whether a plugin path contains a platform directory fragment */
+		private static bool ShouldExcludePlugin(PluginInfo Plugin, List<string> ExcludeFragments)
+		{
+			string RelativePathFromRoot;
+			if(Plugin.LoadedFrom == PluginLoadedFrom.Engine)
+			{
+				RelativePathFromRoot = Utils.CleanDirectorySeparators(Utils.MakePathRelativeTo(Plugin.FileName, BuildConfiguration.RelativeEnginePath), '/');
+			}
+			else
+			{
+				RelativePathFromRoot = Utils.CleanDirectorySeparators(Utils.MakePathRelativeTo(Plugin.FileName, UnrealBuildTool.GetUProjectPath()), '/');
+			}
+			return ExcludeFragments.Any(x => RelativePathFromRoot.Contains(x));
+		}
+
 		/** Sets up the binaries for the target. */
 		protected virtual void SetupBinaries()
 		{
