Adding code to catch too long filenames when building with UBT

diff --git a/Engine/Source/Programs/UnrealBuildTool/System/CPPHeaders.cs b/Engine/Source/Programs/UnrealBuildTool/System/CPPHeaders.cs
--- a/Engine/Source/Programs/UnrealBuildTool/System/CPPHeaders.cs
+++ b/Engine/Source/Programs/UnrealBuildTool/System/CPPHeaders.cs
@@ -126,7 +126,15 @@ namespace UnrealBuildTool
 						{
 							throw new BuildException(Exception, "Failed to combine null or invalid include paths.");
 						}
-						string FullFilePath = Path.GetFullPath( RelativeFilePath );
+						string FullFilePath = null;
+						try
+						{
+							FullFilePath = Path.GetFullPath(RelativeFilePath);
+						}
+						catch (Exception Exception)
+						{
+							throw new BuildException(Exception, "Failed to get full path for include: \"{0}\"", RelativeFilePath);
+						}
 						if( DirectoryLookupCache.FileExists( FullFilePath ) )
 						{
 							Result = FileItem.GetItemByFullPath( FullFilePath );
