[INTEGRATE] Change 2557067 by [REDACTED_PERSON][REDACTED_HANDLE] on 2015/05/19 11:10:14

	Fix garbage dylibs being added to build manifest when compiling a single module.

diff --git a/Engine/Source/Programs/UnrealBuildTool/Mac/MacToolChain.cs b/Engine/Source/Programs/UnrealBuildTool/Mac/MacToolChain.cs
--- a/Engine/Source/Programs/UnrealBuildTool/Mac/MacToolChain.cs
+++ b/Engine/Source/Programs/UnrealBuildTool/Mac/MacToolChain.cs
@@ -1363,7 +1363,7 @@ namespace UnrealBuildTool
 				string LibName = Path.GetFileName(AdditionalLibrary);
 				if (LibName.StartsWith("lib"))
 				{
-					if (Path.GetExtension(AdditionalLibrary) == ".dylib")
+					if (Path.GetExtension(AdditionalLibrary) == ".dylib" && !String.IsNullOrEmpty(BundleContentsDirectory))
 					{
 						string Entry = BundleContentsDirectory + "MacOS/" + LibName;
 						Receipt.AddBuildProduct(Entry, BuildProductType.DynamicLibrary);
