Makefile is invalidated by .uproject file changes.
Change to makefile timestamp checks to use FileInfo instead of FileItem to avoid some time offset issues.
Misc cleanups in some module reload code.

diff --git a/Engine/Source/Developer/HotReload/Private/HotReload.cpp b/Engine/Source/Developer/HotReload/Private/HotReload.cpp
--- a/Engine/Source/Developer/HotReload/Private/HotReload.cpp
+++ b/Engine/Source/Developer/HotReload/Private/HotReload.cpp
@@ -1117,28 +1117,16 @@ bool FHotReloadModule::RecompileModulesAsync( const TArray< FName > ModuleNames,
 
 	TArray< FModuleToRecompile > ModulesToRecompile;
 
-	for( TArray< FName >::TConstIterator CurModuleIt( ModuleNames ); CurModuleIt; ++CurModuleIt )
+	for( FName CurModuleName : ModuleNames )
 	{
-		const FName CurModuleName = *CurModuleIt;
-
 		// Update our set of known modules, in case we don't already know about this module
 		FModuleManager::Get().AddModule( CurModuleName );
 
-		FString NewModuleFileNameOnSuccess = FModuleManager::Get().GetModuleFilename(CurModuleName);
-
 		// Find a unique file name for the module
-		FString UniqueSuffix;
-		FString UniqueModuleFileName;
-		FModuleManager::Get().MakeUniqueModuleFilename( CurModuleName, UniqueSuffix, UniqueModuleFileName );
-
-		// If the recompile succeeds, we'll update our cached file name to use the new unique file name
-		// that we setup for the module
-		NewModuleFileNameOnSuccess = UniqueModuleFileName;
-
 		FModuleToRecompile ModuleToRecompile;
 		ModuleToRecompile.ModuleName = CurModuleName.ToString();
-		ModuleToRecompile.ModuleFileSuffix = UniqueSuffix;
-		ModuleToRecompile.NewModuleFilename = UniqueModuleFileName;
+		FModuleManager::Get().MakeUniqueModuleFilename( CurModuleName, ModuleToRecompile.ModuleFileSuffix, ModuleToRecompile.NewModuleFilename );
+
 		ModulesToRecompile.Add( ModuleToRecompile );
 	}
 
@@ -1253,20 +1241,20 @@ bool FHotReloadModule::StartCompilingModuleDLLs(const FString& GameName, const T
 
 	// Pass a module file suffix to UBT if we have one
 	FString ModuleArg;
-	for( int32 CurModuleIndex = 0; CurModuleIndex < ModuleNames.Num(); ++CurModuleIndex )
+	for( const FModuleToRecompile& Module : ModuleNames )
 	{
-		if( !ModuleNames[ CurModuleIndex ].ModuleFileSuffix.IsEmpty() )
+		if( !Module.ModuleFileSuffix.IsEmpty() )
 		{
-			ModuleArg += FString::Printf( TEXT( " -ModuleWithSuffix %s %s" ), *ModuleNames[ CurModuleIndex ].ModuleName, *ModuleNames[ CurModuleIndex ].ModuleFileSuffix );
+			ModuleArg += FString::Printf( TEXT( " -ModuleWithSuffix %s %s" ), *Module.ModuleName, *Module.ModuleFileSuffix );
 		}
 		else
 		{
-			ModuleArg += FString::Printf( TEXT( " -Module %s" ), *ModuleNames[ CurModuleIndex ].ModuleName );
+			ModuleArg += FString::Printf( TEXT( " -Module %s" ), *Module.ModuleName );
 		}
-		Ar.Logf( TEXT( "Recompiling %s..." ), *ModuleNames[ CurModuleIndex ].ModuleName );
+		Ar.Logf( TEXT( "Recompiling %s..." ), *Module.ModuleName );
 
 		// prepare the compile info in the FModuleInfo so that it can be compared after compiling
-		FName ModuleFName(*ModuleNames[ CurModuleIndex ].ModuleName);
+		FName ModuleFName(*Module.ModuleName);
 		UpdateModuleCompileData(ModuleFName);
 	}
 
diff --git a/Engine/Source/Programs/UnrealBuildTool/System/UnrealBuildTool.cs b/Engine/Source/Programs/UnrealBuildTool/System/UnrealBuildTool.cs
--- a/Engine/Source/Programs/UnrealBuildTool/System/UnrealBuildTool.cs
+++ b/Engine/Source/Programs/UnrealBuildTool/System/UnrealBuildTool.cs
@@ -2273,13 +2273,13 @@ namespace UnrealBuildTool
 		/// <returns>The loaded makefile, or null if it failed for some reason.  On failure, the 'ReasonNotLoaded' variable will contain information about why</returns>
         static UBTMakefile LoadUBTMakefile( List<TargetDescriptor> TargetDescs, out string ReasonNotLoaded )
         {
-            var UBTMakefileItem = FileItem.GetItemByFullPath( GetUBTMakefilePath( TargetDescs ) );
+            var UBTMakefileInfo = new FileInfo( GetUBTMakefilePath( TargetDescs ) );
 			ReasonNotLoaded = null;
 
             // Check the directory timestamp on the project files directory.  If the user has generated project files more
             // recently than the UBTMakefile, then we need to consider the file to be out of date
             bool bForceOutOfDate = false;
-            if( UBTMakefileItem.bExists )
+            if( UBTMakefileInfo.Exists )
             {
                 // @todo ubtmake: This will only work if the directory timestamp actually changes with every single GPF.  Force delete existing files before creating new ones?  Eh... really we probably just want to delete + create a file in that folder
                 //			-> UPDATE: Seems to work OK right now though on Windows platform, maybe due to GUID changes
@@ -2290,7 +2290,7 @@ namespace UnrealBuildTool
                     if( Directory.Exists( ProjectFileGenerator.IntermediateProjectFilesPath ) )
                     {
                         var EngineProjectFilesLastUpdateTime = new FileInfo(ProjectFileGenerator.ProjectTimestampFile).LastWriteTime;
-						if( UBTMakefileItem.LastWriteTime.CompareTo( EngineProjectFilesLastUpdateTime ) < 0 )
+						if( UBTMakefileInfo.LastWriteTime.CompareTo( EngineProjectFilesLastUpdateTime ) < 0 )
 						{
 							// Engine project files are newer than UBTMakefile
 							bForceOutOfDate = true;
@@ -2309,21 +2309,34 @@ namespace UnrealBuildTool
 				{
 					// Check the game project directory too
 					if( UnrealBuildTool.HasUProjectFile() )
-					{ 
-						var MasterProjectRelativePath = UnrealBuildTool.GetUProjectPath();
-						var GameIntermediateProjectFilesPath = Path.Combine( MasterProjectRelativePath, "Intermediate", "ProjectFiles" );
-						if( Directory.Exists( GameIntermediateProjectFilesPath ) )
+					{
+						string ProjectFilename = UnrealBuildTool.GetUProjectFile();
+						var ProjectFileInfo = new FileInfo(ProjectFilename);
+						if( !ProjectFileInfo.Exists || UBTMakefileInfo.LastWriteTime.CompareTo(ProjectFileInfo.LastWriteTime ) < 0 )
+						{
+							// .uproject file is newer than UBTMakefile
+							bForceOutOfDate = true;
+							Log.TraceVerbose("Makefile is older than .uproject file, ignoring it" );
+
+							ReasonNotLoaded = ".uproject file is newer";
+						}
+						else
 						{
-							var GameProjectFilesLastUpdateTime = new DirectoryInfo( GameIntermediateProjectFilesPath ).LastWriteTime;
-							if( UBTMakefileItem.LastWriteTime.CompareTo( GameProjectFilesLastUpdateTime ) < 0 )
+							var MasterProjectRelativePath = UnrealBuildTool.GetUProjectPath();
+							var GameIntermediateProjectFilesPath = Path.Combine( MasterProjectRelativePath, "Intermediate", "ProjectFiles" );
+							if( Directory.Exists( GameIntermediateProjectFilesPath ) )
 							{
-								// Game project files are newer than UBTMakefile
-								bForceOutOfDate = true;
-								Log.TraceVerbose("Makefile is older than generated game project files, ignoring it" );
+								var GameProjectFilesLastUpdateTime = new DirectoryInfo( GameIntermediateProjectFilesPath ).LastWriteTime;
+								if( UBTMakefileInfo.LastWriteTime.CompareTo( GameProjectFilesLastUpdateTime ) < 0 )
+								{
+									// Game project files are newer than UBTMakefile
+									bForceOutOfDate = true;
+									Log.TraceVerbose("Makefile is older than generated game project files, ignoring it" );
 
-								ReasonNotLoaded = "game project files are newer";
+									ReasonNotLoaded = "game project files are newer";
+								}
 							}
-						}	
+						}
 					}
 				}
 
@@ -2331,7 +2344,7 @@ namespace UnrealBuildTool
 				{
 					// Check to see if UnrealBuildTool.exe was compiled more recently than the UBTMakefile
 					var UnrealBuildToolTimestamp = new FileInfo( Assembly.GetExecutingAssembly().Location ).LastWriteTime;
-					if( UBTMakefileItem.LastWriteTime.CompareTo( UnrealBuildToolTimestamp ) < 0 )
+					if( UBTMakefileInfo.LastWriteTime.CompareTo( UnrealBuildToolTimestamp ) < 0 )
 					{
 						// UnrealBuildTool.exe was compiled more recently than the UBTMakefile
 						Log.TraceVerbose("Makefile is older than UnrealBuildTool.exe, ignoring it" );
@@ -2356,7 +2369,7 @@ namespace UnrealBuildTool
                 {
 					var LoadUBTMakefileStartTime = DateTime.UtcNow;
 
-                    using (FileStream Stream = new FileStream(UBTMakefileItem.AbsolutePath, FileMode.Open, FileAccess.Read))
+                    using (FileStream Stream = new FileStream(UBTMakefileInfo.FullName, FileMode.Open, FileAccess.Read))
                     {	
                         BinaryFormatter Formatter = new BinaryFormatter();
                         LoadedUBTMakefile = Formatter.Deserialize(Stream) as UBTMakefile;
@@ -2376,7 +2389,7 @@ namespace UnrealBuildTool
 
                 if( LoadedUBTMakefile != null && !LoadedUBTMakefile.IsValidMakefile() )
                 {
-                    Log.TraceWarning("Loaded makefile appears to have invalid contents, ignoring it ({0})", UBTMakefileItem.AbsolutePath );
+                    Log.TraceWarning("Loaded makefile appears to have invalid contents, ignoring it ({0})", UBTMakefileInfo.FullName );
                     LoadedUBTMakefile = null;
 					ReasonNotLoaded = "existing makefile appears to be invalid";
                 }
@@ -2391,9 +2404,9 @@ namespace UnrealBuildTool
 						{
 							var TargetCsFile = new FileInfo(TargetCsFilename);
 							bool bTargetCsFileExists = TargetCsFile.Exists;
-							if (!bTargetCsFileExists || TargetCsFile.LastWriteTime > UBTMakefileItem.LastWriteTime)
+							if (!bTargetCsFileExists || TargetCsFile.LastWriteTime > UBTMakefileInfo.LastWriteTime)
 							{
-								Log.TraceWarning("{0} has been {1} since makefile was built, ignoring it ({2})", TargetCsFilename, bTargetCsFileExists ? "changed" : "deleted", UBTMakefileItem.AbsolutePath);
+								Log.TraceWarning("{0} has been {1} since makefile was built, ignoring it ({2})", TargetCsFilename, bTargetCsFileExists ? "changed" : "deleted", UBTMakefileInfo.FullName);
 								LoadedUBTMakefile = null;
 								ReasonNotLoaded = string.Format("changes to target files");
 								goto SkipRemainingTimestampChecks;
@@ -2407,9 +2420,9 @@ namespace UnrealBuildTool
 							{
 								var BuildCsFile = new FileInfo(BuildCsFilename);
 								bool bBuildCsFileExists = BuildCsFile.Exists;
-								if (!bBuildCsFileExists || BuildCsFile.LastWriteTime > UBTMakefileItem.LastWriteTime)
+								if (!bBuildCsFileExists || BuildCsFile.LastWriteTime > UBTMakefileInfo.LastWriteTime)
 								{
-									Log.TraceWarning("{0} has been {1} since makefile was built, ignoring it ({2})", BuildCsFilename, bBuildCsFileExists ? "changed" : "deleted", UBTMakefileItem.AbsolutePath);
+									Log.TraceWarning("{0} has been {1} since makefile was built, ignoring it ({2})", BuildCsFilename, bBuildCsFileExists ? "changed" : "deleted", UBTMakefileInfo.FullName);
 									LoadedUBTMakefile = null;
 									ReasonNotLoaded = string.Format("changes to module files");
 									goto SkipRemainingTimestampChecks;
diff --git a/Engine/Source/Runtime/Core/Private/Modules/ModuleManager.cpp b/Engine/Source/Runtime/Core/Private/Modules/ModuleManager.cpp
--- a/Engine/Source/Runtime/Core/Private/Modules/ModuleManager.cpp
+++ b/Engine/Source/Runtime/Core/Private/Modules/ModuleManager.cpp
@@ -5,7 +5,7 @@
 #include "ModuleVersion.h"
 #include "EngineBuildSettings.h"
 #include "UProjectInfo.h"
-
+#include "ScopeExit.h"
 
 DEFINE_LOG_CATEGORY_STATIC(LogModuleManager, Log, All);
 
@@ -203,21 +203,10 @@ void FModuleManager::AddModule(const FName InModuleName)
 	TSharedRef<FModuleInfo> ModuleInfo(new FModuleInfo());
 	
 	// Make sure module info is added to known modules and proper delegates are fired on exit.
-	struct FAtExit
+	ON_SCOPE_EXIT
 	{
-		FAtExit(const FName InModule_Name, const TSharedRef<FModuleInfo>& InModuleInfo)
-			: ModuleName(InModule_Name)
-			, ModuleInfo(InModuleInfo)
-		{ }
-
-		~FAtExit()
-		{
-			FModuleManager::Get().AddModuleToModulesList(ModuleName, ModuleInfo);
-		}
-
-		FName ModuleName;
-		TSharedRef<FModuleInfo> ModuleInfo;
-	} AtExit(InModuleName, ModuleInfo);
+		FModuleManager::Get().AddModuleToModulesList(InModuleName, ModuleInfo);
+	};
 
 #if !IS_MONOLITHIC
 	FString ModuleNameString = InModuleName.ToString();
