UE-15926 : missing symbols file.

  make sure the the generated symbols file is checked in.

diff --git a/Engine/Source/Programs/AutomationTool/HTML5/HTML5Platform.Automation.cs b/Engine/Source/Programs/AutomationTool/HTML5/HTML5Platform.Automation.cs
--- a/Engine/Source/Programs/AutomationTool/HTML5/HTML5Platform.Automation.cs
+++ b/Engine/Source/Programs/AutomationTool/HTML5/HTML5Platform.Automation.cs
@@ -98,6 +98,8 @@ public class HTML5Platform : Platform
         }
         File.SetAttributes(Path.Combine(PackagePath, GameExe), FileAttributes.Normal);
         File.SetAttributes(Path.Combine(PackagePath, GameExe) + ".mem", FileAttributes.Normal);
+		File.SetAttributes(Path.Combine(PackagePath, GameExe) + ".symbols", FileAttributes.Normal);
+
 
         // put the HTML file to the package directory
 
diff --git a/Engine/Source/Programs/UnrealBuildTool/HTML5/HTML5ToolChain.cs b/Engine/Source/Programs/UnrealBuildTool/HTML5/HTML5ToolChain.cs
--- a/Engine/Source/Programs/UnrealBuildTool/HTML5/HTML5ToolChain.cs
+++ b/Engine/Source/Programs/UnrealBuildTool/HTML5/HTML5ToolChain.cs
@@ -531,10 +531,11 @@ namespace UnrealBuildTool
 
         public override void AddFilesToReceipt(BuildReceipt Receipt, UEBuildBinary Binary)
         {
-            // we need to include the generated .mem file.  
+            // we need to include the generated .mem and .symbols file.  
 			if(Binary.Config.Type != UEBuildBinaryType.StaticLibrary)
 			{
 	            Receipt.AddBuildProduct(Binary.Config.OutputFilePath + ".mem", BuildProductType.RequiredResource);
+				Receipt.AddBuildProduct(Binary.Config.OutputFilePath + ".symbols", BuildProductType.RequiredResource);
 			}
         }
 
