Turn shadow variable warnings back into warnings; they are breaking too many builds.

diff --git a/Engine/Source/Programs/UnrealBuildTool/HTML5/HTML5ToolChain.cs b/Engine/Source/Programs/UnrealBuildTool/HTML5/HTML5ToolChain.cs
--- a/Engine/Source/Programs/UnrealBuildTool/HTML5/HTML5ToolChain.cs
+++ b/Engine/Source/Programs/UnrealBuildTool/HTML5/HTML5ToolChain.cs
@@ -66,7 +66,7 @@ namespace UnrealBuildTool
 
 			if (bEnableShadowVariableWarning)
 			{
-				Result += " -Wshadow";
+				Result += " -Wshadow -Wno-error=shadow";
 			}
 
             // JavsScript option overrides (see src/settings.js)
diff --git a/Engine/Source/Programs/UnrealBuildTool/IOS/IOSToolChain.cs b/Engine/Source/Programs/UnrealBuildTool/IOS/IOSToolChain.cs
--- a/Engine/Source/Programs/UnrealBuildTool/IOS/IOSToolChain.cs
+++ b/Engine/Source/Programs/UnrealBuildTool/IOS/IOSToolChain.cs
@@ -232,7 +232,7 @@ namespace UnrealBuildTool
 				
 			if (CompileEnvironment.Config.bEnableShadowVariableWarning)
 			{
-				Result += " -Wshadow";
+				Result += " -Wshadow -Wno-error=shadow";
 			}
 
 			Result += " -Wno-unused-variable";
diff --git a/Engine/Source/Programs/UnrealBuildTool/Linux/LinuxToolChain.cs b/Engine/Source/Programs/UnrealBuildTool/Linux/LinuxToolChain.cs
--- a/Engine/Source/Programs/UnrealBuildTool/Linux/LinuxToolChain.cs
+++ b/Engine/Source/Programs/UnrealBuildTool/Linux/LinuxToolChain.cs
@@ -353,7 +353,7 @@ namespace UnrealBuildTool
 
 			if (CompileEnvironment.Config.bEnableShadowVariableWarning)
 			{
-				Result += " -Wshadow";
+				Result += " -Wshadow -Wno-error=shadow";
 			}
 
             //Result += " -DOPERATOR_NEW_INLINE=FORCENOINLINE";
diff --git a/Engine/Source/Programs/UnrealBuildTool/Mac/MacToolChain.cs b/Engine/Source/Programs/UnrealBuildTool/Mac/MacToolChain.cs
--- a/Engine/Source/Programs/UnrealBuildTool/Mac/MacToolChain.cs
+++ b/Engine/Source/Programs/UnrealBuildTool/Mac/MacToolChain.cs
@@ -172,7 +172,7 @@ namespace UnrealBuildTool
 
 			if (CompileEnvironment.Config.bEnableShadowVariableWarning)
 			{
-				Result += " -Wshadow";
+				Result += " -Wshadow -Wno-error=shadow";
 			}
 
 			// @todo: Remove these two when the code is fixed and they're no longer needed
diff --git a/Engine/Source/Programs/UnrealBuildTool/Windows/VCToolChain.cs b/Engine/Source/Programs/UnrealBuildTool/Windows/VCToolChain.cs
--- a/Engine/Source/Programs/UnrealBuildTool/Windows/VCToolChain.cs
+++ b/Engine/Source/Programs/UnrealBuildTool/Windows/VCToolChain.cs
@@ -677,7 +677,7 @@ namespace UnrealBuildTool
 
 				if (CompileEnvironment.Config.bEnableShadowVariableWarning)
 				{
-					Arguments.Append(" -Wshadow");
+					Arguments.Append(" -Wshadow -Wno-error=shadow");
 				}
 
 				// @todo clang: Kind of a shame to turn these off.  We'd like to catch unused variables, but it is tricky with how our assertion macros work.
