Fix dependencies of precompiled modules being linked into the game binary. Was causing CEF to be linked in to UE4Game in Rocket, causing launch-on to fail.

* Plugin modules are now explicitly added into the modules list for the binaries they're going to be a part of (this wasn't possible before we supported modules being in static libraries AND a DLL/EXE at the same time)
* Receipt now only includes runtime dependencies for things which are actually runnable

diff --git a/Engine/Source/Programs/UnrealBuildTool/Configuration/UEBuildTarget.cs b/Engine/Source/Programs/UnrealBuildTool/Configuration/UEBuildTarget.cs
--- a/Engine/Source/Programs/UnrealBuildTool/Configuration/UEBuildTarget.cs
+++ b/Engine/Source/Programs/UnrealBuildTool/Configuration/UEBuildTarget.cs
@@ -1408,6 +1408,10 @@ namespace UnrealBuildTool
 			foreach(UEBuildBinary Binary in AppBinaries)
 			{
 				BuildReceipt BinaryReceipt = Binary.MakeReceipt(ToolChain);
+				if(Binary.Config.Type == UEBuildBinaryType.StaticLibrary)
+				{
+					BinaryReceipt.RuntimeDependencies.Clear();
+				}
 				Receipt.Merge(BinaryReceipt);
 			}
 
@@ -1452,17 +1456,13 @@ namespace UnrealBuildTool
 			{
 				var ExecutableBinary = AppBinaries[0];
 
-				// Search through every binary for dependencies. When building plugin binaries that reference optional engine modules, 
-				// we still need to link them into the executable.
-				foreach (UEBuildBinary Binary in AppBinaries)
+				// Add all the modules that the executable depends on. Plugins will be already included in this list.
+				var AllReferencedModules = ExecutableBinary.GetAllDependencyModules(bIncludeDynamicallyLoaded: true, bForceCircular: true);
+				foreach (var CurModule in AllReferencedModules)
 				{
-					var AllReferencedModules = Binary.GetAllDependencyModules(bIncludeDynamicallyLoaded: true, bForceCircular: true);
-					foreach (var CurModule in AllReferencedModules)
+					if (CurModule.Binary == null || CurModule.Binary == ExecutableBinary || CurModule.Binary.Config.Type == UEBuildBinaryType.StaticLibrary)
 					{
-						if (CurModule.Binary == null || CurModule.Binary == ExecutableBinary || CurModule.Binary.Config.Type == UEBuildBinaryType.StaticLibrary)
-						{
-							ExecutableBinary.AddModule(CurModule.Name);
-						}
+						ExecutableBinary.AddModule(CurModule.Name);
 					}
 				}
 
@@ -1940,16 +1940,6 @@ namespace UnrealBuildTool
 					}
 				}
 			}
-			foreach (PluginInfo Plugin in EnabledPlugins)
-			{
-				foreach (ModuleDescriptor Module in Plugin.Descriptor.Modules)
-				{
-					if(Module.IsCompiledInConfiguration(Platform, TargetType))
-					{
-						PrivateDependencyModuleNames.Add(Module.Name);
-					}
-				}
-			}
 
 			// We ALWAYS have to write this file as the IMPLEMENT_PRIMARY_GAME_MODULE function depends on the UELinkerFixups function existing.
 			{				
@@ -2237,6 +2227,12 @@ namespace UnrealBuildTool
 						string ModuleFileName = RulesCompiler.GetModuleFilename(Module.Name);
 						bool bHasSource = (!String.IsNullOrEmpty(ModuleFileName) && Directory.EnumerateFiles(Path.GetDirectoryName(ModuleFileName), "*.cpp", SearchOption.AllDirectories).Any());
 						AddBinaryForModule(Module.Name, BinaryType, bAllowCompilation: bHasSource, bIsCrossTarget: false);
+
+						// Add it to the binary if we're compiling monolithic (and it's enabled)
+						if(ShouldCompileMonolithic() && EnabledPlugins.Contains(Plugin))
+						{
+							AppBinaries[0].AddModule(Module.Name);
+						}
 					}
 				}
 			}
