Fixed custom icon paths on Mac

diff --git a/Engine/Source/Programs/UnrealBuildTool/Mac/MacToolChain.cs b/Engine/Source/Programs/UnrealBuildTool/Mac/MacToolChain.cs
--- a/Engine/Source/Programs/UnrealBuildTool/Mac/MacToolChain.cs
+++ b/Engine/Source/Programs/UnrealBuildTool/Mac/MacToolChain.cs
@@ -1027,8 +1027,9 @@ namespace UnrealBuildTool
 					}
 					else
 					{
-						CustomResourcesPath = Path.GetDirectoryName(UProjectFilePath) + "/Source/" + GameName + "/Resources/Mac";
-						CustomBuildPath = Path.GetDirectoryName(UProjectFilePath) + "/Build/Mac";
+						string FullUProjectFilePath = Path.GetFullPath(UProjectFilePath);
+						CustomResourcesPath = Path.GetDirectoryName(FullUProjectFilePath) + "/Source/" + GameName + "/Resources/Mac";
+						CustomBuildPath = Path.GetDirectoryName(FullUProjectFilePath) + "/Build/Mac";
 					}
 
 					bool bBuildingEditor = GameName.EndsWith("Editor");
