Fix linking against static libraries which have circular references back to the original binary. Plugins depend on other engine modules, but linking a monolithic game should not add a library dependency on the monolithic executable.

diff --git a/Engine/Source/Programs/UnrealBuildTool/Configuration/UEBuildBinary.cs b/Engine/Source/Programs/UnrealBuildTool/Configuration/UEBuildBinary.cs
--- a/Engine/Source/Programs/UnrealBuildTool/Configuration/UEBuildBinary.cs
+++ b/Engine/Source/Programs/UnrealBuildTool/Configuration/UEBuildBinary.cs
@@ -730,7 +730,7 @@ namespace UnrealBuildTool
 				if (!BuildConfiguration.bRunUnrealCodeAnalyzer)
 				{
 					// Allow the module to modify the link environment for the binary.
-					Module.SetupPrivateLinkEnvironment(BinaryLinkEnvironment, BinaryDependencies, LinkEnvironmentVisitedModules);
+					Module.SetupPrivateLinkEnvironment(this, BinaryLinkEnvironment, BinaryDependencies, LinkEnvironmentVisitedModules);
 				}
 			}
 
diff --git a/Engine/Source/Programs/UnrealBuildTool/Configuration/UEBuildModule.cs b/Engine/Source/Programs/UnrealBuildTool/Configuration/UEBuildModule.cs
--- a/Engine/Source/Programs/UnrealBuildTool/Configuration/UEBuildModule.cs
+++ b/Engine/Source/Programs/UnrealBuildTool/Configuration/UEBuildModule.cs
@@ -810,13 +810,14 @@ namespace UnrealBuildTool
 
 		/** Sets up the environment for linking this module. */
 		public virtual void SetupPrivateLinkEnvironment(
+			UEBuildBinary SourceBinary,
 			LinkEnvironment LinkEnvironment,
 			List<UEBuildBinary> BinaryDependencies,
 			Dictionary<UEBuildModule, bool> VisitedModules
 			)
 		{
 			// Allow the module's public dependencies to add library paths and additional libraries to the link environment.
-			SetupPublicLinkEnvironment(Binary, LinkEnvironment.Config.LibraryPaths, LinkEnvironment.Config.AdditionalLibraries, LinkEnvironment.Config.Frameworks, LinkEnvironment.Config.WeakFrameworks,
+			SetupPublicLinkEnvironment(SourceBinary, LinkEnvironment.Config.LibraryPaths, LinkEnvironment.Config.AdditionalLibraries, LinkEnvironment.Config.Frameworks, LinkEnvironment.Config.WeakFrameworks,
 				LinkEnvironment.Config.AdditionalFrameworks,LinkEnvironment.Config.AdditionalShadowFiles, LinkEnvironment.Config.AdditionalBundleResources, LinkEnvironment.Config.DelayLoadDLLs, BinaryDependencies, VisitedModules);
 
 			// Also allow the module's public and private dependencies to modify the link environment.
@@ -826,7 +827,7 @@ namespace UnrealBuildTool
 			foreach (var DependencyName in AllDependencyModuleNames)
 			{
 				var DependencyModule = Target.GetModuleByName(DependencyName);
-				DependencyModule.SetupPublicLinkEnvironment(Binary, LinkEnvironment.Config.LibraryPaths, LinkEnvironment.Config.AdditionalLibraries, LinkEnvironment.Config.Frameworks, LinkEnvironment.Config.WeakFrameworks,
+				DependencyModule.SetupPublicLinkEnvironment(SourceBinary, LinkEnvironment.Config.LibraryPaths, LinkEnvironment.Config.AdditionalLibraries, LinkEnvironment.Config.Frameworks, LinkEnvironment.Config.WeakFrameworks,
 					LinkEnvironment.Config.AdditionalFrameworks, LinkEnvironment.Config.AdditionalShadowFiles, LinkEnvironment.Config.AdditionalBundleResources, LinkEnvironment.Config.DelayLoadDLLs, BinaryDependencies, VisitedModules);
 			}
 		}
@@ -2203,12 +2204,13 @@ namespace UnrealBuildTool
 		}
 
 		public override void SetupPrivateLinkEnvironment(
+			UEBuildBinary SourceBinary,
 			LinkEnvironment LinkEnvironment,
 			List<UEBuildBinary> BinaryDependencies,
 			Dictionary<UEBuildModule, bool> VisitedModules
 			)
 		{
-			base.SetupPrivateLinkEnvironment(LinkEnvironment, BinaryDependencies, VisitedModules);
+			base.SetupPrivateLinkEnvironment(SourceBinary, LinkEnvironment, BinaryDependencies, VisitedModules);
 
 			// Setup the link environment for linking a CLR binary.
 			LinkEnvironment.Config.CLRMode = CPPCLRMode.CLREnabled;
diff --git a/Engine/Source/Programs/UnrealBuildTool/Mac/MacToolChain.cs b/Engine/Source/Programs/UnrealBuildTool/Mac/MacToolChain.cs
--- a/Engine/Source/Programs/UnrealBuildTool/Mac/MacToolChain.cs
+++ b/Engine/Source/Programs/UnrealBuildTool/Mac/MacToolChain.cs
@@ -1356,7 +1356,7 @@ namespace UnrealBuildTool
 			var LinkEnvironmentVisitedModules = new Dictionary<UEBuildModule, bool>();
 			foreach (var Module in Modules)
 			{
-				Module.SetupPrivateLinkEnvironment(BinaryLinkEnvironment, BinaryDependencies, LinkEnvironmentVisitedModules);
+				Module.SetupPrivateLinkEnvironment(Binary, BinaryLinkEnvironment, BinaryDependencies, LinkEnvironmentVisitedModules);
 			}
 
 			foreach (string AdditionalLibrary in BinaryLinkEnvironment.Config.AdditionalLibraries)
