Fix missing build product when compiling UE4Game on Mac with the -precompile option.

diff --git a/Engine/Source/Programs/UnrealBuildTool/Mac/MacToolChain.cs b/Engine/Source/Programs/UnrealBuildTool/Mac/MacToolChain.cs
--- a/Engine/Source/Programs/UnrealBuildTool/Mac/MacToolChain.cs
+++ b/Engine/Source/Programs/UnrealBuildTool/Mac/MacToolChain.cs
@@ -1359,15 +1359,18 @@ namespace UnrealBuildTool
 				Module.SetupPrivateLinkEnvironment(Binary, BinaryLinkEnvironment, BinaryDependencies, LinkEnvironmentVisitedModules);
 			}
 
-			foreach (string AdditionalLibrary in BinaryLinkEnvironment.Config.AdditionalLibraries)
+			if(Binary.Config.Type != UEBuildBinaryType.StaticLibrary)
 			{
-				string LibName = Path.GetFileName(AdditionalLibrary);
-				if (LibName.StartsWith("lib"))
+				foreach (string AdditionalLibrary in BinaryLinkEnvironment.Config.AdditionalLibraries)
 				{
-					if (Path.GetExtension(AdditionalLibrary) == ".dylib" && !String.IsNullOrEmpty(BundleContentsDirectory))
+					string LibName = Path.GetFileName(AdditionalLibrary);
+					if (LibName.StartsWith("lib"))
 					{
-						string Entry = BundleContentsDirectory + "MacOS/" + LibName;
-						Receipt.AddBuildProduct(Entry, BuildProductType.DynamicLibrary);
+						if (Path.GetExtension(AdditionalLibrary) == ".dylib" && !String.IsNullOrEmpty(BundleContentsDirectory))
+						{
+							string Entry = BundleContentsDirectory + "MacOS/" + LibName;
+							Receipt.AddBuildProduct(Entry, BuildProductType.DynamicLibrary);
+						}
 					}
 				}
 			}
