Don't include framework dependencies of static libraries in the app bundle.

diff --git a/Engine/Source/Programs/UnrealBuildTool/Mac/MacToolChain.cs b/Engine/Source/Programs/UnrealBuildTool/Mac/MacToolChain.cs
--- a/Engine/Source/Programs/UnrealBuildTool/Mac/MacToolChain.cs
+++ b/Engine/Source/Programs/UnrealBuildTool/Mac/MacToolChain.cs
@@ -1350,17 +1350,17 @@ namespace UnrealBuildTool
 			}
 
 			// We need to know what third party dylibs would be copied to the bundle
-			var Modules = Binary.GetAllDependencyModules(bIncludeDynamicallyLoaded: false, bForceCircular: false);
-			var BinaryLinkEnvironment = Binary.Target.GlobalLinkEnvironment.DeepCopy();
-			var BinaryDependencies = new List<UEBuildBinary>();
-			var LinkEnvironmentVisitedModules = new Dictionary<UEBuildModule, bool>();
-			foreach (var Module in Modules)
-			{
-				Module.SetupPrivateLinkEnvironment(Binary, BinaryLinkEnvironment, BinaryDependencies, LinkEnvironmentVisitedModules);
-			}
-
 			if(Binary.Config.Type != UEBuildBinaryType.StaticLibrary)
 			{
+				var Modules = Binary.GetAllDependencyModules(bIncludeDynamicallyLoaded: false, bForceCircular: false);
+				var BinaryLinkEnvironment = Binary.Target.GlobalLinkEnvironment.DeepCopy();
+				var BinaryDependencies = new List<UEBuildBinary>();
+				var LinkEnvironmentVisitedModules = new Dictionary<UEBuildModule, bool>();
+				foreach (var Module in Modules)
+				{
+					Module.SetupPrivateLinkEnvironment(Binary, BinaryLinkEnvironment, BinaryDependencies, LinkEnvironmentVisitedModules);
+				}
+
 				foreach (string AdditionalLibrary in BinaryLinkEnvironment.Config.AdditionalLibraries)
 				{
 					string LibName = Path.GetFileName(AdditionalLibrary);
@@ -1373,20 +1373,20 @@ namespace UnrealBuildTool
 						}
 					}
 				}
-			}
 
-			foreach (UEBuildBundleResource Resource in BinaryLinkEnvironment.Config.AdditionalBundleResources)
-			{
-				if (Directory.Exists(Resource.ResourcePath))
+				foreach (UEBuildBundleResource Resource in BinaryLinkEnvironment.Config.AdditionalBundleResources)
 				{
-					foreach (string ResourceFile in Directory.GetFiles(Resource.ResourcePath, "*", SearchOption.AllDirectories))
+					if (Directory.Exists(Resource.ResourcePath))
 					{
-						Receipt.AddBuildProduct(Path.Combine(BundleContentsDirectory, Resource.BundleContentsSubdir, ResourceFile.Substring(Path.GetDirectoryName(Resource.ResourcePath).Length + 1)), BuildProductType.RequiredResource);
+						foreach (string ResourceFile in Directory.GetFiles(Resource.ResourcePath, "*", SearchOption.AllDirectories))
+						{
+							Receipt.AddBuildProduct(Path.Combine(BundleContentsDirectory, Resource.BundleContentsSubdir, ResourceFile.Substring(Path.GetDirectoryName(Resource.ResourcePath).Length + 1)), BuildProductType.RequiredResource);
+						}
+					}
+					else
+					{
+						Receipt.AddBuildProduct(Path.Combine(BundleContentsDirectory, Resource.BundleContentsSubdir, Path.GetFileName(Resource.ResourcePath)), BuildProductType.RequiredResource);
 					}
-				}
-				else
-				{
-					Receipt.AddBuildProduct(Path.Combine(BundleContentsDirectory, Resource.BundleContentsSubdir, Path.GetFileName(Resource.ResourcePath)), BuildProductType.RequiredResource);
 				}
 			}
 
