Fix calls to IsCompiledInConfiguration() which rely on UBT's global state. The function is called from UBT, where this information is not accurate.

diff --git a/Engine/Source/Programs/AutomationTool/ProjectUtils.cs b/Engine/Source/Programs/AutomationTool/ProjectUtils.cs
--- a/Engine/Source/Programs/AutomationTool/ProjectUtils.cs
+++ b/Engine/Source/Programs/AutomationTool/ProjectUtils.cs
@@ -220,7 +220,7 @@ namespace AutomationTool
 				{
 					if(UProjectInfo.IsPluginEnabledForProject(Plugin, Project, TargetPlatformType) != Plugin.Descriptor.bEnabledByDefault)
 					{
-						if(Plugin.Descriptor.Modules.Any(Module => Module.IsCompiledInConfiguration(TargetPlatformType, TargetRules.TargetType.Game)))
+						if(Plugin.Descriptor.Modules.Any(Module => Module.IsCompiledInConfiguration(TargetPlatformType, TargetRules.TargetType.Game, bBuildDeveloperTools: false, bBuildEditor: false)))
 						{
 							RetVal = true;
 							break;
diff --git a/Engine/Source/Programs/AutomationTool/Scripts/BuildPluginCommand.Automation.cs b/Engine/Source/Programs/AutomationTool/Scripts/BuildPluginCommand.Automation.cs
--- a/Engine/Source/Programs/AutomationTool/Scripts/BuildPluginCommand.Automation.cs
+++ b/Engine/Source/Programs/AutomationTool/Scripts/BuildPluginCommand.Automation.cs
@@ -80,7 +80,9 @@ class BuildPlugin : BuildCommand
 		List<string> ModuleNames = new List<string>();
 		foreach(ModuleDescriptor Module in Plugin.Modules)
 		{
-			if(Module.IsCompiledInConfiguration(Platform, TargetType))
+			bool bBuildDeveloperTools = (TargetType == TargetRules.TargetType.Editor || TargetType == TargetRules.TargetType.Program);
+			bool bBuildEditor = (TargetType == TargetRules.TargetType.Editor);
+			if(Module.IsCompiledInConfiguration(Platform, TargetType, bBuildDeveloperTools, bBuildEditor))
 			{
 				ModuleNames.Add(Module.Name);
 			}
diff --git a/Engine/Source/Programs/UnrealBuildTool/Configuration/UEBuildTarget.cs b/Engine/Source/Programs/UnrealBuildTool/Configuration/UEBuildTarget.cs
--- a/Engine/Source/Programs/UnrealBuildTool/Configuration/UEBuildTarget.cs
+++ b/Engine/Source/Programs/UnrealBuildTool/Configuration/UEBuildTarget.cs
@@ -2221,7 +2221,7 @@ namespace UnrealBuildTool
 			{
 				foreach(ModuleDescriptor Module in Plugin.Descriptor.Modules)
 				{
-					if (Module.IsCompiledInConfiguration(Platform, TargetType))
+					if (Module.IsCompiledInConfiguration(Platform, TargetType, UEBuildConfiguration.bBuildDeveloperTools, UEBuildConfiguration.bBuildEditor))
 					{
 						// Add the corresponding binary for it
 						string ModuleFileName = RulesCompiler.GetModuleFilename(Module.Name);
@@ -2266,7 +2266,7 @@ namespace UnrealBuildTool
 				{
 					foreach(ModuleDescriptor Module in Plugin.Descriptor.Modules)
 					{
-						if (Module.IsCompiledInConfiguration(Platform, TargetType))
+						if (Module.IsCompiledInConfiguration(Platform, TargetType, UEBuildConfiguration.bBuildDeveloperTools, UEBuildConfiguration.bBuildEditor))
 						{
 							PrecompiledModuleNames.Add(Module.Name);
 						}
diff --git a/Engine/Source/Programs/UnrealBuildTool/System/ModuleDescriptor.cs b/Engine/Source/Programs/UnrealBuildTool/System/ModuleDescriptor.cs
--- a/Engine/Source/Programs/UnrealBuildTool/System/ModuleDescriptor.cs
+++ b/Engine/Source/Programs/UnrealBuildTool/System/ModuleDescriptor.cs
@@ -160,7 +160,9 @@ namespace UnrealBuildTool
 		/// </summary>
 		/// <param name="Platform">The platform being compiled for</param>
 		/// <param name="TargetType">The type of the target being compiled</param>
-		public bool IsCompiledInConfiguration(UnrealTargetPlatform Platform, TargetRules.TargetType TargetType)
+		/// <param name="bBuildDeveloperTools">Whether the configuration includes developer tools (typically UEBuildConfiguration.bBuildDeveloperTools for UBT callers)</param>
+		/// <param name="bBuildEditor">Whether the configuration includes the editor (typically UEBuildConfiguration.bBuildEditor for UBT callers)</param>
+		public bool IsCompiledInConfiguration(UnrealTargetPlatform Platform, TargetRules.TargetType TargetType, bool bBuildDeveloperTools, bool bBuildEditor)
 		{
 			// Check the platform is whitelisted
 			if(WhitelistPlatforms != null && WhitelistPlatforms.Length > 0 && !WhitelistPlatforms.Contains(Platform))
@@ -181,10 +183,10 @@ namespace UnrealBuildTool
 				case ModuleHostType.RuntimeNoCommandlet:
 					return true;
 				case ModuleHostType.Developer:
-					return UEBuildConfiguration.bBuildDeveloperTools;
+					return bBuildDeveloperTools;
 				case ModuleHostType.Editor:
 				case ModuleHostType.EditorNoCommandlet:
-					return TargetType == TargetRules.TargetType.Editor || UEBuildConfiguration.bBuildEditor;
+					return TargetType == TargetRules.TargetType.Editor || bBuildEditor;
 				case ModuleHostType.Program:
 					return TargetType == TargetRules.TargetType.Program;
 			}
