PR #1171: QMakefileGenerator project target fixes.

- Contributed by [REDACTED_PERSON].

diff --git a/Engine/Source/Programs/UnrealBuildTool/System/QMakefileGenerator.cs b/Engine/Source/Programs/UnrealBuildTool/System/QMakefileGenerator.cs
--- a/Engine/Source/Programs/UnrealBuildTool/System/QMakefileGenerator.cs
+++ b/Engine/Source/Programs/UnrealBuildTool/System/QMakefileGenerator.cs
@@ -375,7 +375,7 @@ namespace UnrealBuildTool
 								QMakeProjectCmdArg = " -project=\"\\\"$$gameProjectFile\\\"\"";
 							}
 							var ConfName = Enum.GetName(typeof(UnrealTargetConfiguration), CurConfiguration);
-							QMakeFileContent.Append(String.Format("{0}-Linux-{1}.commands = $$build {2} {0} Linux {1} $$args\n", TargetName, ConfName, QMakeProjectCmdArg));
+                            QMakeFileContent.Append(String.Format("{0}-Linux-{1}.commands = $$build {0} Linux {1} {2} $$args\n", TargetName, ConfName, QMakeProjectCmdArg));
 							QMakeTargetList += "\t" + TargetName + "-Linux-" + ConfName + " \\\n"; // , TargetName, ConfName);
 						}
 					}
@@ -386,7 +386,7 @@ namespace UnrealBuildTool
 					QMakeProjectCmdArg = " -project=\"\\\"$$gameProjectFile\\\"\"";
 				}
 
-				QMakeFileContent.Append(String.Format("{0}.commands = $$build {1} {0} Linux Development $$args\n\n", TargetName, QMakeProjectCmdArg));
+                QMakeFileContent.Append(String.Format("{0}.commands = $$build {0} Linux Development {1} $$args\n\n", TargetName, QMakeProjectCmdArg));
 				QMakeTargetList += "\t" + TargetName + " \\\n";
 			}
 
