PR #1170: MakefileGenerator configure target fix

- Contributed by [REDACTED_PERSON].
- UnrealBuildTool_Mono.csproj is removed from 4.8 and above.

diff --git a/Engine/Source/Programs/UnrealBuildTool/System/MakefileGenerator.cs b/Engine/Source/Programs/UnrealBuildTool/System/MakefileGenerator.cs
--- a/Engine/Source/Programs/UnrealBuildTool/System/MakefileGenerator.cs
+++ b/Engine/Source/Programs/UnrealBuildTool/System/MakefileGenerator.cs
@@ -148,7 +148,7 @@ namespace UnrealBuildTool
             {
                 // Make sure UBT is updated.
                 MakefileContent.Append("\txbuild /property:Configuration=Development /property:TargetFrameworkVersion=v4.0 /verbosity:quiet /nologo ");
-                MakefileContent.Append("$(UNREALROOTPATH)/Engine/Source/Programs/UnrealBuildTool/UnrealBuildTool_Mono.csproj\n");
+                MakefileContent.Append("$(UNREALROOTPATH)/Engine/Source/Programs/UnrealBuildTool/UnrealBuildTool.csproj\n");
                 MakefileContent.Append("\t$(PROJECTBUILD) -makefile -qmakefile -cmakefile -project=\"\\\"$(GAMEPROJECTFILE)\\\"\" -game -engine \n");
             }
             else
