PR #1135: Support enabling exceptions in a Linux build.

- Contributed by [REDACTED_PERSON].

diff --git a/Engine/Source/Programs/UnrealBuildTool/Linux/LinuxToolChain.cs b/Engine/Source/Programs/UnrealBuildTool/Linux/LinuxToolChain.cs
--- a/Engine/Source/Programs/UnrealBuildTool/Linux/LinuxToolChain.cs
+++ b/Engine/Source/Programs/UnrealBuildTool/Linux/LinuxToolChain.cs
@@ -407,6 +407,11 @@ namespace UnrealBuildTool
                 Result += " -ftls-model=local-dynamic";
             }
 
+            if (CompileEnvironment.Config.bEnableExceptions)
+            {
+                Result += " -fexceptions";
+            }
+
             //Result += " -v";                            // for better error diagnosis
 
 			Result += ArchitectureSpecificDefines(CompileEnvironment.Config.Target.Architecture);
