Stop UBT creating an intermediate directory in Engine/Source when compiling a plugin against UE4Game (which it then warns about on subsequent runs)

diff --git a/Engine/Source/Programs/UnrealBuildTool/Configuration/UEBuildTarget.cs b/Engine/Source/Programs/UnrealBuildTool/Configuration/UEBuildTarget.cs
--- a/Engine/Source/Programs/UnrealBuildTool/Configuration/UEBuildTarget.cs
+++ b/Engine/Source/Programs/UnrealBuildTool/Configuration/UEBuildTarget.cs
@@ -2747,7 +2747,14 @@ namespace UnrealBuildTool
 				GlobalCompileEnvironment.Config.OutputDirectory = Path.Combine(BuildConfiguration.PlatformIntermediatePath, OutputAppName, IntermediateConfiguration.ToString());
 				if (ShouldCompileMonolithic())
 				{
-					GlobalCompileEnvironment.Config.OutputDirectory = Path.Combine(UnrealBuildTool.GetUProjectPath(), BuildConfiguration.PlatformIntermediateFolder, OutputAppName, IntermediateConfiguration.ToString());
+					if(!String.IsNullOrEmpty(UnrealBuildTool.GetUProjectPath()))
+					{
+						GlobalCompileEnvironment.Config.OutputDirectory = Path.Combine(UnrealBuildTool.GetUProjectPath(), BuildConfiguration.PlatformIntermediateFolder, OutputAppName, IntermediateConfiguration.ToString());
+					}
+					else if(ForeignPlugins.Count > 0)
+					{
+						GlobalCompileEnvironment.Config.OutputDirectory = Path.Combine(Path.GetDirectoryName(ForeignPlugins[0]), BuildConfiguration.PlatformIntermediateFolder, OutputAppName, IntermediateConfiguration.ToString());
+					}
 				}
 			}
 			else
