Output an error if trying to compile a target for an unsupported platform in Rocket (eg. editor in Win32).

diff --git a/Engine/Source/Programs/UnrealBuildTool/System/VCProject.cs b/Engine/Source/Programs/UnrealBuildTool/System/VCProject.cs
--- a/Engine/Source/Programs/UnrealBuildTool/System/VCProject.cs
+++ b/Engine/Source/Programs/UnrealBuildTool/System/VCProject.cs
@@ -191,7 +191,7 @@ namespace UnrealBuildTool
 
 			var SupportedConfigurations = new List<UnrealTargetConfiguration>();
 			var SupportedPlatforms = new List<UnrealTargetPlatform>();
-			if (!ProjectFileGenerator.bCreateDummyConfigsForUnsupportedPlatforms || ProjectFileGenerator.bGeneratingRocketProjectFiles)
+			if (!ProjectFileGenerator.bCreateDummyConfigsForUnsupportedPlatforms)
 			{
 				if( ProjectTarget.TargetRules != null )
 				{
@@ -908,6 +908,23 @@ namespace UnrealBuildTool
 						"		<NMakeOutput/>" + ProjectFileGenerator.NewLine +
 						"	</PropertyGroup>" + ProjectFileGenerator.NewLine);
 				}
+				else if(ProjectFileGenerator.bGeneratingRocketProjectFiles && Combination.ProjectTarget != null && Combination.ProjectTarget.TargetRules != null && !Combination.ProjectTarget.TargetRules.SupportsPlatform(Combination.Platform))
+				{
+					List<UnrealTargetPlatform> SupportedPlatforms = new List<UnrealTargetPlatform>();
+					Combination.ProjectTarget.TargetRules.GetSupportedPlatforms(ref SupportedPlatforms);
+
+					string ProjectRelativeUnusedDirectory = NormalizeProjectPath(Path.Combine(ProjectFileGenerator.EngineRelativePath, BuildConfiguration.BaseIntermediateFolder, "Unused"));
+
+					VCProjectFileContent.AppendFormat(
+						"	<PropertyGroup " + ConditionString + ">" + ProjectFileGenerator.NewLine +
+						"		<OutDir>" + ProjectRelativeUnusedDirectory + Path.DirectorySeparatorChar + "</OutDir>" + ProjectFileGenerator.NewLine +
+						"		<IntDir>" + ProjectRelativeUnusedDirectory + Path.DirectorySeparatorChar + "</IntDir>" + ProjectFileGenerator.NewLine +
+						"		<NMakeBuildCommandLine>@echo {0} is not a supported platform for {1}. Valid platforms are {2}.</NMakeBuildCommandLine>" + ProjectFileGenerator.NewLine +
+						"		<NMakeReBuildCommandLine>@echo {0} is not a supported platform for {1}. Valid platforms are {2}.</NMakeReBuildCommandLine>" + ProjectFileGenerator.NewLine +
+						"		<NMakeCleanCommandLine>@echo {0} is not a supported platform for {1}. Valid platforms are {2}.</NMakeCleanCommandLine>" + ProjectFileGenerator.NewLine +
+						"		<NMakeOutput/>" + ProjectFileGenerator.NewLine +
+						"	</PropertyGroup>" + ProjectFileGenerator.NewLine, Combination.Platform, Utils.GetFilenameWithoutAnyExtensions(Combination.ProjectTarget.TargetFilePath), String.Join(", ", SupportedPlatforms.Select(x => x.ToString())));
+				}
 				else
 				{
 					TargetRules TargetRulesObject = Combination.ProjectTarget.TargetRules;
