UE-14803 - Android Java files constantly being rewritten triggering rebuilds
- Updated logic to change how files are written and when the OBB data file is written.

diff --git a/Engine/Source/Programs/UnrealBuildTool/Android/UEDeployAndroid.cs b/Engine/Source/Programs/UnrealBuildTool/Android/UEDeployAndroid.cs
--- a/Engine/Source/Programs/UnrealBuildTool/Android/UEDeployAndroid.cs
+++ b/Engine/Source/Programs/UnrealBuildTool/Android/UEDeployAndroid.cs
@@ -397,10 +397,21 @@ namespace UnrealBuildTool.Android
             //
             obbData.Append("};\n"); // close class definition off
 
-            if (obbDataFile == null || !obbDataFile.SequenceEqual(obbData.ToString().Split('\n')))
+            if (obbDataFile == null || !obbDataFile.SequenceEqual((obbData.ToString()).Split('\n')))
             {
                 MakeDirectoryIfRequired(FileName);
-                File.WriteAllText(FileName, obbData.ToString());
+                using (StreamWriter outputFile = new StreamWriter(FileName, false))
+                {
+                    var obbSrc = obbData.ToString().Split('\n');
+                    foreach (var line in obbSrc)
+                    {
+                        outputFile.WriteLine(line);
+                    }
+                }
+            }
+            else
+            {
+                Log.TraceInformation("\n==== OBB data file up to date so not writing. ====");
             }
         }
 
@@ -411,12 +422,12 @@ namespace UnrealBuildTool.Android
             string[] DestFileContent = File.Exists(ShimFileName) ? File.ReadAllLines(ShimFileName) : null;
 
             StringBuilder ShimFileContent = new StringBuilder("package com.epicgames.ue4;\n\n");
-            ShimFileContent.AppendFormat("import {0}.OBBDownloaderService;", replacements["$$PackageName$$"]);
-            ShimFileContent.AppendFormat("import {0}.DownloaderActivity;", replacements["$$PackageName$$"]);
+            ShimFileContent.AppendFormat("import {0}.OBBDownloaderService;\n", replacements["$$PackageName$$"]);
+            ShimFileContent.AppendFormat("import {0}.DownloaderActivity;\n", replacements["$$PackageName$$"]);
             ShimFileContent.Append("\n\npublic class DownloadShim\n{\n");
             ShimFileContent.Append("\tpublic static OBBDownloaderService DownloaderService;\n");
             ShimFileContent.Append("\tpublic static DownloaderActivity DownloadActivity;\n");
-            ShimFileContent.Append("\tpublic static Class<DownloaderActivity> GetDownloaderType() { return DownloaderActivity.class; }");
+            ShimFileContent.Append("\tpublic static Class<DownloaderActivity> GetDownloaderType() { return DownloaderActivity.class; }\n");
             ShimFileContent.Append("}\n");
 
             Log.TraceInformation("\n==== Writing to shim file {0} ====", ShimFileName);
@@ -425,7 +436,18 @@ namespace UnrealBuildTool.Android
             if (DestFileContent == null || !DestFileContent.SequenceEqual((ShimFileContent.ToString()).Split('\n')))
             {
                 MakeDirectoryIfRequired(ShimFileName);
-                File.WriteAllText(ShimFileName, ShimFileContent.ToString());
+                using (StreamWriter outputFile = new StreamWriter(ShimFileName, false))
+                {
+                    var shimSrc = ShimFileContent.ToString().Split('\n');
+                    foreach (var line in shimSrc)
+                    {
+                        outputFile.WriteLine(line);
+                    }
+                }
+            }
+            else
+            {
+                Log.TraceInformation("\n==== Shim data file up to date so not writing. ====");
             }
 
             // Now we move on to the template files
@@ -465,6 +487,10 @@ namespace UnrealBuildTool.Android
                         }
                     }
                 }
+                else
+                {
+                    Log.TraceInformation("\n==== Template target file up to date so not writing. ====");
+                }
             }
         }
 
@@ -1153,7 +1179,13 @@ namespace UnrealBuildTool.Android
 
             // Generate the OBB and Shim files here
             string ObbFileLocation = ProjectDirectory + "/Saved/StagedBuilds/Android" + CookFlavor + ".obb";
-            WriteJavaOBBDataFile(UE4OBBDataFileName, PackageName, new List<string>{ObbFileLocation});
+            
+            // This is kind of a small hack to get around a rewrite problem
+            // We need to make sure the file is there but if the OBB file doesn't exist then we don't want to replace it
+            if (File.Exists(ObbFileLocation) || !File.Exists(UE4OBBDataFileName))
+            {
+                WriteJavaOBBDataFile(UE4OBBDataFileName, PackageName, new List<string> { ObbFileLocation });
+            }
 
             WriteJavaDownloadSupportFiles(UE4DownloadShimFileName, templates, new Dictionary<string, string>{
                 { "$$GameName$$", ProjectName },
