Fix for target file lookup where multiple projects have the same prefix, e.g. TP_Vehicle and TP_VehicleAdv.

diff --git a/Engine/Source/Programs/UnrealBuildTool/System/UProjectInfo.cs b/Engine/Source/Programs/UnrealBuildTool/System/UProjectInfo.cs
--- a/Engine/Source/Programs/UnrealBuildTool/System/UProjectInfo.cs
+++ b/Engine/Source/Programs/UnrealBuildTool/System/UProjectInfo.cs
@@ -53,7 +53,7 @@ namespace UnrealBuildTool
 				foreach (KeyValuePair<string, UProjectInfo> Entry in ProjectInfoDictionary)
 				{
 					FileInfo ProjectFileInfo = new FileInfo(Entry.Key);
-					string ProjectDir = ProjectFileInfo.DirectoryName;
+					string ProjectDir = ProjectFileInfo.DirectoryName.TrimEnd(Path.DirectorySeparatorChar) + Path.DirectorySeparatorChar;
 					if (TargetFilename.StartsWith(ProjectDir, StringComparison.InvariantCultureIgnoreCase))
 					{
 						FileInfo TargetInfo = new FileInfo(TargetFilename);
