UEPLAT-831 - AndroidManifest.xml update for ES3.1 targets
- Changed MakeAPK loop so that it can deal with multiple manifests
- Removed ability to make a Fat APK from same code as we don't/won't use it

diff --git a/Engine/Source/Programs/AutomationTool/Android/AndroidPlatform.Automation.cs b/Engine/Source/Programs/AutomationTool/Android/AndroidPlatform.Automation.cs
--- a/Engine/Source/Programs/AutomationTool/Android/AndroidPlatform.Automation.cs
+++ b/Engine/Source/Programs/AutomationTool/Android/AndroidPlatform.Automation.cs
@@ -114,8 +114,8 @@ public class AndroidPlatform : Platform
 
 	public override void Package(ProjectParams Params, DeploymentContext SC, int WorkingCL)
 	{
-		string[] Architectures = UnrealBuildTool.AndroidToolChain.GetAllArchitectures();
-		string[] GPUArchitectures = UnrealBuildTool.AndroidToolChain.GetAllGPUArchitectures();
+		var Architectures = UnrealBuildTool.AndroidToolChain.GetAllArchitectures();
+		var GPUArchitectures = UnrealBuildTool.AndroidToolChain.GetAllGPUArchitectures();
 		bool bMakeSeparateApks = UnrealBuildTool.Android.UEDeployAndroid.ShouldMakeSeparateApks();
 		bool bPackageDataInsideApk = UnrealBuildTool.Android.UEDeployAndroid.PackageDataInsideApk(false);
 
@@ -300,8 +300,8 @@ public class AndroidPlatform : Platform
 			throw new AutomationException(ErrorString);
 		}
 
-		string[] Architectures = UnrealBuildTool.AndroidToolChain.GetAllArchitectures();
-		string[] GPUArchitectures = UnrealBuildTool.AndroidToolChain.GetAllGPUArchitectures();
+		var Architectures = UnrealBuildTool.AndroidToolChain.GetAllArchitectures();
+		var GPUArchitectures = UnrealBuildTool.AndroidToolChain.GetAllGPUArchitectures();
 		bool bMakeSeparateApks = UnrealBuildTool.Android.UEDeployAndroid.ShouldMakeSeparateApks();
 		bool bPackageDataInsideApk = UnrealBuildTool.Android.UEDeployAndroid.PackageDataInsideApk(false);
 
@@ -506,7 +506,7 @@ public class AndroidPlatform : Platform
 		// Setup the OBB name and add the storage path (queried from the device) to it
 		string DeviceStorageQueryCommand = GetStorageQueryCommand();
 		ProcessResult Result = RunAdbCommand(Params, DeviceStorageQueryCommand, null, ERunOptions.AppMustExist);
-		String StorageLocation = Result.Output.Trim();
+        String StorageLocation = Result.Output.Trim(); // "/mnt/sdcard";
 		string DeviceObbName = StorageLocation + "/" + GetDeviceObbName(ApkName);
 		string RemoteDir = StorageLocation + "/UE4Game/" + Params.ShortProjectName;
 
@@ -862,7 +862,7 @@ public class AndroidPlatform : Platform
 			return "";
 		}
 
-		string[] AppArchitectures = AndroidToolChain.GetAllArchitectures();
+		var AppArchitectures = AndroidToolChain.GetAllArchitectures();
 
 		// ask the device
 		ProcessResult ABIResult = RunAdbCommand(Params, " shell getprop ro.product.cpu.abi", null, ERunOptions.AppMustExist);
@@ -871,7 +871,7 @@ public class AndroidPlatform : Platform
 		string DeviceArch = UnrealBuildTool.Android.UEDeployAndroid.GetUE4Arch(ABIResult.Output.Trim());
 
 		// if the architecture wasn't built, look for a backup
-		if (Array.IndexOf(AppArchitectures, DeviceArch) == -1)
+		if (!AppArchitectures.Contains(DeviceArch))
 		{
 			// go from 64 to 32-bit
 			if (DeviceArch == "-arm64")
@@ -881,7 +881,7 @@ public class AndroidPlatform : Platform
 			// go from 64 to 32-bit
 			else if (DeviceArch == "-x64")
 			{
-				if (Array.IndexOf(AppArchitectures, "-x86") != -1)
+				if (!AppArchitectures.Contains("-x86"))
 				{
 					DeviceArch = "-x86";
 				}
@@ -910,7 +910,7 @@ public class AndroidPlatform : Platform
 		}
 
 		// if after the fallbacks, we still don't have it, we can't continue
-		if (Array.IndexOf(AppArchitectures, DeviceArch) == -1)
+		if (!AppArchitectures.Contains(DeviceArch))
 		{
 			string ErrorString = String.Format("Unable to run because you don't have an apk that is usable on {0}. Looked for {1}", Params.Device, DeviceArch);
             ErrorReporter.Error(ErrorString, (int)ErrorCodes.Error_NoApkSuitableForArchitecture);
@@ -929,7 +929,7 @@ public class AndroidPlatform : Platform
 			return "";
 		}
 
-		string[] AppGPUArchitectures = AndroidToolChain.GetAllGPUArchitectures();
+		var AppGPUArchitectures = AndroidToolChain.GetAllGPUArchitectures();
 
 		// get the device extensions
 		ProcessResult ExtensionsResult = RunAdbCommand(Params, "shell dumpsys SurfaceFlinger", null, ERunOptions.AppMustExist);
diff --git a/Engine/Source/Programs/UnrealBuildTool/Android/AndroidToolChain.cs b/Engine/Source/Programs/UnrealBuildTool/Android/AndroidToolChain.cs
--- a/Engine/Source/Programs/UnrealBuildTool/Android/AndroidToolChain.cs
+++ b/Engine/Source/Programs/UnrealBuildTool/Android/AndroidToolChain.cs
@@ -5,6 +5,7 @@ using System.Collections.Generic;
 using System.Text.RegularExpressions;
 using System.Diagnostics;
 using System.IO;
+using System.Linq;
 using Microsoft.Win32;
 
 namespace UnrealBuildTool
@@ -15,11 +16,11 @@ namespace UnrealBuildTool
 		static private float ClangVersionFloat = 0;
 
 		// the list of architectures we will compile for
-		static private string[] Arches = null;
+		static private List<string> Arches = null;
 		// the list of GPU architectures we will compile for
-		static private string[] GPUArchitectures = null;
+		static private List<string> GPUArchitectures = null;
 		// a list of all architecture+GPUArchitecture names (-armv7-es2, etc)
-		static private string[] AllComboNames = null;
+		static private List<string> AllComboNames = null;
 
 		static private Dictionary<string, string[]> AllArchNames = new Dictionary<string, string[]> {
 			{ "-armv7", new string[] { "armv7", "armeabi-v7a", } }, 
@@ -39,64 +40,75 @@ namespace UnrealBuildTool
 		{
 			// look in ini settings for what platforms to compile for
 			ConfigCacheIni Ini = new ConfigCacheIni(UnrealTargetPlatform.Android, "Engine", UnrealBuildTool.GetUProjectPath());
-			List<string> ProjectArches = new List<string>();
+            Arches = new List<string>();
 			bool bBuild = true;
 			if (Ini.GetBool("/Script/AndroidRuntimeSettings.AndroidRuntimeSettings", "bBuildForArmV7", out bBuild) && bBuild)
 			{
-				ProjectArches.Add("-armv7");
+                Arches.Add("-armv7");
 			}
 			if (Ini.GetBool("/Script/AndroidRuntimeSettings.AndroidRuntimeSettings", "bBuildForArm64", out bBuild) && bBuild)
 			{
-				ProjectArches.Add("-arm64");
+                Arches.Add("-arm64");
 			}
 			if (Ini.GetBool("/Script/AndroidRuntimeSettings.AndroidRuntimeSettings", "bBuildForx86", out bBuild) && bBuild)
 			{
-				ProjectArches.Add("-x86");
+                Arches.Add("-x86");
 			}
 			if (Ini.GetBool("/Script/AndroidRuntimeSettings.AndroidRuntimeSettings", "bBuildForx8664", out bBuild) && bBuild)
 			{
-				ProjectArches.Add("-x64");
+                Arches.Add("-x64");
 			}
 
 			// force armv7 if something went wrong
-			if (ProjectArches.Count == 0)
+            if (Arches.Count == 0)
 			{
-				ProjectArches.Add("-armv7");
+                Arches.Add("-armv7");
 			}
 
-			Arches = ProjectArches.ToArray();
-
 			// Parse selected GPU architectures
-			List<string> ProjectGPUArches = new List<string>();
+            GPUArchitectures = new List<string>();
 			if (Ini.GetBool("/Script/AndroidRuntimeSettings.AndroidRuntimeSettings", "bBuildForES2", out bBuild) && bBuild)
 			{
-				ProjectGPUArches.Add("-es2");
+                GPUArchitectures.Add("-es2");
 			}
 			if (Ini.GetBool("/Script/AndroidRuntimeSettings.AndroidRuntimeSettings", "bBuildForES31", out bBuild) && bBuild)
 			{
-				ProjectGPUArches.Add("-es31");
+                GPUArchitectures.Add("-es31");
 			}
 			if (Ini.GetBool("/Script/AndroidRuntimeSettings.AndroidRuntimeSettings", "bBuildForGL4", out bBuild) && bBuild)
 			{
-				ProjectGPUArches.Add("-gl4");
+                GPUArchitectures.Add("-gl4");
 			}
-			if (ProjectGPUArches.Count == 0)
+            if (GPUArchitectures.Count == 0)
 			{
-				ProjectGPUArches.Add("-es2");
+                GPUArchitectures.Add("-es2");
 			}
-			GPUArchitectures = ProjectGPUArches.ToArray();
 
-			List<string> FullArchCombinations = new List<string>();
-			foreach (string Arch in Arches)
-			{
-				foreach (string GPUArch in GPUArchitectures)
-				{
-					FullArchCombinations.Add(Arch + GPUArch);
-				}
-			}
-			AllComboNames = FullArchCombinations.ToArray();
+            AllComboNames = (from Arch in Arches
+                            from GPUArch in GPUArchitectures
+                             select Arch + GPUArch).ToList();
 		}
 
+        static public string GetGLESVersionFromGPUArch(string GPUArch)
+        {
+            GPUArch = GPUArch.Substring(1); // drop the '-' from the start
+            string GLESversion = "";
+            switch (GPUArch)
+            {
+                case "es2":
+                    GLESversion = "0x00020000";
+                    break;
+                case "es31":
+                    GLESversion = "0x00030001";
+                    break;
+                default:
+                    GLESversion = "0x00020000";
+                    break;
+            }
+
+            return GLESversion;
+        }
+
 		public override void SetUpGlobalEnvironment()
 		{
 			base.SetUpGlobalEnvironment();
@@ -104,7 +116,7 @@ namespace UnrealBuildTool
 			ParseArchitectures();
 		}
 
-		static public string[] GetAllArchitectures()
+		static public List<string> GetAllArchitectures()
 		{
 			if (Arches == null)
 			{
@@ -114,7 +126,7 @@ namespace UnrealBuildTool
 			return Arches;
 		}
 
-		static public string[] GetAllGPUArchitectures()
+		static public List<string> GetAllGPUArchitectures()
 		{
 			if (GPUArchitectures == null)
 			{
@@ -584,7 +596,7 @@ namespace UnrealBuildTool
 		static private bool bHasPrintedApiLevel = false;
 		public override CPPOutput CompileCPPFiles(UEBuildTarget Target, CPPEnvironment CompileEnvironment, List<FileItem> SourceFiles, string ModuleName)
 		{
-			if (Arches.Length == 0)
+			if (Arches.Count == 0)
 			{
 				throw new BuildException("At least one architecture (armv7, x86, etc) needs to be selected in the project settings to build");
 			}
@@ -802,13 +814,13 @@ namespace UnrealBuildTool
 		{
 			List<FileItem> Outputs = new List<FileItem>();
 
-			for (int ArchIndex = 0; ArchIndex < Arches.Length; ArchIndex++)
+			for (int ArchIndex = 0; ArchIndex < Arches.Count; ArchIndex++)
 			{
 				string Arch = Arches[ArchIndex];
-				for (int GPUArchIndex = 0; GPUArchIndex < GPUArchitectures.Length; GPUArchIndex++)
+				for (int GPUArchIndex = 0; GPUArchIndex < GPUArchitectures.Count; GPUArchIndex++)
 				{
 					string GPUArchitecture = GPUArchitectures[GPUArchIndex];
-					int OutputPathIndex = ArchIndex * GPUArchitectures.Length + GPUArchIndex;
+					int OutputPathIndex = ArchIndex * GPUArchitectures.Count + GPUArchIndex;
 
 					// Android will have an array of outputs
 					if (LinkEnvironment.Config.OutputFilePaths.Count < OutputPathIndex ||
diff --git a/Engine/Source/Programs/UnrealBuildTool/Android/UEBuildAndroid.cs b/Engine/Source/Programs/UnrealBuildTool/Android/UEBuildAndroid.cs
--- a/Engine/Source/Programs/UnrealBuildTool/Android/UEBuildAndroid.cs
+++ b/Engine/Source/Programs/UnrealBuildTool/Android/UEBuildAndroid.cs
@@ -430,8 +430,8 @@ namespace UnrealBuildTool
 
 		public override List<string> FinalizeBinaryPaths(string BinaryName)
 		{
-			string[] Architectures = AndroidToolChain.GetAllArchitectures();
-			string[] GPUArchitectures = AndroidToolChain.GetAllGPUArchitectures();
+			var Architectures = AndroidToolChain.GetAllArchitectures();
+			var GPUArchitectures = AndroidToolChain.GetAllGPUArchitectures();
 
 			// make multiple output binaries
 			List<string> AllBinaries = new List<string>();
diff --git a/Engine/Source/Programs/UnrealBuildTool/Android/UEDeployAndroid.cs b/Engine/Source/Programs/UnrealBuildTool/Android/UEDeployAndroid.cs
--- a/Engine/Source/Programs/UnrealBuildTool/Android/UEDeployAndroid.cs
+++ b/Engine/Source/Programs/UnrealBuildTool/Android/UEDeployAndroid.cs
@@ -782,13 +782,13 @@ namespace UnrealBuildTool.Android
 				}
 			}
 
-			string[] Arches = AndroidToolChain.GetAllArchitectures();
+			var Arches = AndroidToolChain.GetAllArchitectures();
 			foreach (string Arch in Arches)
 			{
 				CurrentSettings.AppendFormat("Arch={0}{1}", Arch, Environment.NewLine);
 			}
 
-			string[] GPUArchitectures = AndroidToolChain.GetAllGPUArchitectures();
+			var GPUArchitectures = AndroidToolChain.GetAllGPUArchitectures();
 			foreach (string GPUArch in GPUArchitectures)
 			{
 				CurrentSettings.AppendFormat("GPUArch={0}{1}", GPUArch, Environment.NewLine);
@@ -800,8 +800,8 @@ namespace UnrealBuildTool.Android
         private bool CheckDependencies(string ProjectName, string ProjectDirectory, string UE4BuildFilesPath, string GameBuildFilesPath, string EngineDirectory, List<string> SettingsFiles,
 			string CookFlavor, string OutputPath, string UE4BuildPath, bool bMakeSeparateApks, bool bPackageDataInsideApk)
 		{
-			string[] Arches = AndroidToolChain.GetAllArchitectures();
-			string[] GPUArchitectures = AndroidToolChain.GetAllGPUArchitectures();
+			var Arches = AndroidToolChain.GetAllArchitectures();
+			var GPUArchitectures = AndroidToolChain.GetAllGPUArchitectures();
 
 			// check all input files (.so, java files, .ini files, etc)
 			bool bAllInputsCurrent = true;
@@ -931,7 +931,9 @@ namespace UnrealBuildTool.Android
             return PlayLicenseKey;
         }
 
-        private string GenerateManifest(string ProjectName, bool bIsForDistribution, bool bPackageDataInsideApk, string GameBuildFilesPath, bool bHasOBBFiles, bool bDisableVerifyOBBOnStartUp)
+        
+
+        private string GenerateManifest(string ProjectName, bool bIsForDistribution, bool bPackageDataInsideApk, string GameBuildFilesPath, bool bHasOBBFiles, bool bDisableVerifyOBBOnStartUp, string GPUArch)
 		{
 			// ini file to get settings from
 			ConfigCacheIni Ini = GetConfigCacheIni("Engine");
@@ -1104,7 +1106,7 @@ namespace UnrealBuildTool.Android
 			{
 				// need just the number part of the sdk
 				Text.AppendLine(string.Format("\t<uses-sdk android:minSdkVersion=\"{0}\"/>", MinSDKVersion));
-				Text.AppendLine("\t<uses-feature android:glEsVersion=\"0x00020000\" android:required=\"true\" />");
+				Text.AppendLine("\t<uses-feature android:glEsVersion=\"" + AndroidToolChain.GetGLESVersionFromGPUArch(GPUArch) +"\" android:required=\"true\" />");
 				Text.AppendLine("\t<uses-permission android:name=\"android.permission.INTERNET\"/>");
 				Text.AppendLine("\t<uses-permission android:name=\"android.permission.WRITE_EXTERNAL_STORAGE\"/>");
 				Text.AppendLine("\t<uses-permission android:name=\"android.permission.ACCESS_NETWORK_STATE\"/>");
@@ -1158,11 +1160,7 @@ namespace UnrealBuildTool.Android
             string UE4JavaFilePath = Path.Combine(ProjectDirectory, "Build", "Android", GetUE4JavaSrcPath());
 			string UE4BuildFilesPath = GetUE4BuildFilePath(EngineDirectory);
 			string GameBuildFilesPath = Path.Combine(ProjectDirectory, "Build/Android");
-	
-			string[] Arches = AndroidToolChain.GetAllArchitectures();
-			string[] GPUArchitectures = AndroidToolChain.GetAllGPUArchitectures();
-//			int NumArches = Arches.Length * GPUArchitectures.Length;
-                       
+	                     
             // Generate Java files
             string PackageName = GetPackageName(ProjectName);
             string TemplateDestinationBase = Path.Combine(ProjectDirectory, "Build", "Android", "src" , PackageName.Replace('.', '/'));
@@ -1203,21 +1201,9 @@ namespace UnrealBuildTool.Android
 			string BuildSettingsCacheFile = Path.Combine(UE4BuildPath, "UEBuildSettings.txt");
 
 			// do we match previous build settings?
-			bool bBuildSettingsMatch = false;
-
-			string ManifestFile = Path.Combine(UE4BuildPath, "AndroidManifest.xml");
-            string NewManifest = GenerateManifest(ProjectName, bForDistribution, bPackageDataInsideApk, GameBuildFilesPath, bDisallowPackagingDataInApk ? false : File.Exists(ObbFileLocation), bDisableVerifyOBBOnStartUp);
-			string OldManifest = File.Exists(ManifestFile) ? File.ReadAllText(ManifestFile) : "";
-			if (NewManifest == OldManifest) 
-			{
-				bBuildSettingsMatch = true;
-			}
-			else
-			{
-				Log.TraceInformation("AndroidManifest.xml that was generated is different than last build, forcing repackage.");
-			}
+			bool bBuildSettingsMatch = true;
 
-			// get application name and whether it changed, needing to force repackage
+            // get application name and whether it changed, needing to force repackage
 			string ApplicationDisplayName;
 			if (CheckApplicationName(UE4BuildPath, ProjectName, out ApplicationDisplayName))
 			{
@@ -1251,16 +1237,42 @@ namespace UnrealBuildTool.Android
                 JavaFiles.AddRange(from t in templates select t.SourceFile);
                 JavaFiles.AddRange(from t in templates select t.DestinationFile);
 
-				bool bAllInputsCurrent = CheckDependencies(ProjectName, ProjectDirectory, UE4BuildFilesPath, GameBuildFilesPath,
+				bBuildSettingsMatch = CheckDependencies(ProjectName, ProjectDirectory, UE4BuildFilesPath, GameBuildFilesPath,
 					EngineDirectory, JavaFiles, CookFlavor, OutputPath, UE4BuildPath, bMakeSeparateApks, bPackageDataInsideApk);
 
-				if (bAllInputsCurrent)
-				{
-					Log.TraceInformation("Output .apk file(s) are up to date (dependencies and build settings are up to date)");
-					return;
-				}
 			}
 
+            var Arches = AndroidToolChain.GetAllArchitectures();
+            var GPUArchitectures = AndroidToolChain.GetAllGPUArchitectures();
+
+            IEnumerable<Tuple<string, string, string>> BuildList = null;
+
+            if(!bBuildSettingsMatch)
+            {
+                BuildList = from Arch in Arches
+                            from GPUArch in GPUArchitectures
+                            let manifest = GenerateManifest(ProjectName, bForDistribution, bPackageDataInsideApk, GameBuildFilesPath, bDisallowPackagingDataInApk ? false : File.Exists(ObbFileLocation), bDisableVerifyOBBOnStartUp, GPUArch)
+                            select Tuple.Create(Arch, GPUArch, manifest);
+            }
+            else
+            {
+                BuildList = from Arch in Arches
+                            from GPUArch in GPUArchitectures
+                            let manifestFile = Path.Combine(IntermediateAndroidPath, Arch + "_" + GPUArch + "_AndroidManifest.xml")
+                            let manifest = GenerateManifest(ProjectName, bForDistribution, bPackageDataInsideApk, GameBuildFilesPath, bDisallowPackagingDataInApk ? false : File.Exists(ObbFileLocation), bDisableVerifyOBBOnStartUp, GPUArch)
+                            let OldManifest = File.Exists(manifestFile) ? File.ReadAllText(manifestFile) : ""
+                            where manifest != OldManifest
+                            select Tuple.Create(Arch, GPUArch, manifest);
+            }
+
+            // Now we have to spin over all the arch/gpu combinations to make sure they all match
+
+            if (BuildList.Count() == 0)
+            {
+                Log.TraceInformation("Output .apk file(s) are up to date (dependencies and build settings are up to date)");
+                return;
+            }
+
 
 			// Once for all arches code:
 
@@ -1317,163 +1329,148 @@ namespace UnrealBuildTool.Android
 			CopyFileDirectory(UE4BuildFilesPath + "/NoRedist", UE4BuildPath, Replacements);
 			CopyFileDirectory(GameBuildFilesPath, UE4BuildPath, Replacements);
 			CopyFileDirectory(GameBuildFilesPath + "/NotForLicensees", UE4BuildPath, Replacements);
-			CopyFileDirectory(GameBuildFilesPath + "/NoRedist", UE4BuildPath, Replacements);
-
-			// in case the game had an AndroidManifest.xml file, we overwrite it now with the generated one
-			File.WriteAllText(ManifestFile, NewManifest);
-
-			// update metadata files (like project.properties, build.xml) if we are missing a build.xml or if we just overwrote project.properties with a bad version in it (from game/engine dir)
-			UpdateProjectProperties(UE4BuildPath, ProjectName);
+			CopyFileDirectory(GameBuildFilesPath + "/NoRedist", UE4BuildPath, Replacements);	
 
 			// at this point, we can write out the cached build settings to compare for a next build
 			File.WriteAllText(BuildSettingsCacheFile, CurrentBuildSettings);
 
-			// now make the apk(s)
-			string FinalNdkBuildABICommand = "";
-			for (int ArchIndex = 0; ArchIndex < Arches.Length; ArchIndex++)
-			{
-				string Arch = Arches[ArchIndex];
-				for (int GPUArchIndex = 0; GPUArchIndex < GPUArchitectures.Length; GPUArchIndex++)
-				{
-					Log.TraceInformation("\n===={0}====PREPARING NATIVE CODE=================================================================", DateTime.Now.ToString());
+            ///////////////
+            // in case the game had an AndroidManifest.xml file, we overwrite it now with the generated one
+            //File.WriteAllText(ManifestFile, NewManifest);
+            ///////////////
 
-					string GPUArchitecture = GPUArchitectures[GPUArchIndex];
-					string SourceSOName = AndroidToolChain.InlineArchName(OutputPath, Arch, GPUArchitecture);
-					// if the source binary was UE4Game, replace it with the new project name, when re-packaging a binary only build
-					string ApkFilename = Path.GetFileNameWithoutExtension(OutputPath).Replace("UE4Game", ProjectName);
-					string DestApkName = Path.Combine(ProjectDirectory, "Binaries/Android/") + ApkFilename + ".apk";
+            Log.TraceInformation("\n===={0}====PREPARING NATIVE CODE=================================================================", DateTime.Now.ToString());
+            bool HasNDKPath = File.Exists(NDKBuildPath);
 
-					// if we making multiple Apks, we need to put the architecture into the name
-					if (bMakeSeparateApks)
-					{
-						DestApkName = AndroidToolChain.InlineArchName(DestApkName, Arch, GPUArchitecture);
-					}
+            foreach(var build in BuildList)
+            {
+                string Arch = build.Item1;
+                string GPUArchitecture = build.Item2;
+                string Manifest = build.Item3;
+
+                string SourceSOName = AndroidToolChain.InlineArchName(OutputPath, Arch, GPUArchitecture);
+                // if the source binary was UE4Game, replace it with the new project name, when re-packaging a binary only build
+                string ApkFilename = Path.GetFileNameWithoutExtension(OutputPath).Replace("UE4Game", ProjectName);
+                string DestApkName = Path.Combine(ProjectDirectory, "Binaries/Android/") + ApkFilename + ".apk";
+                
+                // As we are always making seperate APKs we need to put the architecture into the name
+                DestApkName = AndroidToolChain.InlineArchName(DestApkName, Arch, GPUArchitecture);
+
+                // Write the manifest to the correct locations (cache and real)
+                String ManifestFile = Path.Combine(IntermediateAndroidPath, Arch + "_" + GPUArchitecture + "_AndroidManifest.xml");
+                File.WriteAllText(ManifestFile, Manifest);
+                ManifestFile = Path.Combine(UE4BuildPath, "AndroidManifest.xml");
+                File.WriteAllText(ManifestFile, Manifest);
+
+                // update metadata files (like project.properties, build.xml) if we are missing a build.xml or if we just overwrote project.properties with a bad version in it (from game/engine dir)
+                UpdateProjectProperties(UE4BuildPath, ProjectName);
+
+                // Copy the generated .so file from the binaries directory to the jni folder
+                if (!File.Exists(SourceSOName))
+                {
+                    throw new BuildException("Can't make an APK without the compiled .so [{0}]", SourceSOName);
+                }
+                if (!Directory.Exists(UE4BuildPath + "/jni"))
+                {
+                    throw new BuildException("Can't make an APK without the jni directory [{0}/jni]", UE4BuildFilesPath);
+                }
 
-					// Copy the generated .so file from the binaries directory to the jni folder
-					if (!File.Exists(SourceSOName))
-					{
-						throw new BuildException("Can't make an APK without the compiled .so [{0}]", SourceSOName);
-					}
-					if (!Directory.Exists(UE4BuildPath + "/jni"))
-					{
-						throw new BuildException("Can't make an APK without the jni directory [{0}/jni]", UE4BuildFilesPath);
-					}
+                String FinalSOName;
+                if(HasNDKPath)
+                {
+                    string LibDir = UE4BuildPath + "/jni/" + GetNDKArch(Arch);
+                    FinalSOName = LibDir + "/libUE4.so";
 
-					// Use ndk-build to do stuff and move the .so file to the lib folder (only if NDK is installed)
-					string FinalSOName = "";
-					if (File.Exists(NDKBuildPath))
-					{
-						string LibDir = UE4BuildPath + "/jni/" + GetNDKArch(Arch);
-						FinalSOName = LibDir + "/libUE4.so";
+                    // check to see if libUE4.so is newer than last time we copied
+                    TimeSpan Diff = File.GetLastWriteTimeUtc(FinalSOName) - File.GetLastWriteTimeUtc(SourceSOName);
+                    if (!File.Exists(FinalSOName) || Diff.TotalSeconds < -1 || Diff.TotalSeconds > 1)
+                    {
+                        Log.TraceInformation("\nCopying new .so file to jni folder...");
+                        Directory.CreateDirectory(LibDir);
+                        // copy the binary to the standard .so location
+                        File.Copy(SourceSOName, FinalSOName, true);
+                    }
+                }
+                else
+                {
+                    // if no NDK, we don't need any of the debugger stuff, so we just copy the .so to where it will end up
+                    FinalSOName = UE4BuildPath + "/libs/" + GetNDKArch(Arch) + "/libUE4.so";
 
-						// check to see if libUE4.so is newer than last time we copied
-						TimeSpan Diff = File.GetLastWriteTimeUtc(FinalSOName) - File.GetLastWriteTimeUtc(SourceSOName);
-						if (!File.Exists(FinalSOName) || Diff.TotalSeconds < -1 || Diff.TotalSeconds > 1)
-						{
-							Log.TraceInformation("\nCopying new .so file to jni folder...");
-							Directory.CreateDirectory(LibDir);
-							// copy the binary to the standard .so location
-							File.Copy(SourceSOName, FinalSOName, true);
-						}
+                    // check to see if libUE4.so is newer than last time we copied
+                    TimeSpan Diff = File.GetLastWriteTimeUtc(FinalSOName) - File.GetLastWriteTimeUtc(SourceSOName);
+                    if (!File.Exists(FinalSOName) || Diff.TotalSeconds < -1 || Diff.TotalSeconds > 1)
+                    {
+                        Directory.CreateDirectory(Path.GetDirectoryName(FinalSOName));
+                        File.Copy(SourceSOName, FinalSOName, true);
+                    }
+                }
 
-						FinalNdkBuildABICommand += GetNDKArch(Arch) + " ";
-					}
-					else
-					{
-						// if no NDK, we don't need any of the debugger stuff, so we just copy the .so to where it will end up
-						FinalSOName = UE4BuildPath + "/libs/" + GetNDKArch(Arch) + "/libUE4.so";
+                // remove any read only flags
+                FileInfo DestFileInfo = new FileInfo(FinalSOName);
+                DestFileInfo.Attributes = DestFileInfo.Attributes & ~FileAttributes.ReadOnly;
+                File.SetLastWriteTimeUtc(FinalSOName, File.GetLastWriteTimeUtc(SourceSOName));
 
-						// check to see if libUE4.so is newer than last time we copied
-						TimeSpan Diff = File.GetLastWriteTimeUtc(FinalSOName) - File.GetLastWriteTimeUtc(SourceSOName);
-						if (!File.Exists(FinalSOName) || Diff.TotalSeconds < -1 || Diff.TotalSeconds > 1)
-						{
-							Directory.CreateDirectory(Path.GetDirectoryName(FinalSOName));
-							File.Copy(SourceSOName, FinalSOName, true);
-						}
-					}
+                // if we need to run ndk-build, do it now
+                if(HasNDKPath)
+                {
+                    string LibSOName = UE4BuildPath + "/libs/" + GetNDKArch(Arch) + "/libUE4.so";
+                    // always delete libs up to this point so fat binaries and incremental builds work together (otherwise we might end up with multiple
+                    // so files in an apk that doesn't want them)
+                    // note that we don't want to delete all libs, just the ones we copied
+                    TimeSpan Diff = File.GetLastWriteTimeUtc(LibSOName) - File.GetLastWriteTimeUtc(FinalSOName);
+                    if (!File.Exists(LibSOName) || Diff.TotalSeconds < -1 || Diff.TotalSeconds > 1)
+                    {
+                        foreach (string Lib in Directory.EnumerateFiles(UE4BuildPath + "/libs", "libUE4*.so", SearchOption.AllDirectories))
+                        {
+                            File.Delete(Lib);
+                        }
 
-					// remove any read only flags
-					FileInfo DestFileInfo = new FileInfo(FinalSOName);
-					DestFileInfo.Attributes = DestFileInfo.Attributes & ~FileAttributes.ReadOnly;
-					File.SetLastWriteTimeUtc(FinalSOName, File.GetLastWriteTimeUtc(SourceSOName));
+                        string CommandLine = "APP_ABI=\"" + GetNDKArch(Arch) + " " + "\"";
+                        if (!bForDistribution)
+                        {
+                            CommandLine += " NDK_DEBUG=1";
+                        }
+                        RunCommandLineProgramAndThrowOnError(UE4BuildPath, NDKBuildPath, CommandLine, "Preparing native code for debugging...", true);
 
-					// now do final stuff per apk (or after all .so's for a shared .apk)
-					if (bMakeSeparateApks || (ArchIndex == Arches.Length - 1 && GPUArchIndex == GPUArchitectures.Length - 1))
-					{
-						// if we need to run ndk-build, do it now (if making a shared .apk, we need to wait until all .libs exist)
-						if (!string.IsNullOrEmpty(FinalNdkBuildABICommand))
-						{
-							string LibSOName = UE4BuildPath + "/libs/" + GetNDKArch(Arch) + "/libUE4.so";
-							// always delete libs up to this point so fat binaries and incremental builds work together (otherwise we might end up with multiple
-							// so files in an apk that doesn't want them)
-							// note that we don't want to delete all libs, just the ones we copied
-							TimeSpan Diff = File.GetLastWriteTimeUtc(LibSOName) - File.GetLastWriteTimeUtc(FinalSOName);
-							if (!File.Exists(LibSOName) || Diff.TotalSeconds < -1 || Diff.TotalSeconds > 1)
-							{
-								foreach (string Lib in Directory.EnumerateFiles(UE4BuildPath + "/libs", "libUE4*.so", SearchOption.AllDirectories))
-								{
-									File.Delete(Lib);
-								}
-
-								string CommandLine = "APP_ABI=\"" + FinalNdkBuildABICommand + "\"";
-								if (!bForDistribution)
-								{
-									CommandLine += " NDK_DEBUG=1";
-								}
-								RunCommandLineProgramAndThrowOnError(UE4BuildPath, NDKBuildPath, CommandLine, "Preparing native code for debugging...", true);
-
-								File.SetLastWriteTimeUtc(LibSOName, File.GetLastWriteTimeUtc(FinalSOName));
-							}
-							// next loop we don't want to redo these ABIs
-							FinalNdkBuildABICommand = "";
-						}
+                        File.SetLastWriteTimeUtc(LibSOName, File.GetLastWriteTimeUtc(FinalSOName));
+                    }
+                }
 
-						// after ndk-build is called, we can now copy in the stl .so (ndk-build deletes old files)
-						// copy libgnustl_shared.so to library (use 4.8 if possible, otherwise 4.6)
-						if (bMakeSeparateApks)
-						{
-							CopySTL(UE4BuildPath, Arch);
-							CopyGfxDebugger(UE4BuildPath, Arch);
-						}
-						else
-						{
-							foreach (string InnerArch in Arches)
-							{
-								CopySTL(UE4BuildPath, InnerArch);
-								CopyGfxDebugger(UE4BuildPath, InnerArch);
-							}
-						}
+                // after ndk-build is called, we can now copy in the stl .so (ndk-build deletes old files)
+                // copy libgnustl_shared.so to library (use 4.8 if possible, otherwise 4.6)
+                CopySTL(UE4BuildPath, Arch);
+                CopyGfxDebugger(UE4BuildPath, Arch);
 
-						Log.TraceInformation("\n===={0}====PERFORMING FINAL APK PACKAGE OPERATION================================================", DateTime.Now.ToString());
+                Log.TraceInformation("\n===={0}====PERFORMING FINAL APK PACKAGE OPERATION================================================", DateTime.Now.ToString());
 
-						string AntBuildType = "debug";
-						string AntOutputSuffix = "-debug";
-						if (bForDistribution)
-						{
-							// this will write out ant.properties with info needed to sign a distribution build
-							PrepareToSignApk(UE4BuildPath);
-							AntBuildType = "release";
-							AntOutputSuffix = "-release";
-						}
+                string AntBuildType = "debug";
+                string AntOutputSuffix = "-debug";
+                if (bForDistribution)
+                {
+                    // this will write out ant.properties with info needed to sign a distribution build
+                    PrepareToSignApk(UE4BuildPath);
+                    AntBuildType = "release";
+                    AntOutputSuffix = "-release";
+                }
 
-						// Use ant to build the .apk file
-						if (Utils.IsRunningOnMono)
-						{
-							RunCommandLineProgramAndThrowOnError(UE4BuildPath, "/bin/sh", "-c '\"" + GetAntPath() + "\" -quiet " + AntBuildType + "'", "Making .apk with Ant... (note: it's safe to ignore javac obsolete warnings)");
-						}
-						else
-						{
-							RunCommandLineProgramAndThrowOnError(UE4BuildPath, "cmd.exe", "/c \"" + GetAntPath() + "\" " + AntBuildType, "Making .apk with Ant... (note: it's safe to ignore javac obsolete warnings)");
-						}
+                // Use ant to build the .apk file
+                if (Utils.IsRunningOnMono)
+                {
+                    RunCommandLineProgramAndThrowOnError(UE4BuildPath, "/bin/sh", "-c '\"" + GetAntPath() + "\" -quiet " + AntBuildType + "'", "Making .apk with Ant... (note: it's safe to ignore javac obsolete warnings)");
+                }
+                else
+                {
+                    RunCommandLineProgramAndThrowOnError(UE4BuildPath, "cmd.exe", "/c \"" + GetAntPath() + "\" " + AntBuildType, "Making .apk with Ant... (note: it's safe to ignore javac obsolete warnings)");
+                }
 
-						// make sure destination exists
-						Directory.CreateDirectory(Path.GetDirectoryName(DestApkName));
+                // make sure destination exists
+                Directory.CreateDirectory(Path.GetDirectoryName(DestApkName));
+
+                // now copy to the final location
+                File.Copy(UE4BuildPath + "/bin/" + ProjectName + AntOutputSuffix + ".apk", DestApkName, true);
+
+            }
 
-						// now copy to the final location
-						File.Copy(UE4BuildPath + "/bin/" + ProjectName + AntOutputSuffix + ".apk", DestApkName, true);
-					}
-				}
-			}
 		}
 
 		private void PrepareToSignApk(string BuildPath)
