Safer concatenation of paths to try to avoid a bad path reported by a user.

diff --git a/Engine/Source/Programs/UnrealBuildTool/Windows/VCEnvironment.cs b/Engine/Source/Programs/UnrealBuildTool/Windows/VCEnvironment.cs
--- a/Engine/Source/Programs/UnrealBuildTool/Windows/VCEnvironment.cs
+++ b/Engine/Source/Programs/UnrealBuildTool/Windows/VCEnvironment.cs
@@ -75,8 +75,15 @@ namespace UnrealBuildTool
 			LibraryLinkerPath    = GetLibraryLinkerToolPath   (PlatformVSToolPath);
 			ResourceCompilerPath = GetResourceCompilerToolPath(Platform);
 
-			var VCVarsBatchFile = Path.Combine(BaseVSToolPath, 
-				(Platform == CPPTargetPlatform.Win64 || Platform == CPPTargetPlatform.UWP) ? "../../VC/bin/x86_amd64/vcvarsx86_amd64.bat" : "vsvars32.bat");
+			// We ensure an extra trailing slash because of a user getting an odd error where the paths seemed to get concatenated wrongly:
+			//
+			// C:\Programme\Microsoft Visual Studio 12.0\Common7\Tools../../VC/bin/x86_amd64/vcvarsx86_amd64.bat
+			//
+			// https://answers.unrealengine.com/questions/233640/unable-to-create-project-files-for-48-preview-3.html
+			//
+			bool   bUse64BitCompiler             = Platform == CPPTargetPlatform.Win64 || Platform == CPPTargetPlatform.UWP;
+			string BaseToolPathWithTrailingSlash = BaseVSToolPath.TrimEnd(Path.DirectorySeparatorChar) + Path.DirectorySeparatorChar;
+			string VCVarsBatchFile               = Path.Combine(BaseToolPathWithTrailingSlash, bUse64BitCompiler ? @"..\..\VC\bin\x86_amd64\vcvarsx86_amd64.bat" : "vsvars32.bat");
 			if (Platform == CPPTargetPlatform.UWP && UWPPlatform.bBuildForStore)
 			{
 				Utils.SetEnvironmentVariablesFromBatchFile(VCVarsBatchFile, "store");
