Fix UBT not allowing plugins with no modules

Merging 2570550

//depot/UE4-UT/...

to //depot/UE4/...

diff --git a/Engine/Source/Programs/UnrealBuildTool/Configuration/UEBuildTarget.cs b/Engine/Source/Programs/UnrealBuildTool/Configuration/UEBuildTarget.cs
--- a/Engine/Source/Programs/UnrealBuildTool/Configuration/UEBuildTarget.cs
+++ b/Engine/Source/Programs/UnrealBuildTool/Configuration/UEBuildTarget.cs
@@ -2263,14 +2263,17 @@ namespace UnrealBuildTool
 
 				// Add all the enabled plugins to the precompiled module list. Plugins are always precompiled, even if bPrecompile is not set, so we should precompile their dependencies.
 				foreach(PluginInfo Plugin in BuildPlugins)
-				{
-					foreach(ModuleDescriptor Module in Plugin.Descriptor.Modules)
-					{
-						if (Module.IsCompiledInConfiguration(Platform, TargetType, UEBuildConfiguration.bBuildDeveloperTools, UEBuildConfiguration.bBuildEditor))
-						{
-							PrecompiledModuleNames.Add(Module.Name);
-						}
-					}
+                {
+                    if (Plugin.Descriptor.Modules != null)
+                    {
+                        foreach (ModuleDescriptor Module in Plugin.Descriptor.Modules)
+                        {
+                            if (Module.IsCompiledInConfiguration(Platform, TargetType, UEBuildConfiguration.bBuildDeveloperTools, UEBuildConfiguration.bBuildEditor))
+                            {
+                                PrecompiledModuleNames.Add(Module.Name);
+                            }
+                        }
+                    }
 				}
 
 				// When running in Rocket, all engine modules have to be precompiled, but are precompiled using a different target (UE4Game) with a separate list of precompiled module names.
