[INTEGRATE] Change 2571827 by [REDACTED_PERSON][REDACTED_HANDLE] on 2015/05/30 17:48:49

	Fix missing resources in UE4Game and games compiled from static libraries in the binary build.

diff --git a/Engine/Source/Programs/UnrealBuildTool/Configuration/UEBuildBinary.cs b/Engine/Source/Programs/UnrealBuildTool/Configuration/UEBuildBinary.cs
--- a/Engine/Source/Programs/UnrealBuildTool/Configuration/UEBuildBinary.cs
+++ b/Engine/Source/Programs/UnrealBuildTool/Configuration/UEBuildBinary.cs
@@ -587,6 +587,13 @@ namespace UnrealBuildTool
 				}
 				DependentLinkEnvironment.Config.AdditionalLibraries.Add(LibraryFileName);
 			}
+			
+			// If we're linking against static library containing the launch module on windows, we need to add the compiled resource separately. We can't link it through the static library.
+			if(Config.Type == UEBuildBinaryType.StaticLibrary && ModuleNames.Contains("Launch") && (Target.Platform == UnrealTargetPlatform.Win32 || Target.Platform == UnrealTargetPlatform.Win64))
+			{
+				string ResourceFilePath = Path.Combine(Config.IntermediateDirectory, "Launch", "PCLaunch.rc.res");
+				DependentLinkEnvironment.InputFiles.Add(FileItem.GetItemByPath(ResourceFilePath));
+			}
 		}
 
 		/// <summary>
@@ -640,6 +647,13 @@ namespace UnrealBuildTool
 				}
 			}
 
+			// Add the compiled resource file if we're building a static library containing the launch module on Windows
+			if(Config.Type == UEBuildBinaryType.StaticLibrary && ModuleNames.Contains("Launch") && (Target.Platform == UnrealTargetPlatform.Win32 || Target.Platform == UnrealTargetPlatform.Win64))
+			{
+				string ResourceFilePath = Path.Combine(Config.IntermediateDirectory, "Launch", "PCLaunch.rc.res");
+				Receipt.AddBuildProduct(ResourceFilePath, BuildProductType.StaticLibrary);
+			}
+			
 			// Add runtime dependencies for all the modules in this binary, and build up a list of all the referenced modules
 			Dictionary<string, UEBuildModule> ReferencedModules = new Dictionary<string,UEBuildModule>();
 			List<UEBuildModule> OrderedModules = new List<UEBuildModule>();
