Unification of environment variable harvesting, GetExecutingAssembly* functions, GetShortPathName and CaselessDictionary from UBT into DotNETUtilities.
Fixing up of existing code which used these facilities - this fixes the 'You are attempting to compile on a machine that does not have a supported compiler!' UAT error on machines with really long PATH variables, and exceptions in envvars with non-ASCII characters.

diff --git a/Engine/Source/Programs/AutomationTool/Automation.cs b/Engine/Source/Programs/AutomationTool/Automation.cs
--- a/Engine/Source/Programs/AutomationTool/Automation.cs
+++ b/Engine/Source/Programs/AutomationTool/Automation.cs
@@ -7,6 +7,7 @@ using System.IO;
 using System.Reflection;
 using System.Diagnostics;
 using UnrealBuildTool;
+using Tools.DotNETCommon.CaselessDictionary;
 
 namespace AutomationTool
 {
diff --git a/Engine/Source/Programs/AutomationTool/CommandUtils.cs b/Engine/Source/Programs/AutomationTool/CommandUtils.cs
--- a/Engine/Source/Programs/AutomationTool/CommandUtils.cs
+++ b/Engine/Source/Programs/AutomationTool/CommandUtils.cs
@@ -13,6 +13,8 @@ using Tools.DotNETCommon.XmlHandler;
 using UnrealBuildTool;
 using System.Runtime.CompilerServices;
 using System.Linq;
+using Tools.DotNETCommon.CaselessDictionary;
+using Tools.DotNETCommon.HarvestEnvVars;
 
 namespace AutomationTool
 {
@@ -1717,53 +1719,26 @@ namespace AutomationTool
 		/// <returns>Dictionary of environment variables set by the batch file.</returns>
 		public static CaselessDictionary<string> GetEnvironmentVariablesFromBatchFile(string BatchFileName, bool AlsoSet = false)
 		{
-			if (File.Exists(BatchFileName))
+			CaselessDictionary<string> Result;
+			try
 			{
-				// Create a wrapper batch file that echoes environment variables to a text file
-				var EnvOutputFileName = CommandUtils.CombinePaths(Path.GetTempPath(), "HarvestEnvVars.txt");
-				var EnvReaderBatchFileName = CommandUtils.CombinePaths(Path.GetTempPath(), "HarvestEnvVars.bat");
-
-				{
-					var EnvReaderBatchFileContent = new List<string>();
-
-					// Run 'vcvars32.bat' (or similar x64 version) to set environment variables
-					EnvReaderBatchFileContent.Add(String.Format("call \"{0}\"", BatchFileName));
-
-					// Pipe all environment variables to a file where we can read them in
-					EnvReaderBatchFileContent.Add(String.Format("set >\"{0}\"", EnvOutputFileName));
-
-					File.WriteAllLines(EnvReaderBatchFileName, EnvReaderBatchFileContent);
-				}
-
-				InternalUtils.SafeDeleteFile(EnvOutputFileName);
-				CommandUtils.Run(EnvReaderBatchFileName, "");
+				Result = HarvestEnvVars.HarvestEnvVarsFromBatchFile(BatchFileName, "", HarvestEnvVars.EPathOverride.User);
+			}
+			catch (Exception Ex)
+			{
+				throw new AutomationException("Failed to harvest environment variables", Ex);
+			}
 
-				var Result = new CaselessDictionary<string>();
-				// Load environment variables
-				var EnvStringsFromFile = File.ReadAllLines(EnvOutputFileName);
-				foreach (var EnvString in EnvStringsFromFile)
+			if (AlsoSet)
+			{
+				// Set the environment variables
+				foreach (var Envvar in Result)
 				{
-					// Parse the environment variable name and value from the string ("name=value")
-					int EqualSignPos = EnvString.IndexOf('=');
-					var EnvironmentVariableName = EnvString.Substring(0, EqualSignPos);
-					var EnvironmentVariableValue = EnvString.Substring(EqualSignPos + 1);
-
-					if (Environment.GetEnvironmentVariable(EnvironmentVariableName) != EnvironmentVariableValue)
-					{
-						Result.Add(EnvironmentVariableName, EnvironmentVariableValue);
-					}
-					if (AlsoSet)
-					{
-						// Set the environment variable
-						Environment.SetEnvironmentVariable(EnvironmentVariableName, EnvironmentVariableValue);
-					}
+					Environment.SetEnvironmentVariable(Envvar.Key, Envvar.Value);
 				}
-				return Result;
-			}
-			else
-			{
-				throw new AutomationException("BatchFile {0} does not exist!", BatchFileName);
 			}
+
+			return Result;
 		}
 
 		#endregion
diff --git a/Engine/Source/Programs/AutomationTool/P4Utils.cs b/Engine/Source/Programs/AutomationTool/P4Utils.cs
--- a/Engine/Source/Programs/AutomationTool/P4Utils.cs
+++ b/Engine/Source/Programs/AutomationTool/P4Utils.cs
@@ -7,6 +7,7 @@ using System.IO;
 using System.ComponentModel;
 using System.Diagnostics;
 using System.Text.RegularExpressions;
+using Tools.DotNETCommon.CaselessDictionary;
 
 namespace AutomationTool
 {
diff --git a/Engine/Source/Programs/AutomationTool/ProjectUtils.cs b/Engine/Source/Programs/AutomationTool/ProjectUtils.cs
--- a/Engine/Source/Programs/AutomationTool/ProjectUtils.cs
+++ b/Engine/Source/Programs/AutomationTool/ProjectUtils.cs
@@ -7,6 +7,8 @@ using System.Text;
 using System.IO;
 using UnrealBuildTool;
 using System.Diagnostics;
+using Tools.DotNETCommon.CaselessDictionary;
+using Tools.DotNETCommon.ExecutingAssembly;
 
 namespace AutomationTool
 {
@@ -192,10 +194,10 @@ namespace AutomationTool
 			string oldCWD = Directory.GetCurrentDirectory();
 			if (BuildConfiguration.RelativeEnginePath == "../../Engine/")
 			{
-				string EngineSourceDirectory = Path.Combine(UnrealBuildTool.Utils.GetExecutingAssemblyDirectory(), "..", "..", "..", "Engine", "Source");
+				string EngineSourceDirectory = Path.Combine(ExecutingAssembly.GetDirectory(), "..", "..", "..", "Engine", "Source");
 				if (!Directory.Exists(EngineSourceDirectory)) // only set the directory if it exists, this should only happen if we are launching the editor from an artist sync
 				{
-					EngineSourceDirectory = Path.Combine(UnrealBuildTool.Utils.GetExecutingAssemblyDirectory(), "..", "..", "..", "Engine", "Binaries");
+					EngineSourceDirectory = Path.Combine(ExecutingAssembly.GetDirectory(), "..", "..", "..", "Engine", "Binaries");
 				}
 				Directory.SetCurrentDirectory(EngineSourceDirectory);
 			}
diff --git a/Engine/Source/Programs/AutomationTool/ScriptCompiler.cs b/Engine/Source/Programs/AutomationTool/ScriptCompiler.cs
--- a/Engine/Source/Programs/AutomationTool/ScriptCompiler.cs
+++ b/Engine/Source/Programs/AutomationTool/ScriptCompiler.cs
@@ -10,6 +10,7 @@ using Microsoft.Win32;
 using System.Reflection;
 using System.Diagnostics;
 using UnrealBuildTool;
+using Tools.DotNETCommon.CaselessDictionary;
 
 namespace AutomationTool
 {
diff --git a/Engine/Source/Programs/AutomationTool/Utils.cs b/Engine/Source/Programs/AutomationTool/Utils.cs
--- a/Engine/Source/Programs/AutomationTool/Utils.cs
+++ b/Engine/Source/Programs/AutomationTool/Utils.cs
@@ -1083,30 +1083,4 @@ namespace AutomationTool
 	}
 
 	#endregion
-
-	#region Case insensitive dictionary
-
-	/// <summary>
-	/// Equivalent of case insensitve Dictionary<string, T>
-	/// </summary>
-	/// <typeparam name="T"></typeparam>
-	public class CaselessDictionary<T> : Dictionary<string, T>
-	{
-		public CaselessDictionary()
-			: base(StringComparer.InvariantCultureIgnoreCase)
-		{
-		}
-
-		public CaselessDictionary(int Capacity)
-			: base(Capacity, StringComparer.InvariantCultureIgnoreCase)
-		{
-		}
-
-		public CaselessDictionary(IDictionary<string, T> Dict)
-			: base(Dict, StringComparer.InvariantCultureIgnoreCase)
-		{
-		}
-	}
-
-	#endregion
 }
diff --git a/Engine/Source/Programs/DotNETCommon/DotNETUtilities/CaselessDictionary.cs b/Engine/Source/Programs/DotNETCommon/DotNETUtilities/CaselessDictionary.cs
new file mode 100644
--- /dev/null
+++ b/Engine/Source/Programs/DotNETCommon/DotNETUtilities/CaselessDictionary.cs
@@ -0,0 +1,29 @@
+﻿// Copyright 1998-2015 Epic Games, Inc. All Rights Reserved.
+
+using System;
+using System.Collections.Generic;
+
+namespace Tools.DotNETCommon.CaselessDictionary
+{
+	/// <summary>
+	/// Equivalent of case insensitive Dictionary<string, T>
+	/// </summary>
+	/// <typeparam name="T"></typeparam>
+	public class CaselessDictionary<T> : Dictionary<string, T>
+	{
+		public CaselessDictionary()
+			: base(StringComparer.InvariantCultureIgnoreCase)
+		{
+		}
+
+		public CaselessDictionary(int Capacity)
+			: base(Capacity, StringComparer.InvariantCultureIgnoreCase)
+		{
+		}
+
+		public CaselessDictionary(IDictionary<string, T> Dict)
+			: base(Dict, StringComparer.InvariantCultureIgnoreCase)
+		{
+		}
+	}
+}
diff --git a/Engine/Source/Programs/DotNETCommon/DotNETUtilities/DotNETUtilities.csproj b/Engine/Source/Programs/DotNETCommon/DotNETUtilities/DotNETUtilities.csproj
--- a/Engine/Source/Programs/DotNETCommon/DotNETUtilities/DotNETUtilities.csproj
+++ b/Engine/Source/Programs/DotNETCommon/DotNETUtilities/DotNETUtilities.csproj
@@ -49,6 +49,10 @@
     <Compile Include="..\MetaData.cs">
       <Link>Properties\MetaData.cs</Link>
     </Compile>
+    <Compile Include="CaselessDictionary.cs" />
+    <Compile Include="ExecutingAssembly.cs" />
+    <Compile Include="FileSystem.cs" />
+    <Compile Include="HarvestEnvVars.cs" />
     <Compile Include="LaunchProcess.cs" />
     <Compile Include="Properties\AssemblyInfo.cs" />
     <Compile Include="SimpleWebRequest.cs" />
diff --git a/Engine/Source/Programs/DotNETCommon/DotNETUtilities/ExecutingAssembly.cs b/Engine/Source/Programs/DotNETCommon/DotNETUtilities/ExecutingAssembly.cs
new file mode 100644
--- /dev/null
+++ b/Engine/Source/Programs/DotNETCommon/DotNETUtilities/ExecutingAssembly.cs
@@ -0,0 +1,32 @@
+﻿// Copyright 1998-2015 Epic Games, Inc. All Rights Reserved.
+
+using System;
+using System.IO;
+
+namespace Tools.DotNETCommon.ExecutingAssembly
+{
+	public static class ExecutingAssembly
+	{
+		/// <summary>
+		/// Gets the executing assembly path (including filename).
+		/// This method is using Assembly.CodeBase property to properly resolve original
+		/// assembly path in case shadow copying is enabled.
+		/// </summary>
+		/// <returns>Absolute path to the executing assembly including the assembly filename.</returns>
+		public static string GetFilename()
+		{
+			return new Uri(System.Reflection.Assembly.GetExecutingAssembly().CodeBase).LocalPath;
+		}
+
+		/// <summary>
+		/// Gets the executing assembly directory.
+		/// This method is using Assembly.CodeBase property to properly resolve original
+		/// assembly directory in case shadow copying is enabled.
+		/// </summary>
+		/// <returns>Absolute path to the directory containing the executing assembly.</returns>
+		public static string GetDirectory()
+		{
+			return Path.GetDirectoryName(GetFilename());
+		}
+	}
+}
diff --git a/Engine/Source/Programs/DotNETCommon/DotNETUtilities/FileSystem.cs b/Engine/Source/Programs/DotNETCommon/DotNETUtilities/FileSystem.cs
new file mode 100644
--- /dev/null
+++ b/Engine/Source/Programs/DotNETCommon/DotNETUtilities/FileSystem.cs
@@ -0,0 +1,32 @@
+﻿// Copyright 1998-2015 Epic Games, Inc. All Rights Reserved.
+
+using System;
+using System.Runtime.InteropServices;
+using System.Text;
+
+namespace Tools.DotNETCommon.FileSystem
+{
+	public static class FileSystem
+	{
+		[DllImport("kernel32.dll", SetLastError=true)]
+		private static extern int GetShortPathName(string pathName, StringBuilder shortName, int cbShortName);
+
+		public static string GetShortPathName(string Path)
+		{
+			int BufferSize = GetShortPathName(Path, null, 0);
+			if (BufferSize == 0)
+			{
+				throw new Exception(string.Format("Unable to convert path {0} to 8.3 format", Path));
+			}
+
+			var Builder = new StringBuilder(BufferSize);
+			int ConversionResult = GetShortPathName(Path, Builder, BufferSize);
+			if (ConversionResult == 0)
+			{
+				throw new Exception(string.Format("Unable to convert path {0} to 8.3 format", Path));
+			}
+
+			return Builder.ToString();
+		}
+	}
+}
diff --git a/Engine/Source/Programs/DotNETCommon/DotNETUtilities/HarvestEnvVars.cs b/Engine/Source/Programs/DotNETCommon/DotNETUtilities/HarvestEnvVars.cs
new file mode 100644
--- /dev/null
+++ b/Engine/Source/Programs/DotNETCommon/DotNETUtilities/HarvestEnvVars.cs
@@ -0,0 +1,188 @@
+﻿// Copyright 1998-2015 Epic Games, Inc. All Rights Reserved.
+
+using System;
+using System.Collections.Generic;
+using System.Diagnostics;
+using System.IO;
+using System.Text;
+using System.Xml;
+using System.Xml.Serialization;
+
+namespace Tools.DotNETCommonPrivate.HarvestEnvVars
+{
+	// Struct which represents a key-value.  Conceptually private, but the serializer doesn't like it to be actually private.
+	[Serializable]
+	[DebuggerDisplay("\\{{Key}={Value}\\}")]
+	public struct EnvVar
+	{
+		#pragma warning disable 0649
+
+		[XmlAttribute("Key")]
+		public string Key;
+
+		[XmlAttribute("Value")]
+		public string Value;
+
+		#pragma warning restore 0649
+	}
+}
+
+namespace Tools.DotNETCommon.HarvestEnvVars
+{
+	public static class HarvestEnvVars
+	{
+		public enum EPathOverride
+		{
+			None,
+			User,
+		}
+
+		private static XmlSerializer EnvVarListSerializer = XmlSerializer.FromTypes(new Type[]{ typeof(List<Tools.DotNETCommonPrivate.HarvestEnvVars.EnvVar>) })[0];
+
+		public static CaselessDictionary.CaselessDictionary<string> HarvestEnvVarsFromBatchFile(string BatchFileName, string BatchFileParameters, EPathOverride PathOverride)
+		{
+			if (!File.Exists(BatchFileName))
+			{
+				throw new Exception(string.Format("Tools.DotNETCommon.HarvestEnvVars.HarvestEnvVars.HarvestEnvVarsFromBatchFile: BatchFile {0} does not exist!", BatchFileName));
+			}
+
+			// Create a wrapper batch file that echoes environment variables to a text file
+			string EnvOutputFileName      = Path.GetTempFileName();
+			string EnvReaderBatchFileName = EnvOutputFileName + ".bat";
+			try
+			{
+				var EnvVarsToXMLExePath = Path.Combine(ExecutingAssembly.ExecutingAssembly.GetDirectory(), "EnvVarsToXML.exe");
+
+				// Convert every path to short filenames to ensure we don't accidentally write out a non-ASCII batch file
+				var ShortBatchFileName       = FileSystem.FileSystem.GetShortPathName(BatchFileName);
+				var ShortEnvOutputFileName   = FileSystem.FileSystem.GetShortPathName(EnvOutputFileName);
+				var ShortEnvVarsToXMLExePath = FileSystem.FileSystem.GetShortPathName(EnvVarsToXMLExePath);
+
+				var EnvReaderBatchFileContent = new StringBuilder();
+
+				// Run 'vcvars32.bat' (or similar x64 version) to set environment variables
+				EnvReaderBatchFileContent.AppendFormat("call {0} {1}", ShortBatchFileName, BatchFileParameters).AppendLine();
+
+				// Pipe all environment variables to a file where we can read them in.
+				// We use a separate executable which runs after the batch file because we want to capture
+				// the environment after it has been set, and there's no easy way of doing this, and parsing
+				// the output of the set command is problematic when the vars contain non-ASCII characters.
+				EnvReaderBatchFileContent.AppendFormat("{0} {1}", ShortEnvVarsToXMLExePath, ShortEnvOutputFileName).AppendLine();
+
+				FileInfo TempFileInfo = new FileInfo(EnvReaderBatchFileName);
+				if (TempFileInfo.Exists)
+				{
+					TempFileInfo.IsReadOnly = false;
+					TempFileInfo.Delete();
+					TempFileInfo.Refresh();
+				}
+
+				Directory.CreateDirectory(Path.GetDirectoryName(EnvReaderBatchFileName));
+				File.WriteAllText(EnvReaderBatchFileName, EnvReaderBatchFileContent.ToString());
+			}
+			catch (Exception Ex)
+			{
+				throw new Exception(string.Format("Failed to create temporary batch file {0} to harvest environment variables (\"{1}\")", EnvReaderBatchFileName, Ex.Message), Ex);
+			}
+
+			// process needs to be disposed when done
+			using (var BatchFileProcess = new Process())
+			{
+				// Run the batch file using cmd.exe with the /U option, to force Unicode output. Many locales have non-ANSI characters in system paths.
+				var StartInfo = BatchFileProcess.StartInfo;
+				StartInfo.FileName        = Path.Combine(Environment.SystemDirectory, "cmd.exe");
+				StartInfo.Arguments       = String.Format("/U /C \"{0}\"", EnvReaderBatchFileName);
+				StartInfo.CreateNoWindow  = true;
+				StartInfo.UseShellExecute = false;
+
+				// @todo UWP: Win10-beta issue, sometimes the process hangs if standard I/O is redirected, and there's bytes
+				// in the buffer and there is no reader for those streams.  Loop back with a later Win10 build
+				bool bRedirectStreams = System.Environment.OSVersion.Platform != PlatformID.Win32NT || System.Environment.OSVersion.Version.Major < 10;
+				bRedirectStreams = false;
+				if (bRedirectStreams)
+				{
+					StartInfo.RedirectStandardOutput = true;
+					StartInfo.RedirectStandardError = true;
+					StartInfo.RedirectStandardInput = true;
+				}
+
+				// Override path variable if it can be particularly long
+				if (PathOverride == EPathOverride.User)
+				{
+					string NewPathVariable = Environment.GetEnvironmentVariable("PATH", EnvironmentVariableTarget.Machine) ?? "";
+					if (String.IsNullOrEmpty(NewPathVariable))
+					{
+						NewPathVariable = Environment.GetEnvironmentVariable("PATH", EnvironmentVariableTarget.User);
+					}
+					else
+					{
+						NewPathVariable += ";" + Environment.GetEnvironmentVariable("PATH", EnvironmentVariableTarget.User);
+					}
+					StartInfo.EnvironmentVariables["PATH"] = NewPathVariable;
+				}
+
+				// Try to launch the process, and produce a friendly error message if it fails.
+				try
+				{
+					// Start the process up and then wait for it to finish
+					BatchFileProcess.Start();
+					if (bRedirectStreams)
+					{
+						BatchFileProcess.StandardOutput.ReadToEnd();
+					}
+					BatchFileProcess.WaitForExit();
+				}
+				catch (Exception Ex)
+				{
+					throw new Exception(string.Format("Failed to start local process for action (\"{0}\"): {1} {2}", Ex.Message, StartInfo.FileName, StartInfo.Arguments), Ex);
+				}
+			}
+
+			// Accept chars which are technically not valid XML - they were written out by XmlSerializer anyway!
+			var Settings = new XmlReaderSettings();
+			Settings.CheckCharacters  = false;
+
+			List<Tools.DotNETCommonPrivate.HarvestEnvVars.EnvVar> EnvVars;
+			try
+			{
+				using (var Stream = new StreamReader(EnvOutputFileName))
+				using (var Reader = XmlReader.Create(Stream, Settings))
+				{
+					EnvVars = (List<Tools.DotNETCommonPrivate.HarvestEnvVars.EnvVar>)EnvVarListSerializer.Deserialize(Reader);
+				}
+			}
+			catch (Exception Ex)
+			{
+				throw new Exception(string.Format("Failed to read environment variables from XML file: {0}", EnvOutputFileName), Ex);
+			}
+
+			// Clean up the temporary files we created earlier on, so the temp directory doesn't fill up
+			// with these guys over time
+			try
+			{
+				File.Delete(EnvOutputFileName);
+			}
+			catch (Exception)
+			{
+				// Unable to delete the temporary file.  Not a big deal.
+			}
+
+			try
+			{
+				File.Delete(EnvReaderBatchFileName);
+			}
+			catch (Exception)
+			{
+				// Unable to delete the temporary file.  Not a big deal.
+			}
+
+			var Result = new CaselessDictionary.CaselessDictionary<string>();
+			foreach (var Var in EnvVars)
+			{
+				Result.Add(Var.Key, Var.Value);
+			}
+
+			return Result;
+		}
+	}
+}
diff --git a/Engine/Source/Programs/UnrealBuildTool/System/DynamicCompilation.cs b/Engine/Source/Programs/UnrealBuildTool/System/DynamicCompilation.cs
--- a/Engine/Source/Programs/UnrealBuildTool/System/DynamicCompilation.cs
+++ b/Engine/Source/Programs/UnrealBuildTool/System/DynamicCompilation.cs
@@ -7,13 +7,14 @@ using System.CodeDom.Compiler;
 using Microsoft.CSharp;
 using System.Reflection;
 using System.Diagnostics;
+using Tools.DotNETCommon.ExecutingAssembly;
 
 namespace UnrealBuildTool
 {
 	public class DynamicCompilation
 	{
 		/// File information for UnrealBuildTool.exe, cached at program start
-		private static FileInfo UBTExecutableFileInfo = new FileInfo( Utils.GetExecutingAssemblyLocation() );
+		private static FileInfo UBTExecutableFileInfo = new FileInfo( ExecutingAssembly.GetFilename() );
 
 		/*
 		 * Checks to see if the assembly needs compilation
diff --git a/Engine/Source/Programs/UnrealBuildTool/System/ProjectFileGenerator.cs b/Engine/Source/Programs/UnrealBuildTool/System/ProjectFileGenerator.cs
--- a/Engine/Source/Programs/UnrealBuildTool/System/ProjectFileGenerator.cs
+++ b/Engine/Source/Programs/UnrealBuildTool/System/ProjectFileGenerator.cs
@@ -8,6 +8,7 @@ using System.Reflection;
 using Microsoft.Win32;
 using System.Linq;
 using System.Diagnostics;
+using Tools.DotNETCommon.ExecutingAssembly;
 
 namespace UnrealBuildTool
 {
@@ -304,7 +305,7 @@ namespace UnrealBuildTool
 					{
 						if (BuildHostPlatform.Current.Platform == UnrealTargetPlatform.Mac)
 						{
-							MasterProjectRelativePath = Path.GetFullPath(Path.Combine(Utils.GetExecutingAssemblyDirectory(), "..", "..", "..", "Engine"));
+							MasterProjectRelativePath = Path.GetFullPath(Path.Combine(ExecutingAssembly.GetDirectory(), "..", "..", "..", "Engine"));
 							GameProjectName = "UE4Game";
 						}
 						if (!Directory.Exists(MasterProjectRelativePath + "/Source"))
diff --git a/Engine/Source/Programs/UnrealBuildTool/System/UProjectInfo.cs b/Engine/Source/Programs/UnrealBuildTool/System/UProjectInfo.cs
--- a/Engine/Source/Programs/UnrealBuildTool/System/UProjectInfo.cs
+++ b/Engine/Source/Programs/UnrealBuildTool/System/UProjectInfo.cs
@@ -8,6 +8,7 @@ using System.Diagnostics;
 using System.Threading;
 using System.Reflection;
 using System.Linq;
+using Tools.DotNETCommon.ExecutingAssembly;
 
 namespace UnrealBuildTool
 {
@@ -127,7 +128,7 @@ namespace UnrealBuildTool
 			List<string> DirectoriesToSearch = new List<string>();
 
 			// Find all the .uprojectdirs files contained in the root folder and add their entries to the search array
-			string RootDirectory = Path.Combine(Utils.GetExecutingAssemblyDirectory(), "..", "..", "..");
+			string RootDirectory = Path.Combine(ExecutingAssembly.GetDirectory(), "..", "..", "..");
 			string EngineSourceDirectory = Path.GetFullPath( Path.Combine(RootDirectory, "Engine", "Source") );
 
 			foreach (var File in Directory.EnumerateFiles(RootDirectory, "*.uprojectdirs", SearchOption.TopDirectoryOnly))
diff --git a/Engine/Source/Programs/UnrealBuildTool/System/UnrealBuildTool.cs b/Engine/Source/Programs/UnrealBuildTool/System/UnrealBuildTool.cs
--- a/Engine/Source/Programs/UnrealBuildTool/System/UnrealBuildTool.cs
+++ b/Engine/Source/Programs/UnrealBuildTool/System/UnrealBuildTool.cs
@@ -9,6 +9,7 @@ using System.Threading;
 using System.Reflection;
 using System.Linq;
 using System.Runtime.Serialization.Formatters.Binary;
+using Tools.DotNETCommon.ExecutingAssembly;
 
 namespace UnrealBuildTool
 {
@@ -194,7 +195,7 @@ namespace UnrealBuildTool
 
         static public string GetUBTPath()
         {
-            string UnrealBuildToolPath = Path.Combine(Utils.GetExecutingAssemblyDirectory(), "UnrealBuildTool.exe");
+            string UnrealBuildToolPath = Path.Combine(ExecutingAssembly.GetDirectory(), "UnrealBuildTool.exe");
             return UnrealBuildToolPath;
         }
 
@@ -636,14 +637,14 @@ namespace UnrealBuildTool
             // Change the working directory to be the Engine/Source folder. We are likely running from Engine/Binaries/DotNET
             // This is critical to be done early so any code that relies on the current directory being Engine/Source will work.
             // UEBuildConfiguration.PostReset is susceptible to this, so we must do this before configs are loaded.
-            string EngineSourceDirectory = Path.Combine(Utils.GetExecutingAssemblyDirectory(), "..", "..", "..", "Engine", "Source");
+            string EngineSourceDirectory = Path.Combine(ExecutingAssembly.GetDirectory(), "..", "..", "..", "Engine", "Source");
 
             [REDACTED_EMAIL]: This is a workaround for recompiling game code in editor
             // The working directory when launching is *not* what we would expect
             if (Directory.Exists(EngineSourceDirectory) == false)
             {
                 // We are assuming UBT always runs from <>/Engine/...
-                EngineSourceDirectory = Path.GetDirectoryName(Utils.GetExecutingAssemblyLocation());
+                EngineSourceDirectory = ExecutingAssembly.GetDirectory();
                 EngineSourceDirectory = EngineSourceDirectory.Replace("\\", "/");
                 Int32 EngineIdx = EngineSourceDirectory.IndexOf("/Engine/");
                 if (EngineIdx != 0)
@@ -939,7 +940,7 @@ namespace UnrealBuildTool
                 Mutex SingleInstanceMutex = null;
                 if (bUseMutex)
                 {
-                    int LocationHash = Utils.GetExecutingAssemblyLocation().GetHashCode();
+                    int LocationHash = ExecutingAssembly.GetFilename().GetHashCode();
 
                     String MutexName;
                     if (bAutoSDKOnly || bValidatePlatforms)
diff --git a/Engine/Source/Programs/UnrealBuildTool/System/Utils.cs b/Engine/Source/Programs/UnrealBuildTool/System/Utils.cs
--- a/Engine/Source/Programs/UnrealBuildTool/System/Utils.cs
+++ b/Engine/Source/Programs/UnrealBuildTool/System/Utils.cs
@@ -11,6 +11,8 @@ using System.Xml.Serialization;
 using System.Runtime.CompilerServices;
 using System.Runtime.InteropServices;
 using System.Linq;
+using Tools.DotNETCommon.CaselessDictionary;
+using Tools.DotNETCommon.HarvestEnvVars;
 
 namespace UnrealBuildTool
 {
@@ -203,40 +205,6 @@ namespace UnrealBuildTool
 			return Default;
 		}
 
-		[Serializable]
-		[DebuggerDisplay("\\{{Key}={Value}\\}")]
-		public struct EnvVar
-		{
-			[XmlAttribute("Key")]
-			public string Key;
-
-			[XmlAttribute("Value")]
-			public string Value;
-		}
-
-		private static XmlSerializer EnvVarListSerializer = XmlSerializer.FromTypes(new Type[]{ typeof(List<EnvVar>) })[0];
-
-		[DllImport("kernel32.dll", SetLastError=true)]
-		private static extern int GetShortPathName(string pathName, StringBuilder shortName, int cbShortName);
-
-		public static string GetShortPathName(string Path)
-		{
-			int BufferSize = GetShortPathName(Path, null, 0);
-			if (BufferSize == 0)
-			{
-				throw new BuildException("Unable to convert path {0} to 8.3 format", Path);
-			}
-
-			var Builder = new StringBuilder(BufferSize);
-			int ConversionResult = GetShortPathName(Path, Builder, BufferSize);
-			if (ConversionResult == 0)
-			{
-				throw new BuildException("Unable to convert path {0} to 8.3 format", Path);
-			}
-
-			return Builder.ToString();
-		}
-
 		/**
 		 * Sets the environment variables from the passed in batch file
 		 * 
@@ -246,143 +214,19 @@ namespace UnrealBuildTool
 		public static void SetEnvironmentVariablesFromBatchFile(string BatchFileName, string Parameters = "")
 		{
 			// @todo ubtmake: Experiment with changing this to run asynchronously at startup, and only blocking if accessed before the .bat file finishes
-			if( File.Exists( BatchFileName ) )
+			CaselessDictionary<string> EnvVars;
+			try
 			{
-				// Create a wrapper batch file that echoes environment variables to a text file
-                string EnvOutputFileName;
-                string EnvReaderBatchFileName;
-                try
-                {
-					EnvOutputFileName = Path.GetTempFileName();
-					EnvReaderBatchFileName = EnvOutputFileName + ".bat";
-
-					Log.TraceVerbose( "Creating .bat file {0} for harvesting environment variables.", EnvReaderBatchFileName );
-
-					var EnvReaderBatchFileContent = new List<string>();
-
-					var EnvVarsToXMLExePath = Path.Combine( GetExecutingAssemblyDirectory(), "EnvVarsToXML.exe" );
-
-					// Convert every path to short filenames to ensure we don't accidentally write out a non-ASCII batch file
-					var ShortBatchFileName       = GetShortPathName(BatchFileName);
-					var ShortEnvOutputFileName   = GetShortPathName(EnvOutputFileName);
-					var ShortEnvVarsToXMLExePath = GetShortPathName(EnvVarsToXMLExePath);
-
-					// Run 'vcvars32.bat' (or similar x64 version) to set environment variables
-					EnvReaderBatchFileContent.Add( String.Format("call \"{0}\" {1}", ShortBatchFileName, Parameters) );
-
-					// Pipe all environment variables to a file where we can read them in.
-					// We use a separate executable which runs after the batch file because we want to capture
-					// the environment after it has been set, and there's no easy way of doing this, and parsing
-					// the output of the set command is problematic when the vars contain non-ASCII characters.
-					EnvReaderBatchFileContent.Add( String.Format( "\"{0}\" \"{1}\"", ShortEnvVarsToXMLExePath, ShortEnvOutputFileName ) );
-
-					ResponseFile.Create( EnvReaderBatchFileName, EnvReaderBatchFileContent );
-                }
-                catch (Exception Ex)
-                {
-                    throw new BuildException(Ex, "Failed to create temporary batch file to harvest environment variables (\"{0}\")", Ex.Message);
-                }
-
-				Log.TraceVerbose( "Finished creating .bat file.  Environment variables will be written to {0}.", EnvOutputFileName );
-
-				// process needs to be disposed when done
-				using(var BatchFileProcess = new Process())
-				{
-					// Run the batch file using cmd.exe with the /U option, to force Unicode output. Many locales have non-ANSI characters in system paths.
-					var StartInfo = BatchFileProcess.StartInfo;
-
-					StartInfo.FileName = Path.Combine(Environment.SystemDirectory, "cmd.exe");
-					StartInfo.Arguments = String.Format("/U /C \"{0}\"", EnvReaderBatchFileName);
-					StartInfo.CreateNoWindow = true;
-					StartInfo.UseShellExecute = false;
-
-					// The engine adds a lot of DLL search paths to the PATH environment variable, and this gets propagated through to UBT. MSVC wants
-					// to add a lot more, so reset it to the system default before spawning the batch file, otherwise we can overflow the max length and fail.
-					string NewPathVariable = Environment.GetEnvironmentVariable("PATH", EnvironmentVariableTarget.Machine) ?? "";
-					if(String.IsNullOrEmpty(NewPathVariable))
-					{
-						NewPathVariable = Environment.GetEnvironmentVariable("PATH", EnvironmentVariableTarget.User);
-					}
-					else
-					{
-						NewPathVariable += ";" + Environment.GetEnvironmentVariable("PATH", EnvironmentVariableTarget.User);
-					}
-					StartInfo.EnvironmentVariables["PATH"] = NewPathVariable;
-
-					// @todo UWP: Win10-beta issue, sometimes the process hangs if standard I/O is redirected, and there's bytes
-					// in the buffer and there is no reader for those streams.  Loop back with a later Win10 build
-					if (System.Environment.OSVersion.Platform != PlatformID.Win32NT ||
-						System.Environment.OSVersion.Version.Major < 10)
-					{
-						StartInfo.RedirectStandardOutput = true;
-						StartInfo.RedirectStandardError = true;
-						StartInfo.RedirectStandardInput = true;
-					}
-
-					Log.TraceVerbose( "Launching {0} to harvest Visual Studio environment settings...", StartInfo.FileName );
-
-					// Try to launch the process, and produce a friendly error message if it fails.
-					try
-					{
-						// Start the process up and then wait for it to finish
-						BatchFileProcess.Start();
-						BatchFileProcess.WaitForExit();
-					}
-					catch(Exception ex)
-					{
-						throw new BuildException( ex, "Failed to start local process for action (\"{0}\"): {1} {2}", ex.Message, StartInfo.FileName, StartInfo.Arguments );
-					}
-
-					Log.TraceVerbose( "Finished launching {0}.", StartInfo.FileName );
-				}
-
-				// Accept chars which are technically not valid XML - they were written out by XmlSerializer anyway!
-				var Settings = new XmlReaderSettings();
-				Settings.CheckCharacters  = false;
-
-				List<EnvVar> EnvVars;
-				try
-				{
-					using (var Stream = new StreamReader(EnvOutputFileName))
-					using (var Reader = XmlReader.Create(Stream, Settings))
-					{
-						EnvVars = (List<EnvVar>)EnvVarListSerializer.Deserialize(Reader);
-					}
-				}
-				catch (Exception e)
-				{
-					throw new BuildException(e, "Failed to read environment variables from XML file: {0}", EnvOutputFileName);
-				}
-
-				foreach (var EnvVar in EnvVars)
-				{
-					Environment.SetEnvironmentVariable(EnvVar.Key, EnvVar.Value);
-				}
-
-				// Clean up the temporary files we created earlier on, so the temp directory doesn't fill up
-				// with these guys over time
-				try
-				{
-					File.Delete( EnvOutputFileName );
-				}
-				catch( Exception )
-				{
-					// Unable to delete the temporary file.  Not a big deal.
-					Log.TraceInformation( "Warning: Was not able to delete temporary file created by Unreal Build Tool: " + EnvOutputFileName );
-				}
-				try
-				{
-					File.Delete( EnvReaderBatchFileName );
-				}
-				catch( Exception )
-				{
-					// Unable to delete the temporary file.  Not a big deal.
-					Log.TraceInformation( "Warning: Was not able to delete temporary file created by Unreal Build Tool: " + EnvReaderBatchFileName );
-				}
+				EnvVars = HarvestEnvVars.HarvestEnvVarsFromBatchFile(BatchFileName, Parameters, HarvestEnvVars.EPathOverride.User);
 			}
-			else
+			catch (Exception Ex)
+			{
+				throw new BuildException(Ex, "Failed to harvest environment variables");
+			}
+
+			foreach (var EnvVar in EnvVars)
 			{
-				throw new BuildException("SetEnvironmentVariablesFromBatchFile: BatchFile {0} does not exist!", BatchFileName);
+				Environment.SetEnvironmentVariable(EnvVar.Key, EnvVar.Value);
 			}
 		}
 
@@ -947,28 +791,6 @@ namespace UnrealBuildTool
 			}
 		}
 
-		/// <summary>
-		/// Gets the executing assembly path (including filename).
-		/// This method is using Assembly.CodeBase property to properly resolve original
-		/// assembly path in case shadow copying is enabled.
-		/// </summary>
-		/// <returns>Absolute path to the executing assembly including the assembly filename.</returns>
-		public static string GetExecutingAssemblyLocation()
-		{
-			return new Uri(System.Reflection.Assembly.GetExecutingAssembly().CodeBase).LocalPath;
-		}
-
-		/// <summary>
-		/// Gets the executing assembly directory.
-		/// This method is using Assembly.CodeBase property to properly resolve original
-		/// assembly directory in case shadow copying is enabled.
-		/// </summary>
-		/// <returns>Absolute path to the directory containing the executing assembly.</returns>
-		public static string GetExecutingAssemblyDirectory()
-		{
-			return Path.GetDirectoryName(GetExecutingAssemblyLocation());
-		}
-
 		/// <summary>
 		/// Checks if given type implements given interface.
 		/// </summary>
diff --git a/Engine/Source/Programs/UnrealBuildTool/UnrealBuildTool.csproj b/Engine/Source/Programs/UnrealBuildTool/UnrealBuildTool.csproj
--- a/Engine/Source/Programs/UnrealBuildTool/UnrealBuildTool.csproj
+++ b/Engine/Source/Programs/UnrealBuildTool/UnrealBuildTool.csproj
@@ -278,6 +278,10 @@
     <None Include="app.config" />
   </ItemGroup>
   <ItemGroup>
+    <ProjectReference Include="..\DotNETCommon\DotNETUtilities\DotNETUtilities.csproj">
+      <Project>{5d7d66e8-8c76-4af9-b3ec-2ef03421d730}</Project>
+      <Name>DotNETUtilities</Name>
+    </ProjectReference>
     <ProjectReference Include="..\EnvVarsToXML\EnvVarsToXML\EnvVarsToXML.csproj">
       <Project>{66e423eb-ffa3-4ce4-bca3-75bdb6a8318d}</Project>
       <Name>EnvVarsToXML</Name>
diff --git a/Engine/Source/Programs/UnrealBuildTool/Utilities/XmlConfigLoader.cs b/Engine/Source/Programs/UnrealBuildTool/Utilities/XmlConfigLoader.cs
--- a/Engine/Source/Programs/UnrealBuildTool/Utilities/XmlConfigLoader.cs
+++ b/Engine/Source/Programs/UnrealBuildTool/Utilities/XmlConfigLoader.cs
@@ -13,6 +13,7 @@ using System.Threading;
 using System.Xml;
 using System.Xml.Linq;
 using System.Xml.Schema;
+using Tools.DotNETCommon.ExecutingAssembly;
 
 namespace UnrealBuildTool
 {
@@ -501,7 +502,7 @@ namespace UnrealBuildTool
 			 *	priority. Not defined classes and fields are left alone.
 			 */
 
-			var UE4EnginePath = new FileInfo(Path.Combine(Utils.GetExecutingAssemblyDirectory(), "..", "..")).FullName;
+			var UE4EnginePath = new FileInfo(Path.Combine(ExecutingAssembly.GetDirectory(), "..", "..")).FullName;
 
 			ConfigLocationHierarchy = new XmlConfigLocation[]
 			{
@@ -960,7 +961,7 @@ namespace UnrealBuildTool
 		/// <returns>Standard BuildConfiguration.xml schema path.</returns>
 		public static string GetXSDPath()
 		{
-			return new FileInfo(Path.Combine(Utils.GetExecutingAssemblyDirectory(), "..", "..", "Saved", "UnrealBuildTool", "BuildConfiguration.Schema.xsd")).FullName;
+			return new FileInfo(Path.Combine(ExecutingAssembly.GetDirectory(), "..", "..", "Saved", "UnrealBuildTool", "BuildConfiguration.Schema.xsd")).FullName;
 		}
 
 		private static XElement CreateXSDElementForField(FieldInfo Field)
