Always allow compiling game modules. Fix for user projects not being able to compile modules that aren't specifically added through the extra modules list.

diff --git a/Engine/Source/Programs/UnrealBuildTool/Configuration/UEBuildTarget.cs b/Engine/Source/Programs/UnrealBuildTool/Configuration/UEBuildTarget.cs
--- a/Engine/Source/Programs/UnrealBuildTool/Configuration/UEBuildTarget.cs
+++ b/Engine/Source/Programs/UnrealBuildTool/Configuration/UEBuildTarget.cs
@@ -2332,7 +2332,7 @@ namespace UnrealBuildTool
 					if(PrecompiledModule is UEBuildModuleCPP && !BoundModuleNames.Contains(PrecompiledModule.Name))
 					{
 						UEBuildBinaryType BinaryType = bCompileMonolithic? UEBuildBinaryType.StaticLibrary : UEBuildBinaryType.DynamicLinkLibrary;
-						UEBuildBinary Binary = AddBinaryForModule(PrecompiledModule.Name, BinaryType, bAllowCompilation: !bUsePrecompiled, bIsCrossTarget: false);
+						UEBuildBinary Binary = AddBinaryForModule(PrecompiledModule.Name, BinaryType, bAllowCompilation: !bUsePrecompiled || RulesCompiler.IsGameModule(PrecompiledModule.Name), bIsCrossTarget: false);
 						PrecompiledBinaries.Add(Binary);
 						BoundModuleNames.Add(PrecompiledModule.Name);
 					}
