FileItems cache cleared when the makefile if the deserialized makefile is invalid, to avoid it using stale data.
Some TraceWarnings changed to TraceVerbose.
Tidy up of some highly-indented logic.

diff --git a/Engine/Source/Programs/UnrealBuildTool/System/ExternalExecution.cs b/Engine/Source/Programs/UnrealBuildTool/System/ExternalExecution.cs
--- a/Engine/Source/Programs/UnrealBuildTool/System/ExternalExecution.cs
+++ b/Engine/Source/Programs/UnrealBuildTool/System/ExternalExecution.cs
@@ -332,8 +332,6 @@ namespace UnrealBuildTool
 		 * */
 		private static bool AreGeneratedCodeFilesOutOfDate(List<UHTModuleInfo> UObjectModules)
 		{
-			bool bIsOutOfDate = false;
-
 			// Get UnrealHeaderTool timestamp. If it's newer than generated headers, they need to be rebuilt too.
 			var HeaderToolTimestamp = CheckIfUnrealHeaderToolIsUpToDate();
 
@@ -368,81 +366,66 @@ namespace UnrealBuildTool
 				// Make sure we have an existing folder for generated code.  If not, then we definitely need to generate code!
 				var GeneratedCodeDirectory = Path.GetDirectoryName( Module.GeneratedCPPFilenameBase );
 				var TestDirectory = (FileSystemInfo)new DirectoryInfo(GeneratedCodeDirectory);
-				if( TestDirectory.Exists )
+				if( !TestDirectory.Exists )
 				{
-					// Grab our special "Timestamp" file that we saved after the last set of headers were generated
-					string TimestampFile = Path.Combine( GeneratedCodeDirectory, @"Timestamp" );
-					var SavedTimestampFileInfo = (FileSystemInfo)new FileInfo(TimestampFile);
-					if (SavedTimestampFileInfo.Exists)
-					{
-						// Make sure the last UHT run completed after UnrealHeaderTool.exe was compiled last, and after the CoreUObject headers were touched last.
-						var SavedTimestamp = SavedTimestampFileInfo.LastWriteTime;
-						if( SavedTimestamp.CompareTo(HeaderToolTimestamp) > 0 &&
-							SavedTimestamp.CompareTo(CoreGeneratedTimestamp) > 0 )
-						{
-							// Iterate over our UObjects headers and figure out if any of them have changed
-							var AllUObjectHeaders = new List<FileItem>();
-							AllUObjectHeaders.AddRange( Module.PublicUObjectClassesHeaders );
-							AllUObjectHeaders.AddRange( Module.PublicUObjectHeaders );
-							AllUObjectHeaders.AddRange( Module.PrivateUObjectHeaders );
-							foreach( var HeaderFile in AllUObjectHeaders )
-							{
-								var HeaderFileTimestamp = HeaderFile.Info.LastWriteTime;
-
-								// Has the source header changed since we last generated headers successfully?
-								if( SavedTimestamp.CompareTo( HeaderFileTimestamp ) < 0 )
-								{
-									bIsOutOfDate = true;
-									break;
-								}
+					// Generated code directory is missing entirely!
+					return true;
+				}
 
-								// When we're running in assembler mode, outdatedness cannot be inferred by checking the directory timestamp
-								// of the source headers.  We don't care if source files were added or removed in this mode, because we're only
-								// able to process the known UObject headers that are in the Makefile.  If UObject header files are added/removed,
-								// we expect the user to re-run GenerateProjectFiles which will force UBTMakefile outdatedness.
-								// @todo ubtmake: Possibly, we should never be doing this check these days.
-								if( UnrealBuildTool.IsGatheringBuild || !UnrealBuildTool.IsAssemblingBuild )
-								{
-									// Also check the timestamp on the directory the source file is in.  If the directory timestamp has
-									// changed, new source files may have been added or deleted.  We don't know whether the new/deleted
-									// files were actually UObject headers, but because we don't know all of the files we processed
-									// in the previous run, we need to assume our generated code is out of date if the directory timestamp
-									// is newer.
-									var HeaderDirectoryTimestamp = new DirectoryInfo( Path.GetDirectoryName( HeaderFile.AbsolutePath ) ).LastWriteTime;
-									if( SavedTimestamp.CompareTo( HeaderDirectoryTimestamp) < 0 )
-									{
-										bIsOutOfDate = true;
-										break;
-									}
-								}
-							}
-						}
-						else
-						{
-							// Generated code is older UnrealHeaderTool.exe or CoreUObject headers.  Out of date!
-							bIsOutOfDate = true;
-						}
-					}
-					else
-					{
-						// Timestamp file was missing (possibly deleted/cleaned), so headers are out of date
-						bIsOutOfDate = true;
-					}
+				// Grab our special "Timestamp" file that we saved after the last set of headers were generated
+				string TimestampFile = Path.Combine( GeneratedCodeDirectory, @"Timestamp" );
+				var SavedTimestampFileInfo = (FileSystemInfo)new FileInfo(TimestampFile);
+				if (!SavedTimestampFileInfo.Exists)
+				{
+					// Timestamp file was missing (possibly deleted/cleaned), so headers are out of date
+					return true;
 				}
-				else
+
+				// Make sure the last UHT run completed after UnrealHeaderTool.exe was compiled last, and after the CoreUObject headers were touched last.
+				var SavedTimestamp = SavedTimestampFileInfo.LastWriteTime;
+				if( HeaderToolTimestamp > SavedTimestamp || CoreGeneratedTimestamp > SavedTimestamp )
 				{
-					// Generated code directory is missing entirely!
-					bIsOutOfDate = true;
+					// Generated code is older UnrealHeaderTool.exe or CoreUObject headers.  Out of date!
+					return true;
 				}
 
-				// If even one module is out of date, we're done!  UHT does them all in one fell swoop.;
-				if( bIsOutOfDate )
+				// Iterate over our UObjects headers and figure out if any of them have changed
+				var AllUObjectHeaders = new List<FileItem>();
+				AllUObjectHeaders.AddRange( Module.PublicUObjectClassesHeaders );
+				AllUObjectHeaders.AddRange( Module.PublicUObjectHeaders );
+				AllUObjectHeaders.AddRange( Module.PrivateUObjectHeaders );
+				foreach( var HeaderFile in AllUObjectHeaders )
 				{
-					break;
+					var HeaderFileTimestamp = HeaderFile.Info.LastWriteTime;
+
+					// Has the source header changed since we last generated headers successfully?
+					if( HeaderFileTimestamp > SavedTimestamp )
+					{
+						return true;
+					}
+
+					// When we're running in assembler mode, outdatedness cannot be inferred by checking the directory timestamp
+					// of the source headers.  We don't care if source files were added or removed in this mode, because we're only
+					// able to process the known UObject headers that are in the Makefile.  If UObject header files are added/removed,
+					// we expect the user to re-run GenerateProjectFiles which will force UBTMakefile outdatedness.
+					// @todo ubtmake: Possibly, we should never be doing this check these days.
+					if( UnrealBuildTool.IsGatheringBuild || !UnrealBuildTool.IsAssemblingBuild )
+					{
+						// Also check the timestamp on the directory the source file is in.  If the directory timestamp has
+						// changed, new source files may have been added or deleted.  We don't know whether the new/deleted
+						// files were actually UObject headers, but because we don't know all of the files we processed
+						// in the previous run, we need to assume our generated code is out of date if the directory timestamp
+						// is newer.
+						var HeaderDirectoryTimestamp = new DirectoryInfo( Path.GetDirectoryName( HeaderFile.AbsolutePath ) ).LastWriteTime;
+						if( HeaderDirectoryTimestamp > SavedTimestamp )
+						{
+							return true;
+						}
+					}
 				}
 			}
 
-			return bIsOutOfDate;
+			return false;
 		}
 
 		/** Updates the intermediate include directory timestamps of all the passed in UObject modules */
diff --git a/Engine/Source/Programs/UnrealBuildTool/System/FileItem.cs b/Engine/Source/Programs/UnrealBuildTool/System/FileItem.cs
--- a/Engine/Source/Programs/UnrealBuildTool/System/FileItem.cs
+++ b/Engine/Source/Programs/UnrealBuildTool/System/FileItem.cs
@@ -174,7 +174,14 @@ namespace UnrealBuildTool
 		/** A list of remote file items that have been created but haven't needed the remote info yet, so we can gang up many into one request */
 		static List<FileItem> DelayedRemoteLookupFiles = new List<FileItem>();
 
-
+		/**
+		 * Clears the FileItem caches.
+		 */
+		public static void ClearCaches()
+		{
+			UniqueSourceFileMap.Clear();
+			DelayedRemoteLookupFiles.Clear();
+		}
 
 		/**
 		 * Resolve any outstanding remote file info lookups
@@ -333,7 +340,7 @@ namespace UnrealBuildTool
 		}
 
 		/** Initialization constructor. */
-		protected FileItem (string FileAbsolutePath)
+		protected FileItem(string FileAbsolutePath)
 		{
 			AbsolutePath = FileAbsolutePath;
 
diff --git a/Engine/Source/Programs/UnrealBuildTool/System/UnrealBuildTool.cs b/Engine/Source/Programs/UnrealBuildTool/System/UnrealBuildTool.cs
--- a/Engine/Source/Programs/UnrealBuildTool/System/UnrealBuildTool.cs
+++ b/Engine/Source/Programs/UnrealBuildTool/System/UnrealBuildTool.cs
@@ -1692,6 +1692,8 @@ namespace UnrealBuildTool
                             // a 'gather' and 'assemble' in the same run.  This will take a while longer, but subsequent runs will be fast!
                             UnrealBuildTool.bIsGatheringBuild_Unsafe = true;
 
+							FileItem.ClearCaches();
+
                             Log.TraceInformation( "Creating makefile for {0}{1}{2} ({3})", 
 								bIsHotReload ? "hot reloading " : "", 
 								TargetDescs[0].TargetName, 
@@ -2272,172 +2274,144 @@ namespace UnrealBuildTool
         /// <param name="TargetDescs">List of targets.  Order is not important</param>
 		/// <param name="ReasonNotLoaded">If the function returns null, this string will contain the reason why</param>
 		/// <returns>The loaded makefile, or null if it failed for some reason.  On failure, the 'ReasonNotLoaded' variable will contain information about why</returns>
-        static UBTMakefile LoadUBTMakefile( List<TargetDescriptor> TargetDescs, out string ReasonNotLoaded )
-        {
-            var UBTMakefileInfo = new FileInfo( GetUBTMakefilePath( TargetDescs ) );
-			ReasonNotLoaded = null;
-
-            // Check the directory timestamp on the project files directory.  If the user has generated project files more
-            // recently than the UBTMakefile, then we need to consider the file to be out of date
-            bool bForceOutOfDate = false;
-            if( UBTMakefileInfo.Exists )
-            {
-                // @todo ubtmake: This will only work if the directory timestamp actually changes with every single GPF.  Force delete existing files before creating new ones?  Eh... really we probably just want to delete + create a file in that folder
-                //			-> UPDATE: Seems to work OK right now though on Windows platform, maybe due to GUID changes
-                // @todo ubtmake: Some platforms may not save any files into this folder.  We should delete + generate a "touch" file to force the directory timestamp to be updated (or just check the timestamp file itself.  We could put it ANYWHERE, actually)
-
-				if( !UnrealBuildTool.RunningRocket() )
-                {
-                    if( Directory.Exists( ProjectFileGenerator.IntermediateProjectFilesPath ) )
-                    {
-                        var EngineProjectFilesLastUpdateTime = new FileInfo(ProjectFileGenerator.ProjectTimestampFile).LastWriteTime;
-						if( UBTMakefileInfo.LastWriteTime.CompareTo( EngineProjectFilesLastUpdateTime ) < 0 )
-						{
-							// Engine project files are newer than UBTMakefile
-							bForceOutOfDate = true;
-							Log.TraceVerbose("Existing makefile is older than generated engine project files, ignoring it" );
+		static UBTMakefile LoadUBTMakefile( List<TargetDescriptor> TargetDescs, out string ReasonNotLoaded )
+		{
+			// Check the directory timestamp on the project files directory.  If the user has generated project files more
+			// recently than the UBTMakefile, then we need to consider the file to be out of date
+			var UBTMakefileInfo = new FileInfo( GetUBTMakefilePath( TargetDescs ) );
+			if( !UBTMakefileInfo.Exists )
+			{
+				// UBTMakefile doesn't even exist, so we won't bother loading it
+				ReasonNotLoaded = "no existing makefile";
+				return null;
+			}
 
-							ReasonNotLoaded = "project files are newer";
-						}
-                    }
-                }
-                else
-                {
-                    // Rocket doesn't need to check engine projects for outdatedness
-                }
+			// @todo ubtmake: This will only work if the directory timestamp actually changes with every single GPF.  Force delete existing files before creating new ones?  Eh... really we probably just want to delete + create a file in that folder
+			//			-> UPDATE: Seems to work OK right now though on Windows platform, maybe due to GUID changes
+			// @todo ubtmake: Some platforms may not save any files into this folder.  We should delete + generate a "touch" file to force the directory timestamp to be updated (or just check the timestamp file itself.  We could put it ANYWHERE, actually)
 
-				if( !bForceOutOfDate )
+			// Rocket doesn't need to check engine projects for outdatedness
+			if( !UnrealBuildTool.RunningRocket() )
+			{
+				if( Directory.Exists( ProjectFileGenerator.IntermediateProjectFilesPath ) )
 				{
-					// Check the game project directory too
-					if( UnrealBuildTool.HasUProjectFile() )
+					var EngineProjectFilesLastUpdateTime = new FileInfo(ProjectFileGenerator.ProjectTimestampFile).LastWriteTime;
+					if( UBTMakefileInfo.LastWriteTime.CompareTo( EngineProjectFilesLastUpdateTime ) < 0 )
 					{
-						string ProjectFilename = UnrealBuildTool.GetUProjectFile();
-						var ProjectFileInfo = new FileInfo(ProjectFilename);
-						if( !ProjectFileInfo.Exists || UBTMakefileInfo.LastWriteTime.CompareTo(ProjectFileInfo.LastWriteTime ) < 0 )
-						{
-							// .uproject file is newer than UBTMakefile
-							bForceOutOfDate = true;
-							Log.TraceVerbose("Makefile is older than .uproject file, ignoring it" );
-
-							ReasonNotLoaded = ".uproject file is newer";
-						}
-						else
-						{
-							var MasterProjectRelativePath = UnrealBuildTool.GetUProjectPath();
-							var GameIntermediateProjectFilesPath = Path.Combine( MasterProjectRelativePath, "Intermediate", "ProjectFiles" );
-							if( Directory.Exists( GameIntermediateProjectFilesPath ) )
-							{
-								var GameProjectFilesLastUpdateTime = new DirectoryInfo( GameIntermediateProjectFilesPath ).LastWriteTime;
-								if( UBTMakefileInfo.LastWriteTime.CompareTo( GameProjectFilesLastUpdateTime ) < 0 )
-								{
-									// Game project files are newer than UBTMakefile
-									bForceOutOfDate = true;
-									Log.TraceVerbose("Makefile is older than generated game project files, ignoring it" );
-
-									ReasonNotLoaded = "game project files are newer";
-								}
-							}
-						}
+						// Engine project files are newer than UBTMakefile
+						Log.TraceVerbose("Existing makefile is older than generated engine project files, ignoring it" );
+						ReasonNotLoaded = "project files are newer";
+						return null;
 					}
 				}
+			}
 
-				if( !bForceOutOfDate )
+			// Check the game project directory too
+			if( UnrealBuildTool.HasUProjectFile() )
+			{
+				string ProjectFilename = UnrealBuildTool.GetUProjectFile();
+				var ProjectFileInfo = new FileInfo(ProjectFilename);
+				if( !ProjectFileInfo.Exists || UBTMakefileInfo.LastWriteTime.CompareTo(ProjectFileInfo.LastWriteTime ) < 0 )
 				{
-					// Check to see if UnrealBuildTool.exe was compiled more recently than the UBTMakefile
-					var UnrealBuildToolTimestamp = new FileInfo( Assembly.GetExecutingAssembly().Location ).LastWriteTime;
-					if( UBTMakefileInfo.LastWriteTime.CompareTo( UnrealBuildToolTimestamp ) < 0 )
-					{
-						// UnrealBuildTool.exe was compiled more recently than the UBTMakefile
-						Log.TraceVerbose("Makefile is older than UnrealBuildTool.exe, ignoring it" );
+					// .uproject file is newer than UBTMakefile
+					Log.TraceVerbose("Makefile is older than .uproject file, ignoring it" );
+					ReasonNotLoaded = ".uproject file is newer";
+					return null;
+				}
 
-						ReasonNotLoaded = "UnrealBuildTool.exe is newer";
-						bForceOutOfDate = true;
+				var MasterProjectRelativePath = UnrealBuildTool.GetUProjectPath();
+				var GameIntermediateProjectFilesPath = Path.Combine( MasterProjectRelativePath, "Intermediate", "ProjectFiles" );
+				if( Directory.Exists( GameIntermediateProjectFilesPath ) )
+				{
+					var GameProjectFilesLastUpdateTime = new DirectoryInfo( GameIntermediateProjectFilesPath ).LastWriteTime;
+					if( UBTMakefileInfo.LastWriteTime.CompareTo( GameProjectFilesLastUpdateTime ) < 0 )
+					{
+						// Game project files are newer than UBTMakefile
+						Log.TraceVerbose("Makefile is older than generated game project files, ignoring it" );
+						ReasonNotLoaded = "game project files are newer";
+						return null;
 					}
 				}
 			}
-			else
-			{
-				// UBTMakefile doesn't even exist, so we won't bother loading it
-				bForceOutOfDate = true;
 
-				ReasonNotLoaded = "no existing makefile";
+			// Check to see if UnrealBuildTool.exe was compiled more recently than the UBTMakefile
+			var UnrealBuildToolTimestamp = new FileInfo( Assembly.GetExecutingAssembly().Location ).LastWriteTime;
+			if( UBTMakefileInfo.LastWriteTime.CompareTo( UnrealBuildToolTimestamp ) < 0 )
+			{
+				// UnrealBuildTool.exe was compiled more recently than the UBTMakefile
+				Log.TraceVerbose("Makefile is older than UnrealBuildTool.exe, ignoring it" );
+				ReasonNotLoaded = "UnrealBuildTool.exe is newer";
+				return null;
 			}
 
-            UBTMakefile LoadedUBTMakefile = null;
-            if( !bForceOutOfDate )
-            { 
-                try
-                {
-					var LoadUBTMakefileStartTime = DateTime.UtcNow;
+			UBTMakefile LoadedUBTMakefile = null;
 
-                    using (FileStream Stream = new FileStream(UBTMakefileInfo.FullName, FileMode.Open, FileAccess.Read))
-                    {	
-                        BinaryFormatter Formatter = new BinaryFormatter();
-                        LoadedUBTMakefile = Formatter.Deserialize(Stream) as UBTMakefile;
-                    }
+			try
+			{
+				var LoadUBTMakefileStartTime = DateTime.UtcNow;
 
-					if( BuildConfiguration.bPrintPerformanceInfo )
-					{ 
-						var LoadUBTMakefileTime = (DateTime.UtcNow - LoadUBTMakefileStartTime).TotalSeconds;
-						Log.TraceInformation( "LoadUBTMakefile took " + LoadUBTMakefileTime + "s" );
-					}
-                }
-                catch (Exception Ex)
-                {
-                    Log.TraceWarning("Failed to read makefile: {0}", Ex.Message);
-					ReasonNotLoaded = "couldn't read existing makefile";
-                }
+				using (FileStream Stream = new FileStream(UBTMakefileInfo.FullName, FileMode.Open, FileAccess.Read))
+				{	
+					BinaryFormatter Formatter = new BinaryFormatter();
+					LoadedUBTMakefile = Formatter.Deserialize(Stream) as UBTMakefile;
+				}
 
-                if( LoadedUBTMakefile != null && !LoadedUBTMakefile.IsValidMakefile() )
-                {
-                    Log.TraceWarning("Loaded makefile appears to have invalid contents, ignoring it ({0})", UBTMakefileInfo.FullName );
-                    LoadedUBTMakefile = null;
-					ReasonNotLoaded = "existing makefile appears to be invalid";
-                }
+				if( BuildConfiguration.bPrintPerformanceInfo )
+				{ 
+					var LoadUBTMakefileTime = (DateTime.UtcNow - LoadUBTMakefileStartTime).TotalSeconds;
+					Log.TraceInformation( "LoadUBTMakefile took " + LoadUBTMakefileTime + "s" );
+				}
+			}
+			catch (Exception Ex)
+			{
+				Log.TraceWarning("Failed to read makefile: {0}", Ex.Message);
+				ReasonNotLoaded = "couldn't read existing makefile";
+				return null;
+			}
+
+			if( !LoadedUBTMakefile.IsValidMakefile() )
+			{
+				Log.TraceWarning("Loaded makefile appears to have invalid contents, ignoring it ({0})", UBTMakefileInfo.FullName );
+				ReasonNotLoaded = "existing makefile appears to be invalid";
+				return null;
+			}
 
-				if( LoadedUBTMakefile != null )
+			// Check if any of the target's Build.cs files are newer than the makefile
+			foreach (var Target in LoadedUBTMakefile.Targets)
+			{
+				string TargetCsFilename = Target.TargetCsFilename;
+				if (TargetCsFilename != null)
 				{
-					// Check if any of the target's Build.cs files are newer than the makefile
-					foreach (var Target in LoadedUBTMakefile.Targets)
+					var TargetCsFile = new FileInfo(TargetCsFilename);
+					bool bTargetCsFileExists = TargetCsFile.Exists;
+					if (!bTargetCsFileExists || TargetCsFile.LastWriteTime > UBTMakefileInfo.LastWriteTime)
 					{
-						string TargetCsFilename = Target.TargetCsFilename;
-						if (TargetCsFilename != null)
-						{
-							var TargetCsFile = new FileInfo(TargetCsFilename);
-							bool bTargetCsFileExists = TargetCsFile.Exists;
-							if (!bTargetCsFileExists || TargetCsFile.LastWriteTime > UBTMakefileInfo.LastWriteTime)
-							{
-								Log.TraceWarning("{0} has been {1} since makefile was built, ignoring it ({2})", TargetCsFilename, bTargetCsFileExists ? "changed" : "deleted", UBTMakefileInfo.FullName);
-								LoadedUBTMakefile = null;
-								ReasonNotLoaded = string.Format("changes to target files");
-								goto SkipRemainingTimestampChecks;
-							}
-						}
+						Log.TraceVerbose("{0} has been {1} since makefile was built, ignoring it ({2})", TargetCsFilename, bTargetCsFileExists ? "changed" : "deleted", UBTMakefileInfo.FullName);
+						ReasonNotLoaded = string.Format("changes to target files");
+						return null;
+					}
+				}
 
-						List<string> BuildCsFilenames = Target.GetAllModuleBuildCsFilenames();
-						foreach (var BuildCsFilename in BuildCsFilenames)
+				List<string> BuildCsFilenames = Target.GetAllModuleBuildCsFilenames();
+				foreach (var BuildCsFilename in BuildCsFilenames)
+				{
+					if (BuildCsFilename != null)
+					{
+						var BuildCsFile = new FileInfo(BuildCsFilename);
+						bool bBuildCsFileExists = BuildCsFile.Exists;
+						if (!bBuildCsFileExists || BuildCsFile.LastWriteTime > UBTMakefileInfo.LastWriteTime)
 						{
-							if (BuildCsFilename != null)
-							{
-								var BuildCsFile = new FileInfo(BuildCsFilename);
-								bool bBuildCsFileExists = BuildCsFile.Exists;
-								if (!bBuildCsFileExists || BuildCsFile.LastWriteTime > UBTMakefileInfo.LastWriteTime)
-								{
-									Log.TraceWarning("{0} has been {1} since makefile was built, ignoring it ({2})", BuildCsFilename, bBuildCsFileExists ? "changed" : "deleted", UBTMakefileInfo.FullName);
-									LoadedUBTMakefile = null;
-									ReasonNotLoaded = string.Format("changes to module files");
-									goto SkipRemainingTimestampChecks;
-								}
-							}
+							Log.TraceVerbose("{0} has been {1} since makefile was built, ignoring it ({2})", BuildCsFilename, bBuildCsFileExists ? "changed" : "deleted", UBTMakefileInfo.FullName);
+							ReasonNotLoaded = string.Format("changes to module files");
+							return null;
 						}
 					}
-
-				SkipRemainingTimestampChecks:;
 				}
-            }
+			}
 
-            return LoadedUBTMakefile;
-        }
+			ReasonNotLoaded = null;
+			return LoadedUBTMakefile;
+		}
 
 
         /// <summary>
