Add markup commands to prevent generating headers pushing the progress bar up to 100%.

diff --git a/Engine/Source/Programs/UnrealBuildTool/System/ExternalExecution.cs b/Engine/Source/Programs/UnrealBuildTool/System/ExternalExecution.cs
--- a/Engine/Source/Programs/UnrealBuildTool/System/ExternalExecution.cs
+++ b/Engine/Source/Programs/UnrealBuildTool/System/ExternalExecution.cs
@@ -509,6 +509,10 @@ namespace UnrealBuildTool
 		 */
 		public static bool ExecuteHeaderToolIfNecessary( UEBuildTarget Target, CPPEnvironment GlobalCompileEnvironment, List<UHTModuleInfo> UObjectModules, string ModuleInfoFileName, ref ECompilationResult UHTResult )
 		{
+			if(ProgressWriter.bWriteMarkup)
+			{
+				Log.WriteLine(TraceEventType.Information, "@progress push 5%");
+			}
 			using (ProgressWriter Progress = new ProgressWriter("Generating code...", false))
 			{
 				// We never want to try to execute the header tool when we're already trying to build it!
@@ -669,6 +673,10 @@ namespace UnrealBuildTool
 
 				Progress.Write(3, 3);
 			}
+			if(ProgressWriter.bWriteMarkup)
+			{
+				Log.WriteLine(TraceEventType.Information, "@progress pop");
+			}
 			return true;
 		}
 	}
