[INTEGRATE] Change 2567143 by [REDACTED_PERSON][REDACTED_HANDLE] on 2015/05/27 15:04:49

	Added 'Minimum OS Version' as a build condition of content projects for windows.

	Fixed a small issue in GameProjectUtils that meant the Default(Engine) configs were not being checked against the project correctly.

	UE-15605

diff --git a/Engine/Config/BaseEngine.ini b/Engine/Config/BaseEngine.ini
--- a/Engine/Config/BaseEngine.ini
+++ b/Engine/Config/BaseEngine.ini
@@ -2192,6 +2192,7 @@ bEnable2DPhysics=false
 [/Script/WindowsTargetPlatform.WindowsTargetSettings]
 +TargetedRHIs=PCD3D_SM5
 +TargetedRHIs=PCD3D_SM4
+MinimumOSVersion=MSOS_Vista
 
 [/Script/LinuxTargetPlatform.LinuxTargetSettings]
 +TargetedRHIs=GLSL_430
diff --git a/Engine/Source/Developer/Windows/WindowsTargetPlatform/Private/GenericWindowsTargetPlatform.h b/Engine/Source/Developer/Windows/WindowsTargetPlatform/Private/GenericWindowsTargetPlatform.h
--- a/Engine/Source/Developer/Windows/WindowsTargetPlatform/Private/GenericWindowsTargetPlatform.h
+++ b/Engine/Source/Developer/Windows/WindowsTargetPlatform/Private/GenericWindowsTargetPlatform.h
@@ -109,6 +109,12 @@ public:
 		return TSuper::SupportsFeature(Feature);
 	}
 
+	virtual void GetBuildProjectSettingKeys(FString& OutSection, TArray<FString>& InBoolKeys, TArray<FString>& InIntKeys, TArray<FString>& InStringKeys) const override
+	{
+		OutSection = TEXT("/Script/WindowsTargetPlatform.WindowsTargetSettings");
+		InStringKeys.Add(TEXT("MinimumOSVersion"));
+	}
+
 #if WITH_ENGINE
 	virtual void GetAllPossibleShaderFormats( TArray<FName>& OutFormats ) const override
 	{
diff --git a/Engine/Source/Editor/GameProjectGeneration/Private/GameProjectUtils.cpp b/Engine/Source/Editor/GameProjectGeneration/Private/GameProjectUtils.cpp
--- a/Engine/Source/Editor/GameProjectGeneration/Private/GameProjectUtils.cpp
+++ b/Engine/Source/Editor/GameProjectGeneration/Private/GameProjectUtils.cpp
@@ -3015,7 +3015,7 @@ bool GameProjectUtils::DoProjectSettingsMatchDefault(const FString& InPlatformNa
 	FConfigFile ProjIni;
 	FConfigFile DefaultIni;
 	FConfigCacheIni::LoadLocalIniFile(ProjIni, TEXT("Engine"), true, *InPlatformName, true);
-	FConfigCacheIni::LoadLocalIniFile(DefaultIni, TEXT("Engine"), true, NULL, true);
+	FConfigCacheIni::LoadExternalIniFile(DefaultIni, TEXT("Engine"), *FPaths::EngineConfigDir(), *FPaths::EngineConfigDir(), true, NULL, true);
 
 	if (InBoolKeys != NULL)
 	{
diff --git a/Engine/Source/Programs/UnrealBuildTool/Windows/UEBuildWindows.cs b/Engine/Source/Programs/UnrealBuildTool/Windows/UEBuildWindows.cs
--- a/Engine/Source/Programs/UnrealBuildTool/Windows/UEBuildWindows.cs
+++ b/Engine/Source/Programs/UnrealBuildTool/Windows/UEBuildWindows.cs
@@ -379,6 +379,26 @@ namespace UnrealBuildTool
             return false;
         }
 
+		public override bool HasDefaultBuildConfig(UnrealTargetPlatform Platform, string ProjectPath)
+		{
+			if (Platform == UnrealTargetPlatform.Win32)
+			{
+				string[] StringKeys = new string[] {
+					"MinimumOSVersion"
+				};
+
+				// look up OS specific settings
+				if (!DoProjectSettingsMatchDefault(Platform, ProjectPath, "/Script/WindowsTargetPlatform.WindowsTargetSettings",
+					null, null, StringKeys))
+				{
+					return false;
+				}
+			}
+
+			// check the base settings
+			return base.HasDefaultBuildConfig(Platform, ProjectPath);
+		}
+
 		private void SetupXPSupportFromConfiguration()
 		{
 			string[] CmdLine = Environment.GetCommandLineArgs();
