Fixed issue with GearVR plugin compiling into unsupported Android platforms
#1198
#ue4
#android

diff --git a/Engine/Plugins/Runtime/GearVR/Source/GearVR/GearVR.Build.cs b/Engine/Plugins/Runtime/GearVR/Source/GearVR/GearVR.Build.cs
--- a/Engine/Plugins/Runtime/GearVR/Source/GearVR/GearVR.Build.cs
+++ b/Engine/Plugins/Runtime/GearVR/Source/GearVR/GearVR.Build.cs
@@ -29,7 +29,7 @@ namespace UnrealBuildTool.Rules
 				}
 				);
 
-			if (Target.Platform == UnrealTargetPlatform.Android)
+			if (Target.Platform == UnrealTargetPlatform.Android && Target.Architecture == "-armv7")
 			{
 				PrivateDependencyModuleNames.AddRange(new string[] { "LibOVRMobile" });
 			}
diff --git a/Engine/Plugins/Runtime/GearVR/Source/GearVR/Private/GearVR.cpp b/Engine/Plugins/Runtime/GearVR/Source/GearVR/Private/GearVR.cpp
--- a/Engine/Plugins/Runtime/GearVR/Source/GearVR/Private/GearVR.cpp
+++ b/Engine/Plugins/Runtime/GearVR/Source/GearVR/Private/GearVR.cpp
@@ -7,7 +7,9 @@
 #include "Android/AndroidJNI.h"
 #include "RHIStaticStates.h"
 
+#if GEARVR_SUPPORTED_PLATFORMS
 #include "VrApi_Helpers.h"
+#endif
 
 #define DEFAULT_PREDICTION_IN_SECONDS 0.035
 
@@ -16,7 +18,9 @@
 #if PLATFORM_SUPPORTS_PRAGMA_PACK
 #pragma pack (push,8)
 #endif
+#if GEARVR_SUPPORTED_PLATFORMS
 #include "../Src/Kernel/OVR_Log.h"
+#endif
 #if PLATFORM_SUPPORTS_PRAGMA_PACK
 #pragma pack (pop)
 #endif
diff --git a/Engine/Plugins/Runtime/GearVR/Source/GearVR/Public/IGearVRPlugin.h b/Engine/Plugins/Runtime/GearVR/Source/GearVR/Public/IGearVRPlugin.h
--- a/Engine/Plugins/Runtime/GearVR/Source/GearVR/Public/IGearVRPlugin.h
+++ b/Engine/Plugins/Runtime/GearVR/Source/GearVR/Public/IGearVRPlugin.h
@@ -5,7 +5,7 @@
 #include "ModuleManager.h"
 #include "IHeadMountedDisplayModule.h"
 
-#define GEARVR_SUPPORTED_PLATFORMS PLATFORM_ANDROID
+#define GEARVR_SUPPORTED_PLATFORMS (PLATFORM_ANDROID && PLATFORM_ANDROID_ARM)
 
 /**
  * The public interface to this module.  In most cases, this interface is only public to sibling modules 
diff --git a/Engine/Source/Programs/UnrealBuildTool/Android/UEBuildAndroid.cs b/Engine/Source/Programs/UnrealBuildTool/Android/UEBuildAndroid.cs
--- a/Engine/Source/Programs/UnrealBuildTool/Android/UEBuildAndroid.cs
+++ b/Engine/Source/Programs/UnrealBuildTool/Android/UEBuildAndroid.cs
@@ -372,6 +372,14 @@ namespace UnrealBuildTool
 
 			InBuildTarget.GlobalCompileEnvironment.Config.Definitions.Add("PLATFORM_ANDROID=1");
 			InBuildTarget.GlobalCompileEnvironment.Config.Definitions.Add("ANDROID=1");
+			switch (GetActiveArchitecture())
+			{
+				case "-armv7":	InBuildTarget.GlobalCompileEnvironment.Config.Definitions.Add("ANDROID_ARM=1");		break;
+				case "-arm64":	InBuildTarget.GlobalCompileEnvironment.Config.Definitions.Add("ANDROID_ARM64=1");	break;
+				case "-x86":	InBuildTarget.GlobalCompileEnvironment.Config.Definitions.Add("ANDROID_X86=1");		break;
+				case "-x64":	InBuildTarget.GlobalCompileEnvironment.Config.Definitions.Add("ANDROID_X64=1");		break;
+				default:		InBuildTarget.GlobalCompileEnvironment.Config.Definitions.Add("ANDROID_ARM=1");		break;
+			}
 
 			InBuildTarget.GlobalCompileEnvironment.Config.Definitions.Add("WITH_DATABASE_SUPPORT=0");
 			InBuildTarget.GlobalCompileEnvironment.Config.Definitions.Add("WITH_EDITOR=0");
diff --git a/Engine/Source/Runtime/Core/Public/HAL/Platform.h b/Engine/Source/Runtime/Core/Public/HAL/Platform.h
--- a/Engine/Source/Runtime/Core/Public/HAL/Platform.h
+++ b/Engine/Source/Runtime/Core/Public/HAL/Platform.h
@@ -23,6 +23,18 @@
 #if !defined(PLATFORM_ANDROID)
 	#define PLATFORM_ANDROID 0
 #endif
+#if !defined(PLATFORM_ANDROID_ARM)
+	#define PLATFORM_ANDROID_ARM 0
+#endif
+#if !defined(PLATFORM_ANDROID_ARM64)
+	#define PLATFORM_ANDROID_ARM64 0
+#endif
+#if !defined(PLATFORM_ANDROID_X86)
+	#define PLATFORM_ANDROID_X86 0
+#endif
+#if !defined(PLATFORM_ANDROID_X64)
+	#define PLATFORM_ANDROID_X64 0
+#endif
 #if !defined(PLATFORM_ANDROIDGL4)
 	#define PLATFORM_ANDROIDGL4 0
 #endif
