Enable shadow warnings for Android targets and clean up the ones found
#ue-14251
#ue4
#android

diff --git a/Engine/Source/Programs/UnrealBuildTool/Android/AndroidToolChain.cs b/Engine/Source/Programs/UnrealBuildTool/Android/AndroidToolChain.cs
--- a/Engine/Source/Programs/UnrealBuildTool/Android/AndroidToolChain.cs
+++ b/Engine/Source/Programs/UnrealBuildTool/Android/AndroidToolChain.cs
@@ -304,6 +304,11 @@ namespace UnrealBuildTool
 			Result += " -Wno-invalid-offsetof";			// needed to suppress warnings about using offsetof on non-POD types.
 			Result += " -Wno-logical-op-parentheses";	// needed for external headers we can't change
 
+			if (CompileEnvironment.Config.bEnableShadowVariableWarning || true)
+			{
+				Result += " -Wshadow -Wno-error=shadow";
+			}
+
 			// new for clang4.5 warnings:
 			if (ClangVersionFloat >= 3.5)
 			{
@@ -625,8 +630,7 @@ namespace UnrealBuildTool
 			}
 
 			var BuildPlatform = UEBuildPlatform.GetBuildPlatformForCPPTargetPlatform(CompileEnvironment.Config.Target.Platform);
-
-
+			var NDKRoot = Environment.GetEnvironmentVariable("NDKROOT").Replace("\\", "/");
 
 			string BasePCHName = "";
 			var PCHExtension = UEBuildPlatform.BuildPlatformDictionary[UnrealTargetPlatform.Android].GetBinaryExtension(UEBuildBinaryType.PrecompiledHeader);
@@ -683,6 +687,7 @@ namespace UnrealBuildTool
 						Action CompileAction = new Action(ActionType.Compile);
 						string FileArguments = "";
 						bool bIsPlainCFile = Path.GetExtension(SourceFile.AbsolutePath).ToUpperInvariant() == ".C";
+						bool bDisableShadowWarning = false;
 
 						// should we disable optimizations on this file?
 						// @todo android - We wouldn't need this if we could disable optimizations per function (via pragma)
@@ -702,6 +707,12 @@ namespace UnrealBuildTool
 						else if (bIsPlainCFile)
 						{
 							FileArguments += GetCompileArguments_C(bDisableOptimizations);
+
+							// remove shadow variable warnings for NDK files
+							if (SourceFile.AbsolutePath.Replace("\\", "/").StartsWith(NDKRoot))
+							{
+								bDisableShadowWarning = true;
+							}
 						}
 						else
 						{
@@ -762,6 +773,16 @@ namespace UnrealBuildTool
 						AllArguments = ActionThread.ExpandEnvironmentVariables(AllArguments);
 						AllArguments = AllArguments.Replace("\\", "/");
 
+						// Remove shadow warning for this file if requested
+						if (bDisableShadowWarning)
+						{
+							int WarningIndex = AllArguments.IndexOf(" -Wshadow");
+							if (WarningIndex > 0)
+							{
+								AllArguments = AllArguments.Remove(WarningIndex, 9);
+							}
+						}
+
 						// Create the response file
 						string ResponseFileName = CompileAction.ProducedItems[0].AbsolutePath + ".response";
 						string ResponseArgument = string.Format("@\"{0}\"", ResponseFile.Create(ResponseFileName, new List<string> { AllArguments }));
diff --git a/Engine/Source/Runtime/Core/Private/Android/AndroidFile.cpp b/Engine/Source/Runtime/Core/Private/Android/AndroidFile.cpp
--- a/Engine/Source/Runtime/Core/Private/Android/AndroidFile.cpp
+++ b/Engine/Source/Runtime/Core/Private/Android/AndroidFile.cpp
@@ -1464,11 +1464,11 @@ public:
 				while ((fileName = AAssetDir_getNextFileName(dir)) != nullptr && Result == true)
 				{
 					bool isDirectory = false;
-					AAssetDir * dir = AAssetManager_openDir(AssetMgr, fileName);
+					AAssetDir * subdir = AAssetManager_openDir(AssetMgr, fileName);
 					if (nullptr != dir)
 					{
 						isDirectory = true;
-						AAssetDir_close(dir);
+						AAssetDir_close(subdir);
 					}
 #if LOG_ANDROID_FILE
 					FPlatformMisc::LowLevelOutputDebugStringf(TEXT("FAndroidPlatformFile::IterateDirectory('%s').. ASSET Visit: '%s'"), Directory, UTF8_TO_TCHAR(fileName));
diff --git a/Engine/Source/Runtime/Core/Public/Android/AndroidJava.h b/Engine/Source/Runtime/Core/Public/Android/AndroidJava.h
--- a/Engine/Source/Runtime/Core/Public/Android/AndroidJava.h
+++ b/Engine/Source/Runtime/Core/Public/Android/AndroidJava.h
@@ -41,7 +41,6 @@ public:
 
 protected:
 
-	FName			ClassName;
 	jobject			Object;
 	jclass			Class;
 
diff --git a/Engine/Source/Runtime/Online/Android/OnlineSubsystemGooglePlay/Private/OnlineIdentityInterfaceGooglePlay.cpp b/Engine/Source/Runtime/Online/Android/OnlineSubsystemGooglePlay/Private/OnlineIdentityInterfaceGooglePlay.cpp
--- a/Engine/Source/Runtime/Online/Android/OnlineSubsystemGooglePlay/Private/OnlineIdentityInterfaceGooglePlay.cpp
+++ b/Engine/Source/Runtime/Online/Android/OnlineSubsystemGooglePlay/Private/OnlineIdentityInterfaceGooglePlay.cpp
@@ -181,12 +181,12 @@ void FOnlineIdentityGooglePlay::GetUserPrivilege(const FUniqueNetId& UserId, EUs
 	Delegate.ExecuteIfBound(UserId, Privilege, (uint32)EPrivilegeResults::NoFailures);
 }
 
-FPlatformUserId FOnlineIdentityGooglePlay::GetPlatformUserIdFromUniqueNetId(const FUniqueNetId& UniqueNetId)
+FPlatformUserId FOnlineIdentityGooglePlay::GetPlatformUserIdFromUniqueNetId(const FUniqueNetId& NetId)
 {
 	for (int i = 0; i < MAX_LOCAL_PLAYERS; ++i)
 	{
 		auto CurrentUniqueId = GetUniquePlayerId(i);
-		if (CurrentUniqueId.IsValid() && (*CurrentUniqueId == UniqueNetId))
+		if (CurrentUniqueId.IsValid() && (*CurrentUniqueId == NetId))
 		{
 			return i;
 		}
diff --git a/Engine/Source/Runtime/Online/Android/OnlineSubsystemGooglePlay/Private/OnlineIdentityInterfaceGooglePlay.h b/Engine/Source/Runtime/Online/Android/OnlineSubsystemGooglePlay/Private/OnlineIdentityInterfaceGooglePlay.h
--- a/Engine/Source/Runtime/Online/Android/OnlineSubsystemGooglePlay/Private/OnlineIdentityInterfaceGooglePlay.h
+++ b/Engine/Source/Runtime/Online/Android/OnlineSubsystemGooglePlay/Private/OnlineIdentityInterfaceGooglePlay.h
@@ -60,7 +60,7 @@ public:
 	virtual FString GetPlayerNickname(const FUniqueNetId& UserId) const override;
 	virtual FString GetAuthToken(int32 LocalUserNum) const override;
 	virtual void GetUserPrivilege(const FUniqueNetId& UserId, EUserPrivileges::Type Privilege, const FOnGetUserPrivilegeCompleteDelegate& Delegate) override;
-	virtual FPlatformUserId GetPlatformUserIdFromUniqueNetId(const FUniqueNetId& UniqueNetId) override;
+	virtual FPlatformUserId GetPlatformUserIdFromUniqueNetId(const FUniqueNetId& NetId) override;
 	// End IOnlineIdentity interface
 
 public:
diff --git a/Engine/Source/Runtime/Online/Android/OnlineSubsystemGooglePlay/Private/OnlineLeaderboardInterfaceGooglePlay.cpp b/Engine/Source/Runtime/Online/Android/OnlineSubsystemGooglePlay/Private/OnlineLeaderboardInterfaceGooglePlay.cpp
--- a/Engine/Source/Runtime/Online/Android/OnlineSubsystemGooglePlay/Private/OnlineLeaderboardInterfaceGooglePlay.cpp
+++ b/Engine/Source/Runtime/Online/Android/OnlineSubsystemGooglePlay/Private/OnlineLeaderboardInterfaceGooglePlay.cpp
@@ -66,7 +66,6 @@ bool FOnlineLeaderboardsGooglePlay::WriteLeaderboards(const FName& SessionName,
 		{
 			const FVariantData& Stat = It.Value();
 			uint64 Score;
-			FOnlinePendingLeaderboardWrite UnreportedScore;
 
 			UE_LOG_ONLINE(Display, TEXT("Here's a stat"));
 
