Remove forced shadow variable warning
#ue4
#android

diff --git a/Engine/Source/Programs/UnrealBuildTool/Android/AndroidToolChain.cs b/Engine/Source/Programs/UnrealBuildTool/Android/AndroidToolChain.cs
--- a/Engine/Source/Programs/UnrealBuildTool/Android/AndroidToolChain.cs
+++ b/Engine/Source/Programs/UnrealBuildTool/Android/AndroidToolChain.cs
@@ -304,7 +304,7 @@ namespace UnrealBuildTool
 			Result += " -Wno-invalid-offsetof";			// needed to suppress warnings about using offsetof on non-POD types.
 			Result += " -Wno-logical-op-parentheses";	// needed for external headers we can't change
 
-			if (CompileEnvironment.Config.bEnableShadowVariableWarning || true)
+			if (CompileEnvironment.Config.bEnableShadowVariableWarning)
 			{
 				Result += " -Wshadow -Wno-error=shadow";
 			}
