UE4 - If the number of physical cores is greater than 4 and not equal to the number of logical cores, then use the average of the two core counts for the number of parallel compiles.

diff --git a/Engine/Source/Programs/UnrealBuildTool/System/LocalExecutor.cs b/Engine/Source/Programs/UnrealBuildTool/System/LocalExecutor.cs
--- a/Engine/Source/Programs/UnrealBuildTool/System/LocalExecutor.cs
+++ b/Engine/Source/Programs/UnrealBuildTool/System/LocalExecutor.cs
@@ -363,10 +363,17 @@ namespace UnrealBuildTool
 			}
 			// The number of actions to execute in parallel is trying to keep the CPU busy enough in presence of I/O stalls.
 			int MaxActionsToExecuteInParallel = 0;
-			// The CPU has more logical cores than physical ones, aka uses hyper-threading. 
-			if( NumCores < System.Environment.ProcessorCount )
+			if (NumCores < System.Environment.ProcessorCount && BuildConfiguration.ProcessorCountMultiplier != 1.0)
 			{
-				MaxActionsToExecuteInParallel = (int) (NumCores * BuildConfiguration.ProcessorCountMultiplier);
+				// The CPU has more logical cores than physical ones, aka uses hyper-threading. 
+				// Use multiplier if provided
+				MaxActionsToExecuteInParallel = (int)(NumCores * BuildConfiguration.ProcessorCountMultiplier);
+			}
+			else if (NumCores < System.Environment.ProcessorCount && NumCores > 4)
+			{
+				// The CPU has more logical cores than physical ones, aka uses hyper-threading. 
+				// Use average of logical and physical if we have "lots of cores"
+				MaxActionsToExecuteInParallel = (int)(NumCores + System.Environment.ProcessorCount) / 2;
 			}
 			// No hyper-threading. Only kicking off a task per CPU to keep machine responsive.
 			else
