Makefile fix for UHT types being added to or removed from existing headers.
CaselessDictionary made serializable.
New FileContentsCacheType for enabling simple lookup and caching of file contents.
Fix for module references differing in case from module definitions e.g. Http and HTTP.
Some general refactoring.

diff --git a/Engine/Source/Programs/DotNETCommon/DotNETUtilities/CaselessDictionary.cs b/Engine/Source/Programs/DotNETCommon/DotNETUtilities/CaselessDictionary.cs
--- a/Engine/Source/Programs/DotNETCommon/DotNETUtilities/CaselessDictionary.cs
+++ b/Engine/Source/Programs/DotNETCommon/DotNETUtilities/CaselessDictionary.cs
@@ -2,6 +2,7 @@
 
 using System;
 using System.Collections.Generic;
+using System.Runtime.Serialization;
 
 namespace Tools.DotNETCommon.CaselessDictionary
 {
@@ -9,6 +10,7 @@ namespace Tools.DotNETCommon.CaselessDictionary
 	/// Equivalent of case insensitive Dictionary<string, T>
 	/// </summary>
 	/// <typeparam name="T"></typeparam>
+	[Serializable]
 	public class CaselessDictionary<T> : Dictionary<string, T>
 	{
 		public CaselessDictionary()
@@ -25,5 +27,10 @@ namespace Tools.DotNETCommon.CaselessDictionary
 			: base(Dict, StringComparer.InvariantCultureIgnoreCase)
 		{
 		}
+
+		protected CaselessDictionary(SerializationInfo Info, StreamingContext Context)
+			: base(Info, Context)
+		{
+		}
 	}
 }
diff --git a/Engine/Source/Programs/DotNETCommon/DotNETUtilities/DotNETUtilities.csproj b/Engine/Source/Programs/DotNETCommon/DotNETUtilities/DotNETUtilities.csproj
--- a/Engine/Source/Programs/DotNETCommon/DotNETUtilities/DotNETUtilities.csproj
+++ b/Engine/Source/Programs/DotNETCommon/DotNETUtilities/DotNETUtilities.csproj
@@ -51,6 +51,7 @@
     </Compile>
     <Compile Include="CaselessDictionary.cs" />
     <Compile Include="ExecutingAssembly.cs" />
+    <Compile Include="FileContentsCacheType.cs" />
     <Compile Include="FileSystem.cs" />
     <Compile Include="HarvestEnvVars.cs" />
     <Compile Include="LaunchProcess.cs" />
diff --git a/Engine/Source/Programs/DotNETCommon/DotNETUtilities/FileContentsCacheType.cs b/Engine/Source/Programs/DotNETCommon/DotNETUtilities/FileContentsCacheType.cs
new file mode 100644
--- /dev/null
+++ b/Engine/Source/Programs/DotNETCommon/DotNETUtilities/FileContentsCacheType.cs
@@ -0,0 +1,29 @@
+// Copyright 1998-2015 Epic Games, Inc. All Rights Reserved.
+
+using System.Collections.Generic;
+using System.IO;
+
+namespace Tools.DotNETCommon.FileContentsCacheType
+{
+	public class FileContentsCacheType
+	{
+		public string GetContents(string Filename)
+		{
+			string Contents;
+			if (FilenameContentsMap.TryGetValue(Filename, out Contents))
+			{
+				return Contents;
+			}
+
+			using (var Reader = new StreamReader(Filename, System.Text.Encoding.UTF8))
+			{
+				Contents = Reader.ReadToEnd();
+				FilenameContentsMap.Add(Filename, Contents);
+			}
+
+			return Contents;
+		}
+
+		private Dictionary<string, string> FilenameContentsMap = new Dictionary<string, string>();
+	}
+}
diff --git a/Engine/Source/Programs/UnrealBuildTool/Configuration/UEBuildBinary.cs b/Engine/Source/Programs/UnrealBuildTool/Configuration/UEBuildBinary.cs
--- a/Engine/Source/Programs/UnrealBuildTool/Configuration/UEBuildBinary.cs
+++ b/Engine/Source/Programs/UnrealBuildTool/Configuration/UEBuildBinary.cs
@@ -8,6 +8,7 @@ using System.IO;
 using System.Linq;
 using System.Xml;
 using System.Globalization;
+using Tools.DotNETCommon.CaselessDictionary;
 
 namespace UnrealBuildTool
 {	
@@ -469,23 +470,21 @@ namespace UnrealBuildTool
 		/// <returns>List of all referenced modules</returns>
 		public override List<UEBuildModule> GetAllDependencyModules(bool bIncludeDynamicallyLoaded, bool bForceCircular)
 		{
-			var OrderedModules = new List<UEBuildModule>();
-			var ReferencedModules = new Dictionary<string, UEBuildModule>( StringComparer.InvariantCultureIgnoreCase );
+			var ReferencedModules = new CaselessDictionary<UEBuildModule.ModuleIndexPair>();
 			foreach( var ModuleName in ModuleNames )
 			{
 				if( !ReferencedModules.ContainsKey( ModuleName ) )
 				{
 					var Module = Target.GetModuleByName( ModuleName );
-					ReferencedModules[ ModuleName ] = Module;
+					ReferencedModules[ ModuleName ] = null;
 
-					bool bOnlyDirectDependencies = false;
-					Module.GetAllDependencyModules(ReferencedModules, OrderedModules, bIncludeDynamicallyLoaded, bForceCircular, bOnlyDirectDependencies);
+					Module.GetAllDependencyModules(ReferencedModules, bIncludeDynamicallyLoaded, bForceCircular, bOnlyDirectDependencies: false);
 
-					OrderedModules.Add( Module );
+					ReferencedModules[ ModuleName ] = new UEBuildModule.ModuleIndexPair{ Module = Module, Index = ReferencedModules.Count };
 				}
 			}
 
-			return OrderedModules;
+			return ReferencedModules.Values.OrderBy(M => M.Index).Select(M => M.Module).ToList();
 		}
 
 		/// <summary>
@@ -655,28 +654,23 @@ namespace UnrealBuildTool
 			}
 			
 			// Add runtime dependencies for all the modules in this binary, and build up a list of all the referenced modules
-			Dictionary<string, UEBuildModule> ReferencedModules = new Dictionary<string,UEBuildModule>();
-			List<UEBuildModule> OrderedModules = new List<UEBuildModule>();
+			var ReferencedModules = new CaselessDictionary<UEBuildModule.ModuleIndexPair>();
 			foreach (string ModuleName in ModuleNames)
 			{
-				UEBuildModule Module = Target.GetModuleByName(ModuleName); 
+				UEBuildModule Module = Target.GetModuleByName(ModuleName);
 				foreach(RuntimeDependency RuntimeDependency in Module.RuntimeDependencies)
 				{
 					Receipt.RuntimeDependencies.Add(new RuntimeDependency(RuntimeDependency));
 				}
-				Module.GetAllDependencyModules(ReferencedModules, OrderedModules, true, false, true);
+				Module.GetAllDependencyModules(ReferencedModules, true, false, true);
 			}
 
 			// Add runtime dependencies for all the referenced external modules. These may be introduce dependencies for the binary without actually being listed for inclusion in it.
-			foreach(UEBuildModule OrderedModule in OrderedModules)
+			foreach(UEBuildExternalModule ExternalModule in ReferencedModules.Values.OrderBy(x => x.Index).Select(x => x.Module).OfType<UEBuildExternalModule>())
 			{
-				UEBuildExternalModule ExternalModule = OrderedModule as UEBuildExternalModule;
-				if(ExternalModule != null)
+				foreach(RuntimeDependency RuntimeDependency in ExternalModule.RuntimeDependencies)
 				{
-					foreach(RuntimeDependency RuntimeDependency in ExternalModule.RuntimeDependencies)
-					{
-						Receipt.RuntimeDependencies.Add(new RuntimeDependency(RuntimeDependency));
-					}
+					Receipt.RuntimeDependencies.Add(new RuntimeDependency(RuntimeDependency));
 				}
 			}
 			return Receipt;
diff --git a/Engine/Source/Programs/UnrealBuildTool/Configuration/UEBuildModule.cs b/Engine/Source/Programs/UnrealBuildTool/Configuration/UEBuildModule.cs
--- a/Engine/Source/Programs/UnrealBuildTool/Configuration/UEBuildModule.cs
+++ b/Engine/Source/Programs/UnrealBuildTool/Configuration/UEBuildModule.cs
@@ -8,6 +8,7 @@ using System.IO;
 using System.Linq;
 using System.Text.RegularExpressions;
 using System.Xml;
+using Tools.DotNETCommon.CaselessDictionary;
 
 namespace UnrealBuildTool
 {
@@ -841,17 +842,22 @@ namespace UnrealBuildTool
 			return Name;
 		}
 
+		[DebuggerDisplay("{Index}: {Module}")]
+		public class ModuleIndexPair
+		{
+			public UEBuildModule Module;
+			public int           Index;
+		}
 
 		/**
 		 * Gets all of the modules referenced by this module
 		 * 
-		 * @param	ReferencedModules	Hash of all referenced modules
-		 * @param	OrderedModules		Ordered list of the referenced modules
+		 * @param	ReferencedModules	Hash of all referenced modules with their addition index.
 		 * @param	bIncludeDynamicallyLoaded	True if dynamically loaded modules (and all of their dependent modules) should be included.
 		 * @param	bForceCircular	True if circular dependencies should be processed
 		 * @param	bOnlyDirectDependencies	True to return only this module's direct dependencies
 		 */
-		public virtual void GetAllDependencyModules(Dictionary<string, UEBuildModule> ReferencedModules, List<UEBuildModule> OrderedModules, bool bIncludeDynamicallyLoaded, bool bForceCircular, bool bOnlyDirectDependencies)
+		public virtual void GetAllDependencyModules(CaselessDictionary<ModuleIndexPair> ReferencedModules, bool bIncludeDynamicallyLoaded, bool bForceCircular, bool bOnlyDirectDependencies)
 		{
 		}
 
@@ -1373,14 +1379,13 @@ namespace UnrealBuildTool
 					if( bAllowSharedPCH && ( !bIsASharedPCHModule || bCanModuleUseOwnSharedPCH ) )
 					{
 						// Figure out which shared PCH tier we're in
-						var DirectDependencyModules = new List<UEBuildModule>();
-						{ 
-							var ReferencedModules = new Dictionary<string, UEBuildModule>( StringComparer.InvariantCultureIgnoreCase );
-							this.GetAllDependencyModules( ReferencedModules, DirectDependencyModules, bIncludeDynamicallyLoaded:false, bForceCircular:false, bOnlyDirectDependencies:true );
+						var ReferencedModules = new CaselessDictionary<ModuleIndexPair>();
+						{
+							this.GetAllDependencyModules( ReferencedModules, bIncludeDynamicallyLoaded:false, bForceCircular:false, bOnlyDirectDependencies:true );
 						}
 
 						int LargestSharedPCHHeaderFileIndex = -1;
-						foreach( var DependencyModule in DirectDependencyModules )
+						foreach( var DependencyModule in ReferencedModules.Values.OrderBy(P => P.Index).Select(P => P.Module) )
 						{
 							// These Shared PCHs are ordered from least complex to most complex.  We'll start at the last one and search backwards.
 							for( var SharedPCHHeaderFileIndex = GlobalCompileEnvironment.SharedPCHHeaderFiles.Count - 1; SharedPCHHeaderFileIndex > LargestSharedPCHHeaderFileIndex; --SharedPCHHeaderFileIndex )
@@ -1823,8 +1828,7 @@ namespace UnrealBuildTool
 		{
 			// @todo ubtmake: We don't really need to scan every file looking for PCH headers, just need one.  The rest is just for error checking.
 			// @todo ubtmake: We don't need all of the direct includes either.  We just need the first, unless we want to check for errors.
-			bool HasUObjects;
-			List<DependencyInclude> DirectIncludeFilenames = CPPEnvironment.GetDirectIncludeDependencies(Target, CPPFile, BuildPlatform, bOnlyCachedDependencies:false, HasUObjects:out HasUObjects);
+			List<DependencyInclude> DirectIncludeFilenames = CPPEnvironment.GetDirectIncludeDependencies(Target, CPPFile, BuildPlatform, bOnlyCachedDependencies:false);
 			if (BuildConfiguration.bPrintDebugInfo)
 			{
 				Log.TraceVerbose("Found direct includes for {0}: {1}", Path.GetFileName(CPPFile.AbsolutePath), string.Join(", ", DirectIncludeFilenames.Select(F => F.IncludeName)));
@@ -1913,11 +1917,19 @@ namespace UnrealBuildTool
 			return Result;
 		}
 
-		public HashSet<FileItem> _AllClassesHeaders     = new HashSet<FileItem>();
-		public HashSet<FileItem> _PublicUObjectHeaders  = new HashSet<FileItem>();
-		public HashSet<FileItem> _PrivateUObjectHeaders = new HashSet<FileItem>();
+		public class UHTModuleInfoCacheType
+		{
+			public UHTModuleInfoCacheType(IEnumerable<string> InHeaderFilenames, UHTModuleInfo InInfo)
+			{
+				HeaderFilenames = InHeaderFilenames;
+				Info            = InInfo;
+			}
 
-		private UHTModuleInfo CachedModuleUHTInfo = null;
+			public IEnumerable<string> HeaderFilenames = null;
+			public UHTModuleInfo       Info            = null;
+		}
+
+		private UHTModuleInfoCacheType UHTModuleInfoCache = null;
 
 		/// Total time spent generating PCHs for modules (not actually compiling, but generating the PCH's input data)
 		public static double TotalPCHGenTime = 0.0;
@@ -1926,69 +1938,23 @@ namespace UnrealBuildTool
 		public static double TotalPCHCacheTime = 0.0;
 
 
-
 		/// <summary>
 		/// If any of this module's source files contain UObject definitions, this will return those header files back to the caller
 		/// </summary>
-		/// <param name="PublicUObjectClassesHeaders">All UObjects headers in the module's Classes folder (legacy)</param>
-		/// <param name="PublicUObjectHeaders">Dependent UObject headers in Public source folders</param>
-		/// <param name="PrivateUObjectHeaders">Dependent UObject headers not in Public source folders</param>
 		/// <returns>
-		public UHTModuleInfo GetUHTModuleInfo()
+		public UHTModuleInfoCacheType GetCachedUHTModuleInfo()
 		{
-			if (CachedModuleUHTInfo != null)
-			{
-				return CachedModuleUHTInfo;
-			}
-
-			var ClassesFolder = Path.Combine(this.ModuleDirectory, "Classes");
-			var PublicFolder  = Path.Combine(this.ModuleDirectory, "Public");
-
-			var BuildPlatform = UEBuildPlatform.GetBuildPlatform(Target.Platform);
-
-			var ClassesFiles = Directory.GetFiles( this.ModuleDirectory, "*.h", SearchOption.AllDirectories );
-			foreach (var ClassHeader in ClassesFiles)
+			if (UHTModuleInfoCache == null)
 			{
-				// Check to see if we know anything about this file.  If we have up-to-date cached information about whether it has
-				// UObjects or not, we can skip doing a test here.
-				var UObjectHeaderFileItem = FileItem.GetExistingItemByPath( ClassHeader );
-
-				bool HasUObjects;
-				var DirectIncludes = CPPEnvironment.GetDirectIncludeDependencies( Target, UObjectHeaderFileItem, BuildPlatform, bOnlyCachedDependencies:false, HasUObjects:out HasUObjects );
-				Debug.Assert( DirectIncludes != null );
-				
-				if (HasUObjects)
-				{ 
-					if (UObjectHeaderFileItem.AbsolutePath.StartsWith(ClassesFolder))
-					{
-						_AllClassesHeaders.Add(UObjectHeaderFileItem);
-					}
-					else if (UObjectHeaderFileItem.AbsolutePath.StartsWith(PublicFolder))
-					{
-						_PublicUObjectHeaders.Add(UObjectHeaderFileItem);
-					}
-					else
-					{
-						_PrivateUObjectHeaders.Add(UObjectHeaderFileItem);
-					}
-				}
+				IEnumerable<string> HeaderFilenames = Directory.GetFiles( ModuleDirectory, "*.h", SearchOption.AllDirectories );
+				UHTModuleInfo       Info            = ExternalExecution.CreateUHTModuleInfo(HeaderFilenames, Target, Name, ModuleDirectory, Type);
+				UHTModuleInfoCache = new UHTModuleInfoCacheType(Info.PublicUObjectHeaders.Concat(Info.PublicUObjectClassesHeaders).Concat(Info.PrivateUObjectHeaders).Select(x => x.AbsolutePath).ToList(), Info);
 			}
 
-			CachedModuleUHTInfo = new UHTModuleInfo 
-			{
-				ModuleName                  = this.Name,
-				ModuleDirectory             = this.ModuleDirectory,
-				ModuleType					= this.Type.ToString(),
-				PublicUObjectClassesHeaders = _AllClassesHeaders    .ToList(),
-				PublicUObjectHeaders        = _PublicUObjectHeaders .ToList(),
-				PrivateUObjectHeaders       = _PrivateUObjectHeaders.ToList(),
-				GeneratedCodeVersion = this.Target.Rules.GetGeneratedCodeVersion()
-			};
-
-			return CachedModuleUHTInfo;
+			return UHTModuleInfoCache;
 		}
-	
-		public override void GetAllDependencyModules( Dictionary<string, UEBuildModule> ReferencedModules, List<UEBuildModule> OrderedModules, bool bIncludeDynamicallyLoaded, bool bForceCircular, bool bOnlyDirectDependencies )
+
+		public override void GetAllDependencyModules( CaselessDictionary<ModuleIndexPair> ReferencedModules, bool bIncludeDynamicallyLoaded, bool bForceCircular, bool bOnlyDirectDependencies )
 		{
 			var AllModuleNames = new List<string>();
 			AllModuleNames.AddRange(PrivateDependencyModuleNames);
@@ -2008,15 +1974,15 @@ namespace UnrealBuildTool
 					if (bForceCircular || !bIsCircular)
 					{
 						var Module = Target.GetModuleByName( DependencyName );
-						ReferencedModules[ DependencyName ] = Module;
+						ReferencedModules[ DependencyName ] = null;
 
 						if( !bOnlyDirectDependencies )
 						{ 
 							// Recurse into dependent modules first
-							Module.GetAllDependencyModules(ReferencedModules, OrderedModules, bIncludeDynamicallyLoaded, bForceCircular, bOnlyDirectDependencies);
+							Module.GetAllDependencyModules(ReferencedModules, bIncludeDynamicallyLoaded, bForceCircular, bOnlyDirectDependencies);
 						}
 
-						OrderedModules.Add( Module );
+						ReferencedModules[ DependencyName ] = new ModuleIndexPair{ Module = Module, Index = ReferencedModules.Where(x => x.Value != null).Count() - 1 };
 					}
 				}
 			}
diff --git a/Engine/Source/Programs/UnrealBuildTool/Configuration/UEBuildTarget.cs b/Engine/Source/Programs/UnrealBuildTool/Configuration/UEBuildTarget.cs
--- a/Engine/Source/Programs/UnrealBuildTool/Configuration/UEBuildTarget.cs
+++ b/Engine/Source/Programs/UnrealBuildTool/Configuration/UEBuildTarget.cs
@@ -8,6 +8,7 @@ using System.Diagnostics;
 using System.IO;
 using System.Xml;
 using System.Runtime.Serialization;
+using Tools.DotNETCommon.CaselessDictionary;
 
 namespace UnrealBuildTool
 {
@@ -188,9 +189,9 @@ namespace UnrealBuildTool
 				OnlyModules = new List<OnlyModule>();
 			}
 
-			if (ModuleCsMap == null)
+			if (FlatModuleCsData == null)
 			{
-				ModuleCsMap = new Dictionary<string,string>();
+				FlatModuleCsData = new CaselessDictionary<FlatModuleCsDataType>();
 			}
 		}
 
@@ -775,10 +776,23 @@ namespace UnrealBuildTool
 
 		/** Used to keep track of all modules by name. */
 		[NonSerialized]
-		private Dictionary<string, UEBuildModule> Modules = new Dictionary<string, UEBuildModule>(StringComparer.InvariantCultureIgnoreCase);
+		private Dictionary<string, UEBuildModule> Modules = new CaselessDictionary<UEBuildModule>();
+
+		[Serializable]
+		public class FlatModuleCsDataType
+		{
+			public FlatModuleCsDataType(string InBuildCsFilename)
+			{
+				BuildCsFilename = InBuildCsFilename;
+			}
+
+			public string       BuildCsFilename;
+			public string       ModuleSourceFolder;
+			public List<string> UHTHeaderNames = new List<string>();
+		}
 
 		/** Used to map names of modules to their .Build.cs filename */
-		private Dictionary<string, string> ModuleCsMap = new Dictionary<string, string>(StringComparer.InvariantCultureIgnoreCase);
+		public CaselessDictionary<FlatModuleCsDataType> FlatModuleCsData = new CaselessDictionary<FlatModuleCsDataType>();
 
 		/** The receipt for this target, which contains a record of this build. */
 		private BuildReceipt Receipt;
@@ -795,9 +809,18 @@ namespace UnrealBuildTool
 		/// A list of the module filenames which were used to build this target.
 		/// </summary>
 		/// <returns></returns>
-		public List<string> GetAllModuleBuildCsFilenames()
+		public IEnumerable<string> GetAllModuleBuildCsFilenames()
+		{
+			return FlatModuleCsData.Values.Select(Data => Data.BuildCsFilename);
+		}
+
+		/// <summary>
+		/// A list of the module filenames which were used to build this target.
+		/// </summary>
+		/// <returns></returns>
+		public IEnumerable<string> GetAllModuleFolders()
 		{
-			return new List<string>(ModuleCsMap.Values);
+			return FlatModuleCsData.Values.SelectMany(Data => Data.UHTHeaderNames);
 		}
 
 		/// <summary>
@@ -1540,52 +1563,58 @@ namespace UnrealBuildTool
 				var UObjectDiscoveryStartTime = DateTime.UtcNow;
 
 				// Reconstruct a full list of binaries. Binaries which aren't compiled are stripped out of AppBinaries, but we still need to scan them for UHT.
-				List<UEBuildBinary> AllAppBinaries = AppBinaries.Union(PrecompiledBinaries).Distinct().ToList();
+				List<UEBuildBinary> AllAppBinaries = AppBinaries.Union(PrecompiledBinaries).ToList();
 	
-				// Figure out which modules have UObjects that we may need to generate headers for
+				var ModulesToGenerateHeadersFor = new HashSet<UEBuildModuleCPP>();
 				foreach( var Binary in AllAppBinaries )
 				{
-					var LocalUObjectModules = UObjectModules;	// For lambda access
 					var DependencyModules = Binary.GetAllDependencyModules(bIncludeDynamicallyLoaded: false, bForceCircular: false);
-					foreach( var DependencyModuleCPP in DependencyModules.OfType<UEBuildModuleCPP>().Where( CPPModule => !LocalUObjectModules.Any( Module => Module.ModuleName == CPPModule.Name ) ) )
+					foreach( var Module in DependencyModules.OfType<UEBuildModuleCPP>() )
 					{
-						if( !DependencyModuleCPP.bIncludedInTarget )
+						if( !Module.bIncludedInTarget )
 						{
-							throw new BuildException( "Expecting module {0} to have bIncludeInTarget set", DependencyModuleCPP.Name );
+							throw new BuildException( "Expecting module {0} to have bIncludeInTarget set", Module.Name );
 						}
 
-						var UHTModuleInfo = DependencyModuleCPP.GetUHTModuleInfo();
-						if( UHTModuleInfo.PublicUObjectClassesHeaders.Count > 0 || UHTModuleInfo.PrivateUObjectHeaders.Count > 0 || UHTModuleInfo.PublicUObjectHeaders.Count > 0 )
+						ModulesToGenerateHeadersFor.Add( Module );
+					}
+				}
+
+				foreach( var Module in ModulesToGenerateHeadersFor )
+				{
+					var UHTModuleInfo = Module.GetCachedUHTModuleInfo();
+					if( UHTModuleInfo.Info.PublicUObjectClassesHeaders.Count > 0 || UHTModuleInfo.Info.PrivateUObjectHeaders.Count > 0 || UHTModuleInfo.Info.PublicUObjectHeaders.Count > 0 )
+					{
+						// If we've got this far and there are no source files then it's likely we're running Rocket and ignoring
+						// engine files, so we don't need a .generated.cpp either
+						UEBuildModuleCPP.AutoGenerateCppInfoClass.BuildInfoClass BuildInfo = null;
+						UHTModuleInfo.Info.GeneratedCPPFilenameBase = Path.Combine( Module.GeneratedCodeDirectory, UHTModuleInfo.Info.ModuleName ) + ".generated";
+						if (Module.SourceFilesToBuild.Count != 0)
 						{
-							// If we've got this far and there are no source files then it's likely we're running Rocket and ignoring
-							// engine files, so we don't need a .generated.cpp either
-							UEBuildModuleCPP.AutoGenerateCppInfoClass.BuildInfoClass BuildInfo = null;
-							UHTModuleInfo.GeneratedCPPFilenameBase = Path.Combine( DependencyModuleCPP.GeneratedCodeDirectory, UHTModuleInfo.ModuleName ) + ".generated";
-							if (DependencyModuleCPP.SourceFilesToBuild.Count != 0)
-							{
-								BuildInfo = new UEBuildModuleCPP.AutoGenerateCppInfoClass.BuildInfoClass( UHTModuleInfo.GeneratedCPPFilenameBase + "*.cpp" );
-							}
+							BuildInfo = new UEBuildModuleCPP.AutoGenerateCppInfoClass.BuildInfoClass( UHTModuleInfo.Info.GeneratedCPPFilenameBase + "*.cpp" );
+						}
 
-							DependencyModuleCPP.AutoGenerateCppInfo = new UEBuildModuleCPP.AutoGenerateCppInfoClass(BuildInfo);
+						Module.AutoGenerateCppInfo = new UEBuildModuleCPP.AutoGenerateCppInfoClass(BuildInfo);
 
-							// If we're running in "gather" mode only, we'll go ahead and cache PCH information for each module right now, so that we don't
-							// have to do it in the assembling phase.  It's OK for gathering to take a bit longer, even if UObject headers are not out of
-							// date in order to save a lot of time in the assembling runs.
-							UHTModuleInfo.PCH = "";
-							if( UnrealBuildTool.IsGatheringBuild && !UnrealBuildTool.IsAssemblingBuild )
+						// If we're running in "gather" mode only, we'll go ahead and cache PCH information for each module right now, so that we don't
+						// have to do it in the assembling phase.  It's OK for gathering to take a bit longer, even if UObject headers are not out of
+						// date in order to save a lot of time in the assembling runs.
+						UHTModuleInfo.Info.PCH = "";
+						if( UnrealBuildTool.IsGatheringBuild && !UnrealBuildTool.IsAssemblingBuild )
+						{
+							// We need to figure out which PCH header this module is including, so that UHT can inject an include statement for it into any .cpp files it is synthesizing
+							var ModuleCompileEnvironment = Module.CreateModuleCompileEnvironment(GlobalCompileEnvironment);
+							Module.CachePCHUsageForModuleSourceFiles(ModuleCompileEnvironment);
+							if (Module.ProcessedDependencies.UniquePCHHeaderFile != null)
 							{
-								// We need to figure out which PCH header this module is including, so that UHT can inject an include statement for it into any .cpp files it is synthesizing
-								var ModuleCompileEnvironment = DependencyModuleCPP.CreateModuleCompileEnvironment(GlobalCompileEnvironment);
-								DependencyModuleCPP.CachePCHUsageForModuleSourceFiles(ModuleCompileEnvironment);
-								if (DependencyModuleCPP.ProcessedDependencies.UniquePCHHeaderFile != null)
-								{
-									UHTModuleInfo.PCH = DependencyModuleCPP.ProcessedDependencies.UniquePCHHeaderFile.AbsolutePath;
-								}
+								UHTModuleInfo.Info.PCH = Module.ProcessedDependencies.UniquePCHHeaderFile.AbsolutePath;
 							}
-
-							UObjectModules.Add( UHTModuleInfo );
-							Log.TraceVerbose( "Detected UObject module: " + UHTModuleInfo.ModuleName );
 						}
+
+						UObjectModules.Add( UHTModuleInfo.Info );
+						FlatModuleCsData[Module.Name].ModuleSourceFolder = Module.ModuleDirectory;
+						FlatModuleCsData[Module.Name].UHTHeaderNames     = UHTModuleInfo.HeaderFilenames.ToList();
+						Log.TraceVerbose( "Detected UObject module: " + UHTModuleInfo.Info.ModuleName );
 					}
 				}
 				
@@ -2145,10 +2174,9 @@ namespace UnrealBuildTool
 								var SharedPCHHeaderFileItem = FileItem.GetExistingItemByPath( SharedPCHHeaderFilePath );
 								if( SharedPCHHeaderFileItem != null )
 								{
-									var ModuleDependencies = new Dictionary<string, UEBuildModule>();
-									var ModuleList = new List<UEBuildModule>();
+									var ModuleDependencies = new CaselessDictionary<UEBuildModule.ModuleIndexPair>();
 									bool bIncludeDynamicallyLoaded = false;
-									CPPModule.GetAllDependencyModules(ModuleDependencies, ModuleList, bIncludeDynamicallyLoaded, bForceCircular: false, bOnlyDirectDependencies:false);
+									CPPModule.GetAllDependencyModules(ModuleDependencies, bIncludeDynamicallyLoaded, bForceCircular: false, bOnlyDirectDependencies:false);
 
 									// Figure out where to insert the shared PCH into our list, based off the module dependency ordering
 									int InsertAtIndex = SharedPCHHeaderFiles.Count;
@@ -2183,7 +2211,7 @@ namespace UnrealBuildTool
 									var NewSharedPCHHeaderInfo = new SharedPCHHeaderInfo();
 									NewSharedPCHHeaderInfo.PCHHeaderFile = SharedPCHHeaderFileItem;
 									NewSharedPCHHeaderInfo.Module = CPPModule;
-									NewSharedPCHHeaderInfo.Dependencies = ModuleDependencies;
+									NewSharedPCHHeaderInfo.Dependencies = ModuleDependencies.Values.OrderBy(x => x.Index).Select(x => x.Module).ToDictionary(M => M.Name);
 									SharedPCHHeaderFiles.Insert( InsertAtIndex, NewSharedPCHHeaderInfo );
 								}
 								else
@@ -2932,7 +2960,7 @@ namespace UnrealBuildTool
 		{
 			Debug.Assert(Module.Target == this);
 			Modules.Add(Module.Name, Module);
-			ModuleCsMap.Add(Module.Name, Module.BuildCsFilename);
+			FlatModuleCsData.Add(Module.Name, new FlatModuleCsDataType(Module.BuildCsFilename));
 		}
 
 		/** Finds a module given its name.  Throws an exception if the module couldn't be found. */
diff --git a/Engine/Source/Programs/UnrealBuildTool/System/ActionGraph.cs b/Engine/Source/Programs/UnrealBuildTool/System/ActionGraph.cs
--- a/Engine/Source/Programs/UnrealBuildTool/System/ActionGraph.cs
+++ b/Engine/Source/Programs/UnrealBuildTool/System/ActionGraph.cs
@@ -1297,8 +1297,7 @@ namespace UnrealBuildTool
 
 				if( ExpandCPPHeaderDependencies && bIsCPPFile )
 				{
-					bool HasUObjects;
-					List<DependencyInclude> DirectlyIncludedFilenames = CPPEnvironment.GetDirectIncludeDependencies( Target, FileItem, BuildPlatform, bOnlyCachedDependencies:false, HasUObjects:out HasUObjects );
+					List<DependencyInclude> DirectlyIncludedFilenames = CPPEnvironment.GetDirectIncludeDependencies( Target, FileItem, BuildPlatform, bOnlyCachedDependencies:false );
 
 					// Resolve the included file name to an actual file.
 					var DirectlyIncludedFiles =
diff --git a/Engine/Source/Programs/UnrealBuildTool/System/CPPHeaders.cs b/Engine/Source/Programs/UnrealBuildTool/System/CPPHeaders.cs
--- a/Engine/Source/Programs/UnrealBuildTool/System/CPPHeaders.cs
+++ b/Engine/Source/Programs/UnrealBuildTool/System/CPPHeaders.cs
@@ -6,6 +6,7 @@ using System.Diagnostics;
 using System.Text;
 using System.Text.RegularExpressions;
 using System.IO;
+using Tools.DotNETCommon.FileContentsCacheType;
 
 namespace UnrealBuildTool
 {
@@ -258,8 +259,7 @@ namespace UnrealBuildTool
 				IncludedFilesMap.Add(CPPFile, IncludedFileList);
 
 				// Gather a list of names of files directly included by this C++ file.
-				bool HasUObjects;
-				List<DependencyInclude> DirectIncludes = GetDirectIncludeDependencies( Target, CPPFile, BuildPlatform, bOnlyCachedDependencies:bOnlyCachedDependencies, HasUObjects:out HasUObjects );
+				List<DependencyInclude> DirectIncludes = GetDirectIncludeDependencies( Target, CPPFile, BuildPlatform, bOnlyCachedDependencies:bOnlyCachedDependencies );
 
 				// Build a list of the unique set of files that are included by this file.
 				var DirectlyIncludedFiles = new HashSet<FileItem>();
@@ -374,18 +374,30 @@ namespace UnrealBuildTool
 		/** Regex that matches C++ code with UObject declarations which we will need to generated code for. */
 		static readonly Regex UObjectRegex = new Regex("^\\s*U(CLASS|STRUCT|ENUM|INTERFACE|DELEGATE)\\b", RegexOptions.Compiled | RegexOptions.Multiline );
 
-		/** Finds the names of files directly included by the given C++ file. */
-		public static List<DependencyInclude> GetDirectIncludeDependencies( UEBuildTarget Target, FileItem CPPFile, IUEBuildPlatform BuildPlatform, bool bOnlyCachedDependencies, out bool HasUObjects )
+		// Maintains a cache of file contents
+		private static FileContentsCacheType FileContentsCache = new FileContentsCacheType();
+
+		// Checks if a file contains UObjects
+		public static bool DoesFileContainUObjects(string Filename)
+		{
+			string Contents = FileContentsCache.GetContents(Filename);
+			return UObjectRegex.IsMatch(Contents);
+		}
+
+		/** Finds the names of files directly included by the given C++ file, and also whether the file contains any UObjects */
+		public static List<DependencyInclude> GetDirectIncludeDependencies( UEBuildTarget Target, FileItem CPPFile, IUEBuildPlatform BuildPlatform, bool bOnlyCachedDependencies )
 		{
 			// Try to fulfill request from cache first.
-			List<DependencyInclude> Result;
-			if( IncludeDependencyCache[Target].GetCachedDirectDependencies( CPPFile, out Result, out HasUObjects ) )
+			List<DependencyInclude> Info = IncludeDependencyCache[Target].GetCachedDependencyInfo( CPPFile );
+			if( Info != null )
 			{
-				return Result;
+				return Info;
 			}
-			else if( bOnlyCachedDependencies )
+
+			var Result = new List<DependencyInclude>();
+			if( bOnlyCachedDependencies )
 			{
-				return new List<DependencyInclude>();
+				return Result;
 			}
 
 			var TimerStartTime = DateTime.UtcNow;
@@ -393,21 +405,17 @@ namespace UnrealBuildTool
 
 			// Get the adjusted filename
 			string FileToRead = CPPFile.AbsolutePath;
-
-			if (BuildPlatform.RequiresExtraUnityCPPWriter() == true &&
-				Path.GetFileName(FileToRead).StartsWith("Module."))
+			if (BuildPlatform.RequiresExtraUnityCPPWriter() && Path.GetFileName(FileToRead).StartsWith("Module."))
 			{
 				FileToRead += ".ex";
 			}
 
 			// Read lines from the C++ file.
-			string FileContents = Utils.ReadAllText( FileToRead );
-            if (string.IsNullOrEmpty(FileContents))
-            {
-                return new List<DependencyInclude>();
-            }
-
-			HasUObjects = CPPEnvironment.UObjectRegex.IsMatch(FileContents);
+			string FileContents = FileContentsCache.GetContents(FileToRead);
+			if (string.IsNullOrEmpty(FileContents))
+			{
+				return Result;
+			}
 
 			// Note: This depends on UBT executing w/ a working directory of the Engine/Source folder!
 			string EngineSourceFolder = Directory.GetCurrentDirectory();
@@ -418,7 +426,6 @@ namespace UnrealBuildTool
 				InstalledFolder = EngineSourceFolder.Substring(0, EngineSourceIdx);
 			}
 
-			Result = new List<DependencyInclude>();
 			if (Utils.IsRunningOnMono)
 			{
 				// Mono crashes when running a regex on a string longer than about 5000 characters, so we parse the file in chunks
@@ -432,18 +439,18 @@ namespace UnrealBuildTool
 						EndIndex = FileContents.Length;
 					}
 
-					CollectHeaders(CPPFile, Result, FileToRead, FileContents, InstalledFolder, StartIndex, EndIndex);
+					Result.AddRange(CollectHeaders(CPPFile, FileToRead, FileContents, InstalledFolder, StartIndex, EndIndex));
 
 					StartIndex = EndIndex + 1;
 				}
 			}
 			else
 			{
-				CollectHeaders(CPPFile, Result, FileToRead, FileContents, InstalledFolder, 0, FileContents.Length);
+				Result = CollectHeaders(CPPFile, FileToRead, FileContents, InstalledFolder, 0, FileContents.Length);
 			}
 
 			// Populate cache with results.
-			IncludeDependencyCache[Target].SetDependencyInfo( CPPFile, Result, HasUObjects );
+			IncludeDependencyCache[Target].SetDependencyInfo( CPPFile, Result );
 
 			CPPEnvironment.DirectIncludeCacheMissesTotalTime += ( DateTime.UtcNow - TimerStartTime ).TotalSeconds;
 
@@ -454,18 +461,19 @@ namespace UnrealBuildTool
 		/// Collects all header files included in a CPPFile
 		/// </summary>
 		/// <param name="CPPFile"></param>
-		/// <param name="Result"></param>
 		/// <param name="FileToRead"></param>
 		/// <param name="FileContents"></param>
 		/// <param name="InstalledFolder"></param>
 		/// <param name="StartIndex"></param>
 		/// <param name="EndIndex"></param>
-		private static void CollectHeaders(FileItem CPPFile, List<DependencyInclude> Result, string FileToRead, string FileContents, string InstalledFolder, int StartIndex, int EndIndex)
+		private static List<DependencyInclude> CollectHeaders(FileItem CPPFile, string FileToRead, string FileContents, string InstalledFolder, int StartIndex, int EndIndex)
 		{
+			var Result = new List<DependencyInclude>();
+
 			Match M = CPPHeaderRegex.Match(FileContents, StartIndex, EndIndex - StartIndex);
-			CaptureCollection CC = M.Groups["HeaderFile"].Captures;
-			Result.Capacity = Math.Max(Result.Count + CC.Count, Result.Capacity);
-			foreach (Capture C in CC)
+			CaptureCollection Captures = M.Groups["HeaderFile"].Captures;
+			Result.Capacity = Result.Count;
+			foreach (Capture C in Captures)
 			{
 				string HeaderValue = C.Value;
 
@@ -475,18 +483,10 @@ namespace UnrealBuildTool
 				}
 
 				//@TODO: The intermediate exclusion is to work around autogenerated absolute paths in Module.SomeGame.cpp style files
-				bool bIsIntermediateOrThirdParty = FileToRead.Contains("Intermediate") || FileToRead.Contains("ThirdParty");
-				bool bCheckForBackwardSlashes = FileToRead.StartsWith(InstalledFolder);
-				if (UnrealBuildTool.HasUProjectFile())
+				bool bCheckForBackwardSlashes = FileToRead.StartsWith(InstalledFolder) || (UnrealBuildTool.HasUProjectFile() && Utils.IsFileUnderDirectory(FileToRead, UnrealBuildTool.GetUProjectPath()));
+				if (bCheckForBackwardSlashes && !FileToRead.Contains("Intermediate") && !FileToRead.Contains("ThirdParty") && HeaderValue.IndexOf('\\', 0) >= 0)
 				{
-					bCheckForBackwardSlashes |= Utils.IsFileUnderDirectory(FileToRead, UnrealBuildTool.GetUProjectPath());
-				}
-				if (bCheckForBackwardSlashes && !bIsIntermediateOrThirdParty)
-				{
-					if (HeaderValue.IndexOf('\\', 0) >= 0)
-					{
-						throw new BuildException("In {0}: #include \"{1}\" contains backslashes ('\\'), please use forward slashes ('/') instead.", FileToRead, C.Value);
-					}
+					throw new BuildException("In {0}: #include \"{1}\" contains backslashes ('\\'), please use forward slashes ('/') instead.", FileToRead, C.Value);
 				}
 				HeaderValue = Utils.CleanDirectorySeparators(HeaderValue);
 				Result.Add(new DependencyInclude(HeaderValue));
@@ -497,13 +497,15 @@ namespace UnrealBuildTool
 			if (Ext == ".MM" || Ext == ".M")
 			{
 				M = MMHeaderRegex.Match(FileContents, StartIndex, EndIndex - StartIndex);
-				CC = M.Groups["HeaderFile"].Captures;
-				Result.Capacity += CC.Count;
-				foreach (Capture C in CC)
+				Captures = M.Groups["HeaderFile"].Captures;
+				Result.Capacity += Captures.Count;
+				foreach (Capture C in Captures)
 				{
 					Result.Add(new DependencyInclude(C.Value));
 				}
 			}
+
+			return Result;
 		}
 	}
 }
diff --git a/Engine/Source/Programs/UnrealBuildTool/System/DependencyCache.cs b/Engine/Source/Programs/UnrealBuildTool/System/DependencyCache.cs
--- a/Engine/Source/Programs/UnrealBuildTool/System/DependencyCache.cs
+++ b/Engine/Source/Programs/UnrealBuildTool/System/DependencyCache.cs
@@ -36,19 +36,6 @@ namespace UnrealBuildTool
 		}
 	}
 
-	[Serializable]
-	struct DependencyInfo
-	{
-		/** List of files this file includes.  These are direct include paths and cannot be resolved to an actual file on 
-		    disk without using the proper list of include directories for this file's module */
-		public List<DependencyInclude> Includes;
-
-		/** True if this file was found to contain UObject class or type definitions, and the file needs to be preprocessed
-		    by UnrealHeaderTool to generate reflection code */
-		public bool HasUObjects;
-	}
-
-
 	/**
 	 * Caches include dependency information to speed up preprocessing on subsequent runs.
 	 */
@@ -67,7 +54,7 @@ namespace UnrealBuildTool
 		private string CachePath;
 
 		/** Dependency lists, keyed (case-insensitively) on file's absolute path. */
-		private Dictionary<string, DependencyInfo> DependencyMap;
+		private Dictionary<string,  List<DependencyInclude>> DependencyMap;
 
 		[NonSerialized]
 		private Dictionary<string, bool> FileExistsInfo;
@@ -157,7 +144,7 @@ namespace UnrealBuildTool
 			CacheCreateDate = DateTimeOffset.Now;
 			CacheUpdateDate = DateTimeOffset.Now;
 			CachePath = Cache.AbsolutePath;
-			DependencyMap = new Dictionary<string, DependencyInfo>();
+			DependencyMap = new Dictionary<string, List<DependencyInclude>>();
 			bIsDirty = false;
 			CreateFileExistsInfo();
 		}
@@ -205,96 +192,58 @@ namespace UnrealBuildTool
 			}
 		}
 
-
-		/// <summary>
-		/// Gets information about this file from our dependency cache.  Only returns information if the file has already been 
-		/// cached, and the file has not been changed since the last time we updated our cache
-		/// </summary>
-		/// <param name="File">The file to check</param>
-		/// <returns>Information about this dependency, or null if we have nothing cached</returns>
-		DependencyInfo? GetCachedDependencyInfo( FileItem File )
+		/**
+		 * Returns the direct dependencies of the specified FileItem if it exists in the cache and they are not stale.
+		 *
+		 * @param  File  File to try to find dependencies in cache
+		 * @returns  Dependency info for File, or null if no dependencies are cached or if the cache is stale.
+		 */
+		public List<DependencyInclude> GetCachedDependencyInfo(FileItem File)
 		{
+			string LowercaseFilePath = File.AbsolutePath.ToLowerInvariant();
+
 			// Check whether File is in cache.
-			DependencyInfo DependencyInfo;
-			if (DependencyMap.TryGetValue(File.AbsolutePath.ToLowerInvariant(), out DependencyInfo))
+			List<DependencyInclude> Result;
+			if (!DependencyMap.TryGetValue(LowercaseFilePath, out Result))
 			{
-				// File is in cache, now check whether last write time is prior to cache creation time.
-				if (File.LastWriteTime < CacheCreateDate)
-				{
-					return DependencyInfo;					
-				}
-				else
-				{
-					// Remove entry from cache as it's stale.
-					RemoveFileFromCache(File.AbsolutePath);
-					return null;
-				}
+				return null;
 			}
 
-			return null;
-		}
-
-		/**
-		 * Returns the direct dependencies of the specified FileItem if it exists in the cache and if the
-		 * file has a last write time before the creation time of the cache. 
-		 * 
-		 * The code also keeps track of whether dependencies have been successfully accessed for a given
-		 * file.
-		 * 
-		 * @param	File				File to try to find dependencies in cache
-		 * @param	Result	[out]		List of dependencies if successful, null otherwise
-		 * @param	HasUObjects			True if the file was found to have UObject classes, otherwise false
-		 */
-		public bool GetCachedDirectDependencies(FileItem File, out List<DependencyInclude> Result, out bool HasUObjects)
-		{
-			Result = null;
-			HasUObjects = false;
+			// File is in cache, now check whether last write time is prior to cache creation time.
+			if (File.LastWriteTime >= CacheCreateDate)
+			{
+				// Remove entry from cache as it's stale.
+				DependencyMap.Remove(LowercaseFilePath);
+				bIsDirty = true;
+				return null;
+			}
 
-			// Check whether File is in cache.
-			DependencyInfo? DependencyInfo = GetCachedDependencyInfo(File);
-			if( DependencyInfo != null )
+			// Check if any of the resolved includes is missing
+			foreach (var Include in Result)
 			{
-				// Check if any of the resolved includes is missing
-				foreach (var Include in DependencyInfo.Value.Includes)
+				if (!String.IsNullOrEmpty(Include.IncludeResolvedName))
 				{
-					if (!String.IsNullOrEmpty(Include.IncludeResolvedName))
+					bool bIncludeExists = false;
+					string LowercaseIncludePath = Include.IncludeResolvedName.ToLowerInvariant();
+					if (!FileExistsInfo.TryGetValue(LowercaseIncludePath, out bIncludeExists))
 					{
-						bool bIncludeExists = false;
-						string FileExistsKey = Include.IncludeResolvedName.ToLowerInvariant();
-						if (FileExistsInfo.TryGetValue(FileExistsKey, out bIncludeExists) == false)
-						{
-							bIncludeExists = System.IO.File.Exists(Include.IncludeResolvedName);
-							FileExistsInfo.Add(FileExistsKey, bIncludeExists);
-						}
-						if (!bIncludeExists)
-						{
-							// Remove entry from cache as it's stale, as well as the include which no longer exists
-							RemoveFileFromCache(Include.IncludeResolvedName);
-							RemoveFileFromCache(File.AbsolutePath);							
-							return false;
-						}
-					}						
+						bIncludeExists = System.IO.File.Exists(Include.IncludeResolvedName);
+						FileExistsInfo.Add(LowercaseIncludePath, bIncludeExists);
+					}
+
+					if (!bIncludeExists)
+					{
+						// Remove entry from cache as it's stale, as well as the include which no longer exists
+						DependencyMap.Remove(LowercaseIncludePath);
+						DependencyMap.Remove(LowercaseFilePath);
+						bIsDirty = true;
+						return null;
+					}
 				}
-				// Cached version is up to date, return it.
-				Result = DependencyInfo.Value.Includes;
-				HasUObjects = DependencyInfo.Value.HasUObjects;
-				return true;
 			}
-			// Not in cache.
-			else
-			{
-				return false;
-			}
-		}
 
-		/// <summary>
-		/// Removes file entry from dependency cache.
-		/// </summary>
-		/// <param name="AbsolutePath"></param>
-		private void RemoveFileFromCache(string AbsolutePath)
-		{
-			DependencyMap.Remove(AbsolutePath.ToLowerInvariant());
-			bIsDirty = true;
+			// Cached version is up to date, return it.
+			return Result;
 		}
 
 		/**
@@ -304,9 +253,9 @@ namespace UnrealBuildTool
 		 * @param	Dependencies	List of dependencies to cache for passed in file
 		 * @param	HasUObjects		True if this file was found to contain UObject classes or types
 		 */
-		public void SetDependencyInfo(FileItem File, List<DependencyInclude> DirectDependencies, bool HasUObjects)
+		public void SetDependencyInfo(FileItem File, List<DependencyInclude> Info)
 		{
-			DependencyMap[File.AbsolutePath.ToLowerInvariant()] = new DependencyInfo() { Includes = DirectDependencies, HasUObjects = HasUObjects };
+			DependencyMap[File.AbsolutePath.ToLowerInvariant()] = Info;
 			bIsDirty = true;
 		}
 
@@ -327,7 +276,7 @@ namespace UnrealBuildTool
 		{
 			foreach (var Dependency in DependencyMap)
 			{
-				foreach (var Include in Dependency.Value.Includes)
+				foreach (var Include in Dependency.Value)
 				{
 					if (Include.IncludeResolvedName == String.Empty)
 					{
@@ -349,7 +298,7 @@ namespace UnrealBuildTool
 		{
 			if (BuildConfiguration.bUseIncludeDependencyResolveCache)
 			{
-				var Includes = DependencyMap[File.AbsolutePath.ToLowerInvariant()].Includes;
+				var Includes = DependencyMap[File.AbsolutePath.ToLowerInvariant()];
 				var IncludeToResolve = Includes[DirectlyIncludedFileNameIndex];
 				if (BuildConfiguration.bTestIncludeDependencyResolveCache)
 				{
diff --git a/Engine/Source/Programs/UnrealBuildTool/System/ExternalExecution.cs b/Engine/Source/Programs/UnrealBuildTool/System/ExternalExecution.cs
--- a/Engine/Source/Programs/UnrealBuildTool/System/ExternalExecution.cs
+++ b/Engine/Source/Programs/UnrealBuildTool/System/ExternalExecution.cs
@@ -152,7 +152,56 @@ namespace UnrealBuildTool
 	 */
 	public class ExternalExecution
 	{
+		/// <summary>
+		/// Generates a UHTModuleInfo for a particular named module under a directory.
+		/// </summary>
+		/// <returns>
+		public static UHTModuleInfo CreateUHTModuleInfo(IEnumerable<string> HeaderFilenames, UEBuildTarget Target, string ModuleName, string ModuleDirectory, UEBuildModuleType ModuleType)
+		{
+			var ClassesFolder = Path.Combine(ModuleDirectory, "Classes");
+			var PublicFolder  = Path.Combine(ModuleDirectory, "Public");
+			var BuildPlatform = UEBuildPlatform.GetBuildPlatform(Target.Platform);
+
+			var AllClassesHeaders     = new List<FileItem>();
+			var PublicUObjectHeaders  = new List<FileItem>();
+			var PrivateUObjectHeaders = new List<FileItem>();
 
+			foreach (var Header in HeaderFilenames)
+			{
+				// Check to see if we know anything about this file.  If we have up-to-date cached information about whether it has
+				// UObjects or not, we can skip doing a test here.
+				var UObjectHeaderFileItem = FileItem.GetExistingItemByPath( Header );
+
+				if (CPPEnvironment.DoesFileContainUObjects(UObjectHeaderFileItem.AbsolutePath))
+				{
+					if (UObjectHeaderFileItem.AbsolutePath.StartsWith(ClassesFolder))
+					{
+						AllClassesHeaders.Add(UObjectHeaderFileItem);
+					}
+					else if (UObjectHeaderFileItem.AbsolutePath.StartsWith(PublicFolder))
+					{
+						PublicUObjectHeaders.Add(UObjectHeaderFileItem);
+					}
+					else
+					{
+						PrivateUObjectHeaders.Add(UObjectHeaderFileItem);
+					}
+				}
+			}
+
+			var Result = new UHTModuleInfo
+			{
+				ModuleName                  = ModuleName,
+				ModuleDirectory             = ModuleDirectory,
+				ModuleType                  = ModuleType.ToString(),
+				PublicUObjectClassesHeaders = AllClassesHeaders,
+				PublicUObjectHeaders        = PublicUObjectHeaders,
+				PrivateUObjectHeaders       = PrivateUObjectHeaders,
+				GeneratedCodeVersion        = Target.Rules.GetGeneratedCodeVersion()
+			};
+
+			return Result;
+		}
 
 		static ExternalExecution()
 		{
diff --git a/Engine/Source/Programs/UnrealBuildTool/System/UnrealBuildTool.cs b/Engine/Source/Programs/UnrealBuildTool/System/UnrealBuildTool.cs
--- a/Engine/Source/Programs/UnrealBuildTool/System/UnrealBuildTool.cs
+++ b/Engine/Source/Programs/UnrealBuildTool/System/UnrealBuildTool.cs
@@ -1659,10 +1659,10 @@ namespace UnrealBuildTool
                     // If we're generating project files, then go ahead and wipe out the existing UBTMakefile for every target, to make sure that
                     // it gets a full dependency scan next time.
                     // NOTE: This is just a safeguard and doesn't have to be perfect.  We also check for newer project file timestamps in LoadUBTMakefile()
+                    string UBTMakefilePath = UnrealBuildTool.GetUBTMakefilePath( TargetDescs );
                     if( ProjectFileGenerator.bGenerateProjectFiles )	// @todo ubtmake: This is only hit when generating IntelliSense for project files.  Probably should be done right inside ProjectFileGenerator.bat
                     {													// @todo ubtmake: Won't catch multi-target cases as GPF always builds one target at a time for Intellisense
                         // Delete the UBTMakefile
-                        var UBTMakefilePath = UnrealBuildTool.GetUBTMakefilePath( TargetDescs );
                         if (File.Exists(UBTMakefilePath))
                         {
                             UEBuildTarget.CleanFile(UBTMakefilePath);
@@ -1684,8 +1684,7 @@ namespace UnrealBuildTool
 
                         // Try to load the UBTMakefile.  It will only be loaded if it has valid content and is not determined to be out of date.    
 						string ReasonNotLoaded;
-                        UBTMakefile = LoadUBTMakefile( TargetDescs, out ReasonNotLoaded );
-
+                        UBTMakefile = LoadUBTMakefile( UBTMakefilePath, out ReasonNotLoaded );
                         if( UBTMakefile == null )
                         { 
                             // If the Makefile couldn't be loaded, then we're not going to be able to continue in "assembler only" mode.  We'll do both
@@ -2271,14 +2270,14 @@ namespace UnrealBuildTool
         /// <summary>
         /// Loads a UBTMakefile from disk
         /// </summary>
-        /// <param name="TargetDescs">List of targets.  Order is not important</param>
+        /// <param name="MakefilePath">Path to the makefile to load</param>
 		/// <param name="ReasonNotLoaded">If the function returns null, this string will contain the reason why</param>
 		/// <returns>The loaded makefile, or null if it failed for some reason.  On failure, the 'ReasonNotLoaded' variable will contain information about why</returns>
-		static UBTMakefile LoadUBTMakefile( List<TargetDescriptor> TargetDescs, out string ReasonNotLoaded )
+		static UBTMakefile LoadUBTMakefile( string MakefilePath, out string ReasonNotLoaded )
 		{
 			// Check the directory timestamp on the project files directory.  If the user has generated project files more
 			// recently than the UBTMakefile, then we need to consider the file to be out of date
-			var UBTMakefileInfo = new FileInfo( GetUBTMakefilePath( TargetDescs ) );
+			var UBTMakefileInfo = new FileInfo( MakefilePath );
 			if( !UBTMakefileInfo.Exists )
 			{
 				// UBTMakefile doesn't even exist, so we won't bother loading it
@@ -2392,7 +2391,7 @@ namespace UnrealBuildTool
 					}
 				}
 
-				List<string> BuildCsFilenames = Target.GetAllModuleBuildCsFilenames();
+				IEnumerable<string> BuildCsFilenames = Target.GetAllModuleBuildCsFilenames();
 				foreach (var BuildCsFilename in BuildCsFilenames)
 				{
 					if (BuildCsFilename != null)
@@ -2409,6 +2408,39 @@ namespace UnrealBuildTool
 				}
 			}
 
+			// We do a check to see if any modules' headers have changed which have
+			// acquired or lost UHT types.  If so, which should be rare,
+			// we'll just invalidate the entire makefile and force it to be rebuilt.
+			foreach (var Target in LoadedUBTMakefile.Targets)
+			{
+				// Get all H files in processed modules newer than the makefile itself
+				var HFilesNewerThanMakefile =
+					new HashSet<string>(
+						Target.FlatModuleCsData
+						.SelectMany(x => x.Value.ModuleSourceFolder != null ? Directory.EnumerateFiles(x.Value.ModuleSourceFolder, "*.h", SearchOption.AllDirectories) : Enumerable.Empty<string>())
+						.Where(y => Directory.GetLastWriteTime(y) > UBTMakefileInfo.LastWriteTime)
+						.OrderBy(z => z).Distinct()
+					);
+
+				// Get all H files in all modules processed in the last makefile build
+				var AllUHTHeaders = new HashSet<string>(Target.FlatModuleCsData.Select(x => x.Value).SelectMany(x => x.UHTHeaderNames));
+
+				// Makefile is invalid if:
+				// * There are any newer files which contain no UHT data, but were previously in the makefile
+				// * There are any newer files contain data which needs processing by UHT, but weren't not previously in the makefile
+				foreach (string Filename in HFilesNewerThanMakefile)
+				{
+					bool bContainsUHTData = CPPEnvironment.DoesFileContainUObjects(Filename);
+					bool bWasProcessed    = AllUHTHeaders.Contains(Filename);
+					if (bContainsUHTData != bWasProcessed)
+					{
+						Log.TraceVerbose("{0} {1} contain UHT types and now {2} , ignoring it ({3})", Filename, bWasProcessed ? "used to" : "didn't", bWasProcessed ? "doesn't" : "does", UBTMakefileInfo.FullName);
+						ReasonNotLoaded = string.Format("List of files containing UHT types has changed");
+						return null;
+					}
+				}
+			}
+
 			ReasonNotLoaded = null;
 			return LoadedUBTMakefile;
 		}
