DependencyInfo restored to fix deserialization errors during project generation.

diff --git a/Engine/Source/Programs/UnrealBuildTool/System/DependencyCache.cs b/Engine/Source/Programs/UnrealBuildTool/System/DependencyCache.cs
--- a/Engine/Source/Programs/UnrealBuildTool/System/DependencyCache.cs
+++ b/Engine/Source/Programs/UnrealBuildTool/System/DependencyCache.cs
@@ -36,6 +36,18 @@ namespace UnrealBuildTool
 		}
 	}
 
+	[Serializable]
+	struct DependencyInfo
+	{
+		/** List of files this file includes.  These are direct include paths and cannot be resolved to an actual file on 
+		    disk without using the proper list of include directories for this file's module */
+		public List<DependencyInclude> Includes;
+
+		/** Deprecated - always false */
+		public bool HasUObjects;
+	}
+
+
 	/**
 	 * Caches include dependency information to speed up preprocessing on subsequent runs.
 	 */
@@ -54,7 +66,7 @@ namespace UnrealBuildTool
 		private string CachePath;
 
 		/** Dependency lists, keyed (case-insensitively) on file's absolute path. */
-		private Dictionary<string,  List<DependencyInclude>> DependencyMap;
+		private Dictionary<string, DependencyInfo> DependencyMap;
 
 		[NonSerialized]
 		private Dictionary<string, bool> FileExistsInfo;
@@ -144,7 +156,7 @@ namespace UnrealBuildTool
 			CacheCreateDate = DateTimeOffset.Now;
 			CacheUpdateDate = DateTimeOffset.Now;
 			CachePath = Cache.AbsolutePath;
-			DependencyMap = new Dictionary<string, List<DependencyInclude>>();
+			DependencyMap = new Dictionary<string, DependencyInfo>();
 			bIsDirty = false;
 			CreateFileExistsInfo();
 		}
@@ -203,7 +215,7 @@ namespace UnrealBuildTool
 			string LowercaseFilePath = File.AbsolutePath.ToLowerInvariant();
 
 			// Check whether File is in cache.
-			List<DependencyInclude> Result;
+			DependencyInfo Result;
 			if (!DependencyMap.TryGetValue(LowercaseFilePath, out Result))
 			{
 				return null;
@@ -219,7 +231,7 @@ namespace UnrealBuildTool
 			}
 
 			// Check if any of the resolved includes is missing
-			foreach (var Include in Result)
+			foreach (var Include in Result.Includes)
 			{
 				if (!String.IsNullOrEmpty(Include.IncludeResolvedName))
 				{
@@ -243,7 +255,7 @@ namespace UnrealBuildTool
 			}
 
 			// Cached version is up to date, return it.
-			return Result;
+			return Result.Includes;
 		}
 
 		/**
@@ -255,7 +267,7 @@ namespace UnrealBuildTool
 		 */
 		public void SetDependencyInfo(FileItem File, List<DependencyInclude> Info)
 		{
-			DependencyMap[File.AbsolutePath.ToLowerInvariant()] = Info;
+			DependencyMap[File.AbsolutePath.ToLowerInvariant()] = new DependencyInfo{ Includes = Info, HasUObjects = false };
 			bIsDirty = true;
 		}
 
@@ -276,7 +288,7 @@ namespace UnrealBuildTool
 		{
 			foreach (var Dependency in DependencyMap)
 			{
-				foreach (var Include in Dependency.Value)
+				foreach (var Include in Dependency.Value.Includes)
 				{
 					if (Include.IncludeResolvedName == String.Empty)
 					{
@@ -298,7 +310,7 @@ namespace UnrealBuildTool
 		{
 			if (BuildConfiguration.bUseIncludeDependencyResolveCache)
 			{
-				var Includes = DependencyMap[File.AbsolutePath.ToLowerInvariant()];
+				var Includes = DependencyMap[File.AbsolutePath.ToLowerInvariant()].Includes;
 				var IncludeToResolve = Includes[DirectlyIncludedFileNameIndex];
 				if (BuildConfiguration.bTestIncludeDependencyResolveCache)
 				{
