Use a response file instead of spawning a process with a really long command line. #html5

diff --git a/Engine/Source/Programs/UnrealBuildTool/HTML5/HTML5ToolChain.cs b/Engine/Source/Programs/UnrealBuildTool/HTML5/HTML5ToolChain.cs
--- a/Engine/Source/Programs/UnrealBuildTool/HTML5/HTML5ToolChain.cs
+++ b/Engine/Source/Programs/UnrealBuildTool/HTML5/HTML5ToolChain.cs
@@ -469,19 +469,23 @@ namespace UnrealBuildTool
 			// Make the final javascript file
 			Action LinkAction = new Action(ActionType.Link);
 
+			// ResponseFile lines. 
+			List<string> ReponseLines = new List<string>(); 
+
 			LinkAction.bCanExecuteRemotely = false;
 			LinkAction.WorkingDirectory = Path.GetFullPath(".");
 			LinkAction.CommandPath = PythonPath;
 			LinkAction.CommandArguments = EMCCPath;
-		    LinkAction.CommandArguments += GetLinkArguments(LinkEnvironment);
+			ReponseLines.Add(GetLinkArguments(LinkEnvironment));
 
 			// Add the input files to a response file, and pass the response file on the command-line.
 			foreach (FileItem InputFile in LinkEnvironment.InputFiles)
 			{
                 //System.Console.WriteLine("File  {0} ", InputFile.AbsolutePath);
-                LinkAction.CommandArguments += string.Format(" \"{0}\"", InputFile.AbsolutePath);
+                ReponseLines.Add(string.Format(" \"{0}\"", InputFile.AbsolutePath));
 				LinkAction.PrerequisiteItems.Add(InputFile);
 			}
+
 			if (!LinkEnvironment.Config.bIsBuildingLibrary)
 			{
                     // Make sure ThirdParty libs are at the end. 
@@ -501,24 +505,30 @@ namespace UnrealBuildTool
 
                         if (Item != null)
                         {
-                            if (Item.ToString().Contains(".js"))
-                                LinkAction.CommandArguments += string.Format(" --js-library \"{0}\"", Item.AbsolutePath);
-                            else
-                                LinkAction.CommandArguments += string.Format(" \"{0}\"", Item.AbsolutePath);
+							if (Item.ToString().Contains(".js"))
+								ReponseLines.Add(string.Format(" --js-library \"{0}\"", Item.AbsolutePath));
+							else
+								ReponseLines.Add(string.Format(" \"{0}\"", Item.AbsolutePath));
                             LinkAction.PrerequisiteItems.Add(Item);
                         }
                     }
 			}
 			// make the file we will create
+
+
 			OutputFile = FileItem.GetItemByPath(LinkEnvironment.Config.OutputFilePath);
 			LinkAction.ProducedItems.Add(OutputFile);
-			LinkAction.CommandArguments += string.Format(" -o \"{0}\"", OutputFile.AbsolutePath);
+			ReponseLines.Add(string.Format(" -o \"{0}\"", OutputFile.AbsolutePath));
 
 		    FileItem OutputBC = FileItem.GetItemByPath(LinkEnvironment.Config.OutputFilePath.Replace(".js", ".bc").Replace(".html", ".bc"));
 		    LinkAction.ProducedItems.Add(OutputBC);
-			LinkAction.CommandArguments += " --emit-symbol-map " + string.Format(" --save-bc \"{0}\"", OutputBC.AbsolutePath);
+			ReponseLines.Add( " --emit-symbol-map " + string.Format(" --save-bc \"{0}\"", OutputBC.AbsolutePath));
 
      		LinkAction.StatusDescription = Path.GetFileName(OutputFile.AbsolutePath);
+
+			string ResponseFileName = GetResponseFileName(LinkEnvironment, OutputFile);
+			LinkAction.CommandArguments += string.Format(" @\"{0}\"", ResponseFile.Create(ResponseFileName, ReponseLines));
+
 			LinkAction.OutputEventHandler = new DataReceivedEventHandler(RemoteOutputReceivedEventHandler);
 
 			return OutputFile;
