- Basic support for compiling Xcode 7 (iOS9, MacOS 10.11)
   - Marked some overrides properly, then gave up and added -Wno-inconsistent-missing-override)
- Improved Xcode selection process for Mac and iOS (no more hardcoding Xcode path - it uses the currently running Xcode, or xcode-select when using commandline)
   - Added AppleToolchain.cs, to start sharing code between Mac and IOS Toolchains (Compile functionality and params to clang could be shared pretty easily)
   -
- Some Utility functions:
   - Added UBT utility to run a commandline and get its output, self-contained (Utils.RunLocalProcessAndReturnStdOut)
   - Added Log.TraceInformationOnce (and Error, Warning, etc) to print out a message only one time, without a bunch of static bools everywhere

diff --git a/Engine/Source/Programs/UnrealBuildTool/IOS/IOSToolChain.cs b/Engine/Source/Programs/UnrealBuildTool/IOS/IOSToolChain.cs
--- a/Engine/Source/Programs/UnrealBuildTool/IOS/IOSToolChain.cs
+++ b/Engine/Source/Programs/UnrealBuildTool/IOS/IOSToolChain.cs
@@ -15,7 +15,7 @@ using Ionic.Zlib;
 
 namespace UnrealBuildTool
 {
-	class IOSToolChain : RemoteToolChain
+	class IOSToolChain : AppleToolChain
 	{
 		public override void RegisterToolChain()
 		{
@@ -47,7 +47,10 @@ namespace UnrealBuildTool
 
 		/** Which developer directory to root from */
 		[XmlConfig]
-		public static string XcodeDeveloperDir = "/Applications/Xcode.app/Contents/Developer/";
+		private static string XcodeDeveloperDir = "xcode-select";
+
+		/** Directory for the developer binaries */
+		private static string ToolchainDir = "";
 
 		/** Location of the SDKs */
 		private static string BaseSDKDir;
@@ -75,6 +78,7 @@ namespace UnrealBuildTool
 			/** Location of the SDKs */
 			BaseSDKDir = XcodeDeveloperDir + "Platforms/iPhoneOS.platform/Developer/SDKs";
 			BaseSDKDirSim = XcodeDeveloperDir + "Platforms/iPhoneSimulator.platform/Developer/SDKs";
+			ToolchainDir = XcodeDeveloperDir + "Toolchains/XcodeDefault.xctoolchain/usr/bin/";
 		}
 
 		/** Hunt down the latest IOS sdk if desired */
@@ -82,92 +86,17 @@ namespace UnrealBuildTool
 		{
 			base.SetUpGlobalEnvironment();
 
-			if (IOSSDKVersion == "latest")
-			{
-				try
-				{
-					string[] SubDirs = null;
-					if (Utils.IsRunningOnMono)
-					{
-						// on the Mac, we can just get the directory name
-						SubDirs = System.IO.Directory.GetDirectories(BaseSDKDir);
-						if (!ProjectFileGenerator.bGenerateProjectFiles)
-						{
-							Log.TraceInformation(String.Format("Directories : {0} {1}", SubDirs, SubDirs[0]));
-						}
-					}
-					else
-					{
-						Hashtable Results = RPCUtilHelper.Command("/", "ls", BaseSDKDir, null);
-						if (Results != null)
-						{
-							string Result = (string)Results["CommandOutput"];
-							SubDirs = Result.Split("\r\n".ToCharArray(), StringSplitOptions.RemoveEmptyEntries);
-						}
-					}
-
-					// loop over the subdirs and parse out the version
-					float MaxSDKVersion = 0.0f;
-					string MaxSDKVersionString = null;
-					foreach (string SubDir in SubDirs)
-					{
-						string SubDirName = Path.GetFileNameWithoutExtension(SubDir);
-						if (SubDirName.StartsWith("iPhoneOS"))
-						{
-							// get the SDK version from the directory name
-							string SDKString = SubDirName.Replace("iPhoneOS", "");
-							float SDKVersion = 0.0f;
-							try
-							{
-								SDKVersion = float.Parse(SDKString, System.Globalization.CultureInfo.InvariantCulture);
-							}
-							catch (Exception)
-							{
-								// weirdly formatted SDKs
-								continue;
-							}
-
-							// update largest SDK version number
-							if (SDKVersion > MaxSDKVersion)
-							{
-								MaxSDKVersion = SDKVersion;
-								MaxSDKVersionString = SDKString;
-							}
-						}
-					}
+			// choose the XCode to use
+			SelectXcode(ref XcodeDeveloperDir);
 
-					// convert back to a string with the exact format
-					if (MaxSDKVersionString != null)
-					{
-						IOSSDKVersion = MaxSDKVersionString;
-					}
-				}
-				catch (Exception Ex)
-				{
-					// on any exception, just use the backup version
-					Log.TraceInformation("Triggered an exception while looking for SDK directory in Xcode.app");
-					Log.TraceInformation("{0}", Ex.ToString());
-				}
+			// make sure BaseSDKDir is up to date
+			PostReset();
 
-				if (IOSSDKVersion == "latest")
-				{
-					throw new BuildException("Unable to determine SDK version from Xcode, we cannot continue");
-				}
-			}
+			// choose the final SDK version
+			SelectSDK(BaseSDKDir, "iPhoneOS", ref IOSSDKVersion);
 
+			// convert to float for easy comparison
 			IOSSDKVersionFloat = float.Parse(IOSSDKVersion, System.Globalization.CultureInfo.InvariantCulture);
-
-			if (!ProjectFileGenerator.bGenerateProjectFiles)
-			{
-				if (BuildHostPlatform.Current.Platform != UnrealTargetPlatform.Mac)
-				{
-					Log.TraceInformation("Compiling with IOS SDK {0} on Mac {1}", IOSSDKVersionFloat, RemoteServerName);
-				}
-				else
-				{
-					Log.TraceInformation("Compiling with IOS SDK {0}", IOSSDKVersionFloat);
-				}
-			}
 		}
 
 		public override void AddFilesToReceipt(BuildReceipt Receipt, UEBuildBinary Binary)
@@ -247,6 +176,12 @@ namespace UnrealBuildTool
 			Result += " -Wno-unused-private-field";
 			Result += " -Wno-invalid-offsetof"; // needed to suppress warnings about using offsetof on non-POD types.
 	
+			if (IOSSDKVersionFloat >= 9.0)
+			{
+				Result += " -Wno-inconsistent-missing-override"; // too many missing overrides...
+				Result += " -Wno-unused-local-typedef"; // PhysX has some, hard to remove
+			}
+
 			Result += " -c";
 
 			// What architecture(s) to build for
@@ -602,7 +537,7 @@ namespace UnrealBuildTool
 				// Add the source file path to the command-line.
 				FileArguments += string.Format(" \"{0}\"", ConvertPath(SourceFile.AbsolutePath), false);
 
-				string CompilerPath = XcodeDeveloperDir + "Toolchains/XcodeDefault.xctoolchain/usr/bin/" + IOSCompiler;
+				string CompilerPath = ToolchainDir + IOSCompiler;
 				if (!Utils.IsRunningOnMono && BuildHostPlatform.Current.Platform != UnrealTargetPlatform.Mac)
 				{
 					CompileAction.ActionHandler = new Action.BlockingActionHandler(RPCUtilHelper.RPCActionHandler);
@@ -624,7 +559,7 @@ namespace UnrealBuildTool
 
 		public override FileItem LinkFiles(LinkEnvironment LinkEnvironment, bool bBuildImportLibraryOnly)
 		{
-			string LinkerPath = XcodeDeveloperDir + "Toolchains/XcodeDefault.xctoolchain/usr/bin/" + 
+			string LinkerPath = ToolchainDir + 
 				(LinkEnvironment.Config.bIsBuildingLibrary ? IOSArchiver : IOSLinker);
 
 			// Create an action that invokes the linker.
@@ -1012,8 +947,8 @@ namespace UnrealBuildTool
 				{
 					Process StripProcess = new Process();
 					StripProcess.StartInfo.WorkingDirectory = RemoteShadowDirectoryMac;
-					StripProcess.StartInfo.FileName = "/usr/bin/xcrun";
-					StripProcess.StartInfo.Arguments = "strip \"" + Target.OutputPath + "\"";
+					StripProcess.StartInfo.FileName = ToolchainDir + "strip";
+					StripProcess.StartInfo.Arguments = "\"" + Target.OutputPath + "\"";
 					StripProcess.OutputDataReceived += new DataReceivedEventHandler(OutputReceivedDataEventHandler);
 					StripProcess.ErrorDataReceived += new DataReceivedEventHandler(OutputReceivedDataEventHandler);
 
@@ -1363,14 +1298,7 @@ namespace UnrealBuildTool
 
 		public override void StripSymbols(string SourceFileName, string TargetFileName)
 		{
-			File.Copy(SourceFileName, TargetFileName, true);
-
-			ProcessStartInfo StartInfo = new ProcessStartInfo();
-			StartInfo.FileName = Path.Combine(XcodeDeveloperDir, "usr/bin/xcrun");
-			StartInfo.Arguments = String.Format("strip \"{0}\" -S", TargetFileName);
-			StartInfo.UseShellExecute = false;
-			StartInfo.CreateNoWindow = true;
-			Utils.RunLocalProcessAndLogOutput(StartInfo);
+			StripSymbolsWithXcode(SourceFileName, TargetFileName, ToolchainDir);
 		}
 	};
 }
diff --git a/Engine/Source/Programs/UnrealBuildTool/Mac/MacToolChain.cs b/Engine/Source/Programs/UnrealBuildTool/Mac/MacToolChain.cs
--- a/Engine/Source/Programs/UnrealBuildTool/Mac/MacToolChain.cs
+++ b/Engine/Source/Programs/UnrealBuildTool/Mac/MacToolChain.cs
@@ -10,7 +10,7 @@ using System.Text;
 
 namespace UnrealBuildTool
 {
-	class MacToolChain : RemoteToolChain
+	class MacToolChain : AppleToolChain
 	{
 		public override void RegisterToolChain()
 		{
@@ -26,6 +26,7 @@ namespace UnrealBuildTool
 
 		/** Which version of the Mac OS SDK to target at build time */
 		public static string MacOSSDKVersion = "latest";
+		public static float MacOSSDKVersionFloat = 0.0f;
 
 		/** Which version of the Mac OS X to allow at run time */
 		public static string MacOSVersion = "10.9";
@@ -33,8 +34,11 @@ namespace UnrealBuildTool
 		/** Minimum version of Mac OS X to actually run on, running on earlier versions will display the system minimum version error dialog & exit. */
 		public static string MinMacOSVersion = "10.9.2";
 
-		/** Which developer directory to root from */
-		private static string XcodeDeveloperDir = "/Applications/Xcode.app/Contents/Developer/";
+		/** Which developer directory to root from? If this is "xcode-select", UBT will query for the currently selected Xcode */
+		private static string XcodeDeveloperDir = "xcode-select";
+
+		/** Directory for the developer binaries */
+		private static string ToolchainDir = "";
 
 		/** Location of the SDKs */
 		private static string BaseSDKDir;
@@ -48,7 +52,6 @@ namespace UnrealBuildTool
 		/** Which archiver to use */
 		private static string MacArchiver = "libtool";
 
-
 		/** Track which scripts need to be deleted before appending to */
 		private bool bHasWipedCopyDylibScript = false;
 		private bool bHasWipedFixDylibScript = false;
@@ -61,88 +64,15 @@ namespace UnrealBuildTool
 		{
 			base.SetUpGlobalEnvironment();
 
-			BaseSDKDir = XcodeDeveloperDir + "Platforms/MacOSX.platform/Developer/SDKs";
+			SelectXcode(ref XcodeDeveloperDir);
 
-			if (MacOSSDKVersion == "latest")
-			{
-				try
-				{
-					string[] SubDirs = null;
-					if (Utils.IsRunningOnMono)
-					{
-						// on the Mac, we can just get the directory name
-						SubDirs = System.IO.Directory.GetDirectories(BaseSDKDir);
-					}
-					else
-					{
-						Hashtable Results = RPCUtilHelper.Command("/", "ls", BaseSDKDir, null);
-						if (Results != null)
-						{
-							string Result = (string)Results["CommandOutput"];
-							SubDirs = Result.Split("\r\n".ToCharArray(), StringSplitOptions.RemoveEmptyEntries);
-						}
-					}
-
-					// loop over the subdirs and parse out the version
-					float MaxSDKVersion = 0.0f;
-					string MaxSDKVersionString = null;
-					foreach (string SubDir in SubDirs)
-					{
-						string SubDirName = Path.GetFileNameWithoutExtension(SubDir);
-						if (SubDirName.StartsWith("MacOSX10."))
-						{
-							// get the SDK version from the directory name
-							string SDKString = SubDirName.Replace("MacOSX10.", "");
-							float SDKVersion = 0.0f;
-							try
-							{
-								SDKVersion = float.Parse(SDKString, System.Globalization.CultureInfo.InvariantCulture);
-							}
-							catch (Exception)
-							{
-								// weirdly formatted SDKs
-								continue;
-							}
-
-							// update largest SDK version number
-							if (SDKVersion > MaxSDKVersion)
-							{
-								MaxSDKVersion = SDKVersion;
-								MaxSDKVersionString = SDKString;
-							}
-						}
-					}
-
-					// convert back to a string with the exact format
-					if (MaxSDKVersionString != null)
-					{
-						MacOSSDKVersion = "10." + MaxSDKVersionString;
-					}
-				}
-				catch (Exception Ex)
-				{
-					// on any exception, just use the backup version
-					Log.TraceInformation("Triggered an exception while looking for SDK directory in Xcode.app");
-					Log.TraceInformation("{0}", Ex.ToString());
-				}
+			BaseSDKDir = XcodeDeveloperDir + "Platforms/MacOSX.platform/Developer/SDKs";
+			ToolchainDir = XcodeDeveloperDir + "Toolchains/XcodeDefault.xctoolchain/usr/bin/";
 
-				if (MacOSSDKVersion == "latest")
-				{
-					throw new BuildException("Unable to determine SDK version from Xcode, we cannot continue");
-				}
-			}
+			SelectSDK(BaseSDKDir, "MacOSX", ref MacOSSDKVersion);
 
-			if (!ProjectFileGenerator.bGenerateProjectFiles)
-			{
-				if (BuildHostPlatform.Current.Platform != UnrealTargetPlatform.Mac)
-				{
-					Log.TraceInformation("Compiling with Mac SDK {0} on Mac {1}", MacOSSDKVersion, RemoteServerName);
-				}
-				else
-				{
-					Log.TraceInformation("Compiling with Mac SDK {0}", MacOSSDKVersion);
-				}
-			}
+			// convert to float for easy comparison
+			MacOSSDKVersionFloat = float.Parse(MacOSSDKVersion, System.Globalization.CultureInfo.InvariantCulture);
 		}
 
 		static string GetCompileArguments_Global(CPPEnvironment CompileEnvironment)
@@ -170,6 +100,13 @@ namespace UnrealBuildTool
 			Result += " -Wno-unused-private-field";
 			Result += " -Wno-invalid-offsetof"; // needed to suppress warnings about using offsetof on non-POD types.
 
+			if (MacOSSDKVersionFloat < 10.9f && MacOSSDKVersionFloat >= 10.11f)
+			{
+				Log.TraceInformationOnce("Got into 10.11 land");
+				Result += " -Wno-inconsistent-missing-override"; // too many missing overrides...
+				Result += " -Wno-unused-local-typedef"; // PhysX has some, hard to remove
+			}
+
 			if (CompileEnvironment.Config.bEnableShadowVariableWarning)
 			{
 				Result += " -Wshadow" + (BuildConfiguration.bShadowVariableErrors? "" : " -Wno-error=shadow");
@@ -463,8 +400,8 @@ namespace UnrealBuildTool
 				}
 
 				CompileAction.WorkingDirectory = GetMacDevSrcRoot();
-				CompileAction.CommandPath = "xcrun";
-				CompileAction.CommandArguments = MacCompiler + Arguments + FileArguments + CompileEnvironment.Config.AdditionalArguments;
+				CompileAction.CommandPath = ToolchainDir + MacCompiler;
+				CompileAction.CommandArguments = Arguments + FileArguments + CompileEnvironment.Config.AdditionalArguments;
 				CompileAction.CommandDescription = "Compile";
 				CompileAction.StatusDescription = Path.GetFileName(SourceFile.AbsolutePath);
 				CompileAction.bIsGCCCompiler = true;
@@ -626,7 +563,7 @@ namespace UnrealBuildTool
 			string VersionArg = LinkEnvironment.Config.bIsBuildingDLL ? " -current_version " + EngineAPIVersion + " -compatibility_version " + EngineDisplayVersion : "";
 
 			string Linker = bIsBuildingLibrary ? MacArchiver : MacLinker;
-			string LinkCommand = "xcrun " + Linker + VersionArg + " " + (bIsBuildingLibrary ? GetArchiveArguments_Global(LinkEnvironment) : GetLinkArguments_Global(LinkEnvironment));
+			string LinkCommand = ToolchainDir + Linker + VersionArg + " " + (bIsBuildingLibrary ? GetArchiveArguments_Global(LinkEnvironment) : GetLinkArguments_Global(LinkEnvironment));
 
 			// Tell the action that we're building an import library here and it should conditionally be
 			// ignored as a prerequisite for other actions
@@ -905,7 +842,7 @@ namespace UnrealBuildTool
 				foreach (string Library in ThirdPartyLibraries)
 				{
 					string LibraryFileName = Path.GetFileName(Library);
-					LinkCommand += "; xcrun install_name_tool -change " + LibraryFileName + " " + DylibsPath + "/" + LibraryFileName + " \"" + ConvertPath(OutputFile.AbsolutePath) + "\"";
+					LinkCommand += "; " + ToolchainDir + "install_name_tool -change " + LibraryFileName + " " + DylibsPath + "/" + LibraryFileName + " \"" + ConvertPath(OutputFile.AbsolutePath) + "\"";
 				}
 			}
 
@@ -1201,8 +1138,8 @@ namespace UnrealBuildTool
 
 			// Deletes ay existing file on the building machine,
 			// note that the source and dest are switched from a copy command
-			GenDebugAction.CommandArguments = string.Format("-c 'rm -rf \"{2}\"; \"{0}\"usr/bin/xcrun dsymutil -f \"{1}\" -o \"{2}\"'",
-				XcodeDeveloperDir,
+			GenDebugAction.CommandArguments = string.Format("-c 'rm -rf \"{2}\"; \"{0}\"dsymutil -f \"{1}\" -o \"{2}\"'",
+				ToolchainDir,
 				InputFile.AbsolutePath,
 				DestFile.AbsolutePath);
 			GenDebugAction.PrerequisiteItems.Add(InputFile);
@@ -1498,14 +1435,7 @@ namespace UnrealBuildTool
 
 		public override void StripSymbols(string SourceFileName, string TargetFileName)
 		{
-			File.Copy(SourceFileName, TargetFileName, true);
-
-			ProcessStartInfo StartInfo = new ProcessStartInfo();
-			StartInfo.FileName = Path.Combine(XcodeDeveloperDir, "usr/bin/xcrun");
-			StartInfo.Arguments = String.Format("strip \"{0}\" -S", TargetFileName);
-			StartInfo.UseShellExecute = false;
-			StartInfo.CreateNoWindow = true;
-			Utils.RunLocalProcessAndLogOutput(StartInfo);
+			StripSymbolsWithXcode(SourceFileName, TargetFileName, ToolchainDir);
 		}
 	};
 }
diff --git a/Engine/Source/Programs/UnrealBuildTool/Mac/UEBuildMac.cs b/Engine/Source/Programs/UnrealBuildTool/Mac/UEBuildMac.cs
--- a/Engine/Source/Programs/UnrealBuildTool/Mac/UEBuildMac.cs
+++ b/Engine/Source/Programs/UnrealBuildTool/Mac/UEBuildMac.cs
@@ -163,6 +163,13 @@ namespace UnrealBuildTool
 
 		public override void ValidateBuildConfiguration(CPPTargetConfiguration Configuration, CPPTargetPlatform Platform, bool bCreateDebugInfo)
 		{
+			// Xcode 7 doesn't work currently with PCH files
+			if (MacToolChain.MacOSSDKVersion == "10.11")
+			{
+				Log.TraceInformation("=== Disabling PCH files due to Xcode 7 compiler bug! ===");
+				BuildConfiguration.bUsePCHFiles = false;
+			}
+
 			if (BuildHostPlatform.Current.Platform != UnrealTargetPlatform.Mac)
 			{
 				// @todo: Temporarily disable precompiled header files when building remotely due to errors
diff --git a/Engine/Source/Programs/UnrealBuildTool/System/Utils.cs b/Engine/Source/Programs/UnrealBuildTool/System/Utils.cs
--- a/Engine/Source/Programs/UnrealBuildTool/System/Utils.cs
+++ b/Engine/Source/Programs/UnrealBuildTool/System/Utils.cs
@@ -275,6 +275,30 @@ namespace UnrealBuildTool
 			return RunLocalProcess(LocalProcess);
 		}
 
+		/// <summary>
+		/// Runs a command line process, and returns simple StdOut output. This doesn't handle errors or return codes
+		/// </summary>
+		/// <returns>The entire StdOut generated from the process as a single trimmed string</returns>
+		/// <param name="Command">Command to run</param>
+		/// <param name="Args">Arguments to Command</param>
+		public static string RunLocalProcessAndReturnStdOut(string Command, string Args)
+		{
+			var StartInfo = new ProcessStartInfo(Command, Args);
+			StartInfo.UseShellExecute = false;
+			StartInfo.RedirectStandardOutput = true;
+			StartInfo.CreateNoWindow = true;
+
+			string FullOutput = "";
+			using (var LocalProcess = Process.Start(StartInfo))
+			{
+				StreamReader OutputReader = LocalProcess.StandardOutput;
+				// trim off any extraneous new lines, helpful for those one-line outputs
+				FullOutput = OutputReader.ReadToEnd().Trim();
+			}
+
+			return FullOutput;
+		}
+
 		/// <summary>
 		/// Given a list of supported platforms, returns a list of names of platforms that should not be supported
 		/// </summary>
@@ -937,6 +961,11 @@ namespace UnrealBuildTool
         /// </summary>
         public static bool bIsVerbose { get { return bLogVerbose; } }
 
+		/// <summary>
+		/// A collection of strings that have been already written once
+		/// </summary>
+		private static List<string> WriteOnceSet = new List<string>();
+
         /// <summary>
         /// Allows code to check if the log system is ready yet.
         /// End users should NOT need to use this. It pretty much exists
@@ -1042,16 +1071,30 @@ namespace UnrealBuildTool
         /// </summary>
         /// <param name="StackFramesToSkip">Number of frames to skip to get to the originator of the log request.</param>
         /// <param name="CustomSource">Custom source string to use. Use the default if null.</param>
+		/// <param name="bWriteOnce">If true, this message will be written only once</param>
         /// <param name="Verbosity">Message verbosity level. We only meaningfully use values up to Verbose</param>
         /// <param name="Format">Message format string.</param>
         /// <param name="Args">Optional arguments</param>
         [MethodImplAttribute(MethodImplOptions.NoInlining)]
-        private static void WriteLinePrivate(int StackFramesToSkip, string CustomSource, TraceEventType Verbosity, string Format, params object[] Args)
+        private static void WriteLinePrivate(int StackFramesToSkip, string CustomSource, bool bWriteOnce, TraceEventType Verbosity, string Format, params object[] Args)
         {
             if (!bIsInitialized)
             {
                 throw new BuildException("Tried to using Logging system before it was ready");
             }
+
+			// if we want this message only written one time, check if it was already written out
+			if (bWriteOnce)
+			{
+				string Formatted = string.Format(Format, Args);
+				if (WriteOnceSet.Contains(Formatted))
+				{
+					return;
+				}
+
+				WriteOnceSet.Add(Formatted);
+			}
+
             if (Verbosity < TraceEventType.Verbose || bLogVerbose)
             {
                 // Do console color highlighting here.
@@ -1115,7 +1158,7 @@ namespace UnrealBuildTool
         {
             if (Condition)
             {
-                WriteLinePrivate(1, null, Verbosity, Format, Args);
+				WriteLinePrivate(1, null, false, Verbosity, Format, Args);
             }
         }
 
@@ -1130,7 +1173,7 @@ namespace UnrealBuildTool
         [MethodImplAttribute(MethodImplOptions.NoInlining)]
         public static void WriteLine(int StackFramesToSkip, TraceEventType Verbosity, string Format, params object[] Args)
         {
-            WriteLinePrivate(StackFramesToSkip + 1, null, Verbosity, Format, Args);
+            WriteLinePrivate(StackFramesToSkip + 1, null, false, Verbosity, Format, Args);
         }
 
         /// <summary>
@@ -1144,7 +1187,7 @@ namespace UnrealBuildTool
         [MethodImplAttribute(MethodImplOptions.NoInlining)]
         public static void WriteLine(int StackFramesToSkip, string CustomSource, TraceEventType Verbosity, string Format, params object[] Args)
         {
-            WriteLinePrivate(StackFramesToSkip + 1, CustomSource, Verbosity, Format, Args);
+            WriteLinePrivate(StackFramesToSkip + 1, CustomSource, false, Verbosity, Format, Args);
         }
 
         /// <summary>
@@ -1156,7 +1199,7 @@ namespace UnrealBuildTool
         [MethodImplAttribute(MethodImplOptions.NoInlining)]
         public static void WriteLine(TraceEventType Verbosity, string Format, params object[] Args)
         {
-            WriteLinePrivate(1, null, Verbosity, Format, Args);
+            WriteLinePrivate(1, null, false, Verbosity, Format, Args);
         }
 
         /// <summary>
@@ -1167,7 +1210,7 @@ namespace UnrealBuildTool
         [MethodImplAttribute(MethodImplOptions.NoInlining)]
         public static void TraceError(string Format, params object[] Args)
         {
-            WriteLinePrivate(1, null, TraceEventType.Error, Format, Args);
+			WriteLinePrivate(1, null, false, TraceEventType.Error, Format, Args);
         }
 
         /// <summary>
@@ -1179,7 +1222,7 @@ namespace UnrealBuildTool
         [MethodImplAttribute(MethodImplOptions.NoInlining)]
         public static void TraceVerbose(string Format, params object[] Args)
         {
-            WriteLinePrivate(1, null, TraceEventType.Verbose, Format, Args);
+			WriteLinePrivate(1, null, false, TraceEventType.Verbose, Format, Args);
         }
 
         /// <summary>
@@ -1190,7 +1233,7 @@ namespace UnrealBuildTool
         [MethodImplAttribute(MethodImplOptions.NoInlining)]
         public static void TraceInformation(string Format, params object[] Args)
         {
-            WriteLinePrivate(1, null, TraceEventType.Information, Format, Args);
+			WriteLinePrivate(1, null, false, TraceEventType.Information, Format, Args);
         }
 
         /// <summary>
@@ -1201,9 +1244,68 @@ namespace UnrealBuildTool
         [MethodImplAttribute(MethodImplOptions.NoInlining)]
         public static void TraceWarning(string Format, params object[] Args)
         {
-            WriteLinePrivate(1, null, TraceEventType.Warning, Format, Args);
+            WriteLinePrivate(1, null, false, TraceEventType.Warning, Format, Args);
         }
-    }
+
+		/// <summary>
+		/// Similar to Trace.WriteLin
+		/// </summary>
+		/// <param name="Verbosity"></param>
+		/// <param name="Format"></param>
+		/// <param name="Args"></param>
+		[MethodImplAttribute(MethodImplOptions.NoInlining)]
+		public static void WriteLineOnce(TraceEventType Verbosity, string Format, params object[] Args)
+		{
+			WriteLinePrivate(1, null, true, Verbosity, Format, Args);
+		}
+
+		/// <summary>
+		/// Writes an error message to the console.
+		/// </summary>
+		/// <param name="Format">Message format string</param>
+		/// <param name="Args">Optional arguments</param>
+		[MethodImplAttribute(MethodImplOptions.NoInlining)]
+		public static void TraceErrorOnce(string Format, params object[] Args)
+		{
+			WriteLinePrivate(1, null, true, TraceEventType.Error, Format, Args);
+		}
+
+		/// <summary>
+		/// Writes a verbose message to the console.
+		/// </summary>
+		/// <param name="Format">Message format string</param>
+		/// <param name="Args">Optional arguments</param>
+		[Conditional("TRACE")]
+		[MethodImplAttribute(MethodImplOptions.NoInlining)]
+		public static void TraceVerboseOnce(string Format, params object[] Args)
+		{
+			WriteLinePrivate(1, null, true, TraceEventType.Verbose, Format, Args);
+		}
+
+		/// <summary>
+		/// Writes a message to the console.
+		/// </summary>
+		/// <param name="Format">Message format string</param>
+		/// <param name="Args">Optional arguments</param>
+		[MethodImplAttribute(MethodImplOptions.NoInlining)]
+		public static void TraceInformationOnce(string Format, params object[] Args)
+		{
+			WriteLinePrivate(1, null, true, TraceEventType.Information, Format, Args);
+		}
+
+		/// <summary>
+		/// Writes a warning message to the console.
+		/// </summary>
+		/// <param name="Format">Message format string</param>
+		/// <param name="Args">Optional arguments</param>
+		[MethodImplAttribute(MethodImplOptions.NoInlining)]
+		public static void TraceWarningOnce(string Format, params object[] Args)
+		{
+			WriteLinePrivate(1, null, true, TraceEventType.Warning, Format, Args);
+		}
+
+
+	}
 
     #region StreamUtils
     public static class StreamUtils
diff --git a/Engine/Source/Programs/UnrealBuildTool/ToolChain/AppleToolChain.cs b/Engine/Source/Programs/UnrealBuildTool/ToolChain/AppleToolChain.cs
new file mode 100644
--- /dev/null
+++ b/Engine/Source/Programs/UnrealBuildTool/ToolChain/AppleToolChain.cs
@@ -0,0 +1,164 @@
+﻿// Copyright 1998-2015 Epic Games, Inc. All Rights Reserved.
+
+using System;
+using System.Collections;
+using System.Collections.Generic;
+using System.IO;
+using System.Diagnostics;
+using System.Text;
+
+namespace UnrealBuildTool
+{
+	public abstract class AppleToolChain : RemoteToolChain
+    {
+		protected void SelectXcode(ref string DeveloperDir)
+		{
+			string Reason = "hardcoded";
+
+			if (DeveloperDir == "xcode-select")
+			{
+				Reason = "xcode-select";
+
+				if (Utils.IsRunningOnMono)
+				{
+					// on the Mac, run xcode-select directly
+					DeveloperDir = Utils.RunLocalProcessAndReturnStdOut("xcode-select", "--print-path");
+				}
+				else
+				{
+					Hashtable Results = RPCUtilHelper.Command("/", "xcode-select", "--print-path", null);
+					if (Results != null)
+					{
+						DeveloperDir = (string)Results["CommandOutput"];
+					}
+				}
+
+				if (DeveloperDir.EndsWith("/") == false)
+				{
+					// we expect this to end with a slash
+					DeveloperDir += "/";
+				}
+			}
+
+			// make sure we get a full path
+			if (Directory.Exists(DeveloperDir) == false)
+			{
+				throw new BuildException("Selected Xcode ('{0}') doesn't exist, cannot continue.", DeveloperDir);
+			}
+
+			if (!DeveloperDir.StartsWith("/Applications/Xcode.app"))
+			{
+				Log.TraceInformationOnce("Compiling with non-standard Xcode ({0}): {1}", Reason, DeveloperDir);
+			}
+
+		}
+
+		protected void SelectSDK(string BaseSDKDir, string OSPrefix, ref string PlatformSDKVersion)
+		{
+			if (PlatformSDKVersion == "latest")
+			{
+				PlatformSDKVersion = "";
+				try
+				{
+					string[] SubDirs = null;
+					if (Utils.IsRunningOnMono)
+					{
+						// on the Mac, we can just get the directory name
+						SubDirs = System.IO.Directory.GetDirectories(BaseSDKDir);
+					}
+					else
+					{
+						Hashtable Results = RPCUtilHelper.Command("/", "ls", BaseSDKDir, null);
+						if (Results != null)
+						{
+							string Result = (string)Results["CommandOutput"];
+							SubDirs = Result.Split("\r\n".ToCharArray(), StringSplitOptions.RemoveEmptyEntries);
+						}
+					}
+
+					// loop over the subdirs and parse out the version
+					int MaxSDKVersionMajor = 0;
+					int MaxSDKVersionMinor = 0;
+					string MaxSDKVersionString = null;
+					foreach (string SubDir in SubDirs)
+					{
+						string SubDirName = Path.GetFileNameWithoutExtension(SubDir);
+						if (SubDirName.StartsWith(OSPrefix))
+						{
+							// get the SDK version from the directory name
+							string SDKString = SubDirName.Replace(OSPrefix, "");
+							int Major = 0;
+							int Minor = 0;
+
+							// parse it into whole and fractional parts (since 10.10 > 10.9 in versions, but not in math)
+							try
+							{
+								string[] Tokens = SDKString.Split(".".ToCharArray());
+								if (Tokens.Length == 2)
+								{
+									Major = int.Parse(Tokens[0]);
+									Minor = int.Parse(Tokens[1]);
+								}
+							}
+							catch (Exception)
+							{
+								// weirdly formatted SDKs
+								continue;
+							}
+
+							// update largest SDK version number
+							if (Major > MaxSDKVersionMajor || (Major == MaxSDKVersionMajor && Minor > MaxSDKVersionMinor))
+							{
+								MaxSDKVersionString = SDKString;
+								MaxSDKVersionMajor = Major;
+								MaxSDKVersionMinor = Minor;
+							}
+						}
+					}
+
+					// use the largest version
+					if (MaxSDKVersionString != null)
+					{
+						PlatformSDKVersion = MaxSDKVersionString;
+					}
+				}
+				catch (Exception Ex)
+				{
+					// on any exception, just use the backup version
+					Log.TraceInformation("Triggered an exception while looking for SDK directory in Xcode.app");
+					Log.TraceInformation("{0}", Ex.ToString());
+				}
+			}
+
+			// make sure we have a valid SDK directory
+			if (!Directory.Exists(Path.Combine(BaseSDKDir, OSPrefix + PlatformSDKVersion + ".sdk")))
+			{
+				throw new BuildException("Invalid SDK {0}{1}.sdk, not found in {2}", OSPrefix, PlatformSDKVersion, BaseSDKDir);
+			}
+
+			if (!ProjectFileGenerator.bGenerateProjectFiles)
+			{
+				if (BuildHostPlatform.Current.Platform != UnrealTargetPlatform.Mac)
+				{
+					Log.TraceInformation("Compiling with {0} SDK {1} on Mac {2}", OSPrefix, PlatformSDKVersion, RemoteServerName);
+				}
+				else
+				{
+					Log.TraceInformation("Compiling with {0} SDK {1}", OSPrefix, PlatformSDKVersion);
+				}
+			}
+		}
+
+		protected void StripSymbolsWithXcode(string SourceFileName, string TargetFileName, string ToolchainDir)
+		{
+			File.Copy(SourceFileName, TargetFileName, true);
+
+			ProcessStartInfo StartInfo = new ProcessStartInfo();
+			StartInfo.FileName = Path.Combine(ToolchainDir, "strip");
+			StartInfo.Arguments = String.Format("\"{0}\" -S", TargetFileName);
+			StartInfo.UseShellExecute = false;
+			StartInfo.CreateNoWindow = true;
+			Utils.RunLocalProcessAndLogOutput(StartInfo);
+		}
+	};
+}
diff --git a/Engine/Source/Programs/UnrealBuildTool/ToolChain/RemoteToolChain.cs b/Engine/Source/Programs/UnrealBuildTool/ToolChain/RemoteToolChain.cs
--- a/Engine/Source/Programs/UnrealBuildTool/ToolChain/RemoteToolChain.cs
+++ b/Engine/Source/Programs/UnrealBuildTool/ToolChain/RemoteToolChain.cs
@@ -19,7 +19,7 @@ namespace UnrealBuildTool
 		{
 			RemoteToolChainPlatform = InPlatform;
 
-			// Register this tool chain for IOS
+			// Register this tool chain for the given platform
 			Log.TraceVerbose("        Registered for {0}", CPPPlatform.ToString());
 			UEToolChain.RegisterPlatformToolChain(CPPPlatform, this);
 		}
diff --git a/Engine/Source/Programs/UnrealBuildTool/UnrealBuildTool.csproj b/Engine/Source/Programs/UnrealBuildTool/UnrealBuildTool.csproj
--- a/Engine/Source/Programs/UnrealBuildTool/UnrealBuildTool.csproj
+++ b/Engine/Source/Programs/UnrealBuildTool/UnrealBuildTool.csproj
@@ -122,6 +122,7 @@
     <Compile Include="Mac\UEBuildMac.cs" />
     <Compile Include="Mac\MacProjectGenerator.cs" />
     <Compile Include="Mac\UEDeployMac.cs" />
+    <Compile Include="ToolChain\AppleToolChain.cs" />
   </ItemGroup>
   <ItemGroup>
     <Compile Include="Android\UEBuildAndroid.cs" />
diff --git a/Engine/Source/Runtime/Core/Private/HAL/FileManagerGeneric.cpp b/Engine/Source/Runtime/Core/Private/HAL/FileManagerGeneric.cpp
--- a/Engine/Source/Runtime/Core/Private/HAL/FileManagerGeneric.cpp
+++ b/Engine/Source/Runtime/Core/Private/HAL/FileManagerGeneric.cpp
@@ -68,15 +68,15 @@ public:
 	{
 		Close();
 	}
-	virtual void Seek( int64 InPos )
+	virtual void Seek( int64 InPos ) override
 	{
 		Pos = InPos;
 	}
-	virtual int64 Tell()
+	virtual int64 Tell() override
 	{
 		return Pos;
 	}
-	virtual void Serialize( void* V, int64 Length )
+	virtual void Serialize( void* V, int64 Length ) override
 	{
 		Pos += Length;
 	}
diff --git a/Engine/Source/Runtime/Core/Private/Mac/MacPlatformOutputDevicesPrivate.h b/Engine/Source/Runtime/Core/Private/Mac/MacPlatformOutputDevicesPrivate.h
--- a/Engine/Source/Runtime/Core/Private/Mac/MacPlatformOutputDevicesPrivate.h
+++ b/Engine/Source/Runtime/Core/Private/Mac/MacPlatformOutputDevicesPrivate.h
@@ -27,7 +27,7 @@ public:
 	 * Error handling function that is being called from within the system wide global
 	 * error handler, e.g. using structured exception handling on the PC.
 	 */
-	void HandleError();
+	virtual void HandleError() override;
 
 private:
 	int32		ErrorPos;
diff --git a/Engine/Source/Runtime/Core/Public/Apple/ApplePlatformFile.h b/Engine/Source/Runtime/Core/Public/Apple/ApplePlatformFile.h
--- a/Engine/Source/Runtime/Core/Public/Apple/ApplePlatformFile.h
+++ b/Engine/Source/Runtime/Core/Public/Apple/ApplePlatformFile.h
@@ -36,8 +36,8 @@ public:
 	virtual bool DirectoryExists(const TCHAR* Directory) override;
 	virtual bool CreateDirectory(const TCHAR* Directory) override;
 	virtual bool DeleteDirectory(const TCHAR* Directory) override;
-	virtual bool IterateDirectory(const TCHAR* Directory, FDirectoryVisitor& Visitor);
-	virtual bool CopyFile(const TCHAR* To, const TCHAR* From);
+	virtual bool IterateDirectory(const TCHAR* Directory, FDirectoryVisitor& Visitor) override;
+	virtual bool CopyFile(const TCHAR* To, const TCHAR* From) override;
 private:
 	int32 Stat(const TCHAR* Filename, struct stat* OutFileInfo);
 };
diff --git a/Engine/Source/Runtime/Core/Public/Containers/Set.h b/Engine/Source/Runtime/Core/Public/Containers/Set.h
--- a/Engine/Source/Runtime/Core/Public/Containers/Set.h
+++ b/Engine/Source/Runtime/Core/Public/Containers/Set.h
@@ -1370,9 +1370,6 @@ private:
 		static_assert(STRUCT_OFFSET(ScriptType, Elements) == STRUCT_OFFSET(RealType, Elements), "FScriptSet's Elements member offset does not match TSet's");
 		static_assert(STRUCT_OFFSET(ScriptType, Hash)     == STRUCT_OFFSET(RealType, Hash),     "FScriptSet's Hash member offset does not match TSet's");
 		static_assert(STRUCT_OFFSET(ScriptType, HashSize) == STRUCT_OFFSET(RealType, HashSize), "FScriptSet's FirstFreeIndex member offset does not match TSet's");
-
-		typedef TSet<bool>  ::SetElementType RealSetElementType1;
-		typedef TSet<double>::SetElementType RealSetElementType2;
 	}
 
 public:
diff --git a/Engine/Source/Runtime/Core/Public/HAL/ConsoleManager.h b/Engine/Source/Runtime/Core/Public/HAL/ConsoleManager.h
--- a/Engine/Source/Runtime/Core/Public/HAL/ConsoleManager.h
+++ b/Engine/Source/Runtime/Core/Public/HAL/ConsoleManager.h
@@ -76,7 +76,7 @@ public:
 	virtual IConsoleCommand* RegisterConsoleCommand(const TCHAR* Name, const TCHAR* Help, const FConsoleCommandWithWorldAndArgsDelegate& Command, uint32 Flags) override;
 	virtual IConsoleCommand* RegisterConsoleCommand(const TCHAR* Name, const TCHAR* Help, const FConsoleCommandWithOutputDeviceDelegate& Command, uint32 Flags) override;
 	virtual IConsoleCommand* RegisterConsoleCommand(const TCHAR* Name, const TCHAR* Help, uint32 Flags) override;
-	virtual IConsoleObject* FindConsoleObject(const TCHAR* Name) const;
+	virtual IConsoleObject* FindConsoleObject(const TCHAR* Name) const override;
 	virtual IConsoleVariable* FindConsoleVariable(const TCHAR* Name) const override;
 	virtual void ForEachConsoleObject(const FConsoleObjectVisitor& Visitor, const TCHAR* ThatStartsWith) const override;
 	virtual bool ProcessUserConsoleInput(const TCHAR* InInput, FOutputDevice& Ar, UWorld* InWorld) override;
diff --git a/Engine/Source/Runtime/Core/Public/HAL/FileManagerGeneric.h b/Engine/Source/Runtime/Core/Public/HAL/FileManagerGeneric.h
--- a/Engine/Source/Runtime/Core/Public/HAL/FileManagerGeneric.h
+++ b/Engine/Source/Runtime/Core/Public/HAL/FileManagerGeneric.h
@@ -236,7 +236,7 @@ public:
 	{
 		return Pos;
 	}
-	virtual int64 TotalSize();
+	virtual int64 TotalSize() override;
 	virtual bool Close() final;
 	virtual void Serialize( void* V, int64 Length ) final;
 	virtual void Flush() final;
diff --git a/Engine/Source/Runtime/Core/Public/Mac/MacPlatformRunnableThread.h b/Engine/Source/Runtime/Core/Public/Mac/MacPlatformRunnableThread.h
--- a/Engine/Source/Runtime/Core/Public/Mac/MacPlatformRunnableThread.h
+++ b/Engine/Source/Runtime/Core/Public/Mac/MacPlatformRunnableThread.h
@@ -22,7 +22,7 @@ private:
 	/**
 	 * Allows a platform subclass to setup anything needed on the thread before running the Run function
 	 */
-	virtual void PreRun()
+	virtual void PreRun() override
 	{
 		pthread_setname_np(TCHAR_TO_ANSI(*ThreadName));
 	}
@@ -36,7 +36,7 @@ private:
 	/**
 	 * Allows platforms to adjust stack size
 	 */
-	virtual uint32 AdjustStackSize(uint32 InStackSize)
+	virtual uint32 AdjustStackSize(uint32 InStackSize) override
 	{
 		if (InStackSize < GetDefaultStackSize())
 		{
diff --git a/Engine/Source/Runtime/Core/Public/Misc/ExpressionParserTypes.h b/Engine/Source/Runtime/Core/Public/Misc/ExpressionParserTypes.h
--- a/Engine/Source/Runtime/Core/Public/Misc/ExpressionParserTypes.h
+++ b/Engine/Source/Runtime/Core/Public/Misc/ExpressionParserTypes.h
@@ -371,7 +371,7 @@ struct TOperatorEvaluationEnvironment : IOperatorEvaluationEnvironment
 	{
 		return Operators.ExecPreUnary(Operator, R, Context);
 	}
-	virtual FExpressionResult ExecPostUnary(const FExpressionToken& Operator, const FExpressionToken& L) const
+	virtual FExpressionResult ExecPostUnary(const FExpressionToken& Operator, const FExpressionToken& L) const override
 	{
 		return Operators.ExecPostUnary(Operator, L, Context);
 	}
diff --git a/Engine/Source/Runtime/Projects/Private/PluginManager.h b/Engine/Source/Runtime/Projects/Private/PluginManager.h
--- a/Engine/Source/Runtime/Projects/Private/PluginManager.h
+++ b/Engine/Source/Runtime/Projects/Private/PluginManager.h
@@ -61,7 +61,7 @@ public:
 	~FPluginManager();
 
 	/** IPluginManager interface */
-	virtual void RefreshPluginsList();
+	virtual void RefreshPluginsList() override;
 	virtual bool LoadModulesForEnabledPlugins( const ELoadingPhase::Type LoadingPhase ) override;
 	virtual void SetRegisterMountPointDelegate( const FRegisterMountPointDelegate& Delegate ) override;
 	virtual bool AreRequiredPluginsAvailable() override;
