- Fixed stripping commandline [JIRA UE-17021]

diff --git a/Engine/Source/Programs/UnrealBuildTool/IOS/IOSToolChain.cs b/Engine/Source/Programs/UnrealBuildTool/IOS/IOSToolChain.cs
--- a/Engine/Source/Programs/UnrealBuildTool/IOS/IOSToolChain.cs
+++ b/Engine/Source/Programs/UnrealBuildTool/IOS/IOSToolChain.cs
@@ -70,15 +70,29 @@ namespace UnrealBuildTool
 		/** Additional frameworks stored locally so we have access without LinkEnvironment */
 		public List<UEBuildFramework> RememberedAdditionalFrameworks = new List<UEBuildFramework>();
 
+		private static void SetupXcodePaths(bool bVerbose)
+		{
+			// choose the XCode to use
+			SelectXcode(ref XcodeDeveloperDir, bVerbose);
+
+			// update cached paths
+			BaseSDKDir = XcodeDeveloperDir + "Platforms/iPhoneOS.platform/Developer/SDKs";
+			BaseSDKDirSim = XcodeDeveloperDir + "Platforms/iPhoneSimulator.platform/Developer/SDKs";
+			ToolchainDir = XcodeDeveloperDir + "Toolchains/XcodeDefault.xctoolchain/usr/bin/";
+
+			// make sure SDK is selected
+			SelectSDK(BaseSDKDir, "iPhoneOS", ref IOSSDKVersion, bVerbose);
+
+			// convert to float for easy comparison
+			IOSSDKVersionFloat = float.Parse(IOSSDKVersion, System.Globalization.CultureInfo.InvariantCulture);
+		}
+
 		/// <summary>
 		/// Function to call to after reset default data.
 		/// </summary>
 		public static void PostReset()
 		{
-			/** Location of the SDKs */
-			BaseSDKDir = XcodeDeveloperDir + "Platforms/iPhoneOS.platform/Developer/SDKs";
-			BaseSDKDirSim = XcodeDeveloperDir + "Platforms/iPhoneSimulator.platform/Developer/SDKs";
-			ToolchainDir = XcodeDeveloperDir + "Toolchains/XcodeDefault.xctoolchain/usr/bin/";
+			SetupXcodePaths(false);
 		}
 
 		/** Hunt down the latest IOS sdk if desired */
@@ -86,17 +100,7 @@ namespace UnrealBuildTool
 		{
 			base.SetUpGlobalEnvironment();
 
-			// choose the XCode to use
-			SelectXcode(ref XcodeDeveloperDir);
-
-			// make sure BaseSDKDir is up to date
-			PostReset();
-
-			// choose the final SDK version
-			SelectSDK(BaseSDKDir, "iPhoneOS", ref IOSSDKVersion);
-
-			// convert to float for easy comparison
-			IOSSDKVersionFloat = float.Parse(IOSSDKVersion, System.Globalization.CultureInfo.InvariantCulture);
+			SetupXcodePaths(true);
 		}
 
 		public override void AddFilesToReceipt(BuildReceipt Receipt, UEBuildBinary Binary)
@@ -1298,6 +1302,8 @@ namespace UnrealBuildTool
 
 		public override void StripSymbols(string SourceFileName, string TargetFileName)
 		{
+			SetupXcodePaths(false);
+
 			StripSymbolsWithXcode(SourceFileName, TargetFileName, ToolchainDir);
 		}
 	};
diff --git a/Engine/Source/Programs/UnrealBuildTool/Mac/MacToolChain.cs b/Engine/Source/Programs/UnrealBuildTool/Mac/MacToolChain.cs
--- a/Engine/Source/Programs/UnrealBuildTool/Mac/MacToolChain.cs
+++ b/Engine/Source/Programs/UnrealBuildTool/Mac/MacToolChain.cs
@@ -60,21 +60,26 @@ namespace UnrealBuildTool
 
 		public List<string> BuiltBinaries = new List<string>();
 
-		public override void SetUpGlobalEnvironment()
+		private static void SetupXcodePaths(bool bVerbose)
 		{
-			base.SetUpGlobalEnvironment();
-
-			SelectXcode(ref XcodeDeveloperDir);
+			SelectXcode(ref XcodeDeveloperDir, bVerbose);
 
 			BaseSDKDir = XcodeDeveloperDir + "Platforms/MacOSX.platform/Developer/SDKs";
 			ToolchainDir = XcodeDeveloperDir + "Toolchains/XcodeDefault.xctoolchain/usr/bin/";
 
-			SelectSDK(BaseSDKDir, "MacOSX", ref MacOSSDKVersion);
+			SelectSDK(BaseSDKDir, "MacOSX", ref MacOSSDKVersion, bVerbose);
 
 			// convert to float for easy comparison
 			MacOSSDKVersionFloat = float.Parse(MacOSSDKVersion, System.Globalization.CultureInfo.InvariantCulture);
 		}
 
+		public override void SetUpGlobalEnvironment()
+		{
+			base.SetUpGlobalEnvironment();
+
+			SetupXcodePaths(true);
+		}
+
 		static string GetCompileArguments_Global(CPPEnvironment CompileEnvironment)
 		{
 			string Result = "";
@@ -1435,6 +1440,8 @@ namespace UnrealBuildTool
 
 		public override void StripSymbols(string SourceFileName, string TargetFileName)
 		{
+			SetupXcodePaths(false);
+
 			StripSymbolsWithXcode(SourceFileName, TargetFileName, ToolchainDir);
 		}
 	};
diff --git a/Engine/Source/Programs/UnrealBuildTool/ToolChain/AppleToolChain.cs b/Engine/Source/Programs/UnrealBuildTool/ToolChain/AppleToolChain.cs
--- a/Engine/Source/Programs/UnrealBuildTool/ToolChain/AppleToolChain.cs
+++ b/Engine/Source/Programs/UnrealBuildTool/ToolChain/AppleToolChain.cs
@@ -11,7 +11,7 @@ namespace UnrealBuildTool
 {
 	public abstract class AppleToolChain : RemoteToolChain
     {
-		protected void SelectXcode(ref string DeveloperDir)
+		protected static void SelectXcode(ref string DeveloperDir, bool bVerbose)
 		{
 			string Reason = "hardcoded";
 
@@ -46,14 +46,14 @@ namespace UnrealBuildTool
 				throw new BuildException("Selected Xcode ('{0}') doesn't exist, cannot continue.", DeveloperDir);
 			}
 
-			if (!DeveloperDir.StartsWith("/Applications/Xcode.app"))
+			if (bVerbose && !DeveloperDir.StartsWith("/Applications/Xcode.app"))
 			{
 				Log.TraceInformationOnce("Compiling with non-standard Xcode ({0}): {1}", Reason, DeveloperDir);
 			}
 
 		}
 
-		protected void SelectSDK(string BaseSDKDir, string OSPrefix, ref string PlatformSDKVersion)
+		protected static void SelectSDK(string BaseSDKDir, string OSPrefix, ref string PlatformSDKVersion, bool bVerbose)
 		{
 			if (PlatformSDKVersion == "latest")
 			{
@@ -136,7 +136,7 @@ namespace UnrealBuildTool
 				throw new BuildException("Invalid SDK {0}{1}.sdk, not found in {2}", OSPrefix, PlatformSDKVersion, BaseSDKDir);
 			}
 
-			if (!ProjectFileGenerator.bGenerateProjectFiles)
+			if (bVerbose && !ProjectFileGenerator.bGenerateProjectFiles)
 			{
 				if (BuildHostPlatform.Current.Platform != UnrealTargetPlatform.Mac)
 				{
