- Fixed xcode-select for non-IOS/Mac builds when doing remote builds

diff --git a/Engine/Source/Programs/UnrealBuildTool/IOS/IOSToolChain.cs b/Engine/Source/Programs/UnrealBuildTool/IOS/IOSToolChain.cs
--- a/Engine/Source/Programs/UnrealBuildTool/IOS/IOSToolChain.cs
+++ b/Engine/Source/Programs/UnrealBuildTool/IOS/IOSToolChain.cs
@@ -46,7 +46,6 @@ namespace UnrealBuildTool
 		public static string BuildIOSVersion = "7.0";
 
 		/** Which developer directory to root from */
-		[XmlConfig]
 		private static string XcodeDeveloperDir = "xcode-select";
 
 		/** Directory for the developer binaries */
@@ -87,14 +86,6 @@ namespace UnrealBuildTool
 			IOSSDKVersionFloat = float.Parse(IOSSDKVersion, System.Globalization.CultureInfo.InvariantCulture);
 		}
 
-		/// <summary>
-		/// Function to call to after reset default data.
-		/// </summary>
-		public static void PostReset()
-		{
-			SetupXcodePaths(false);
-		}
-
 		/** Hunt down the latest IOS sdk if desired */
 		public override void SetUpGlobalEnvironment()
 		{
