Remove 60hz Morpheus codepaths

diff --git a/Engine/Source/Programs/UnrealBuildTool/Windows/UEBuildWindows.cs b/Engine/Source/Programs/UnrealBuildTool/Windows/UEBuildWindows.cs
--- a/Engine/Source/Programs/UnrealBuildTool/Windows/UEBuildWindows.cs
+++ b/Engine/Source/Programs/UnrealBuildTool/Windows/UEBuildWindows.cs
@@ -555,11 +555,9 @@ namespace UnrealBuildTool
             String MorpheusShaderPath = Path.Combine(BuildConfiguration.RelativeEnginePath, "Shaders/PS4/PostProcessHMDMorpheus.usf");
             if (File.Exists(MorpheusShaderPath))
             {
-                InBuildTarget.GlobalCompileEnvironment.Config.Definitions.Add("HAS_MORPHEUS=1");
-                InBuildTarget.GlobalCompileEnvironment.Config.Definitions.Add("MORPHEUS_120HZ=0");
+                InBuildTarget.GlobalCompileEnvironment.Config.Definitions.Add("HAS_MORPHEUS=1");                
 
-				//on PS4 the SDK now handles distortion correction.  On PC we will still have to handle it manually,
-				//but we require SDK changes before we can get the required data.  For the moment, no platform does in-engine morpheus distortion
+				//on PS4 the SDK now handles distortion correction.  On PC we will still have to handle it manually,				
 				InBuildTarget.GlobalCompileEnvironment.Config.Definitions.Add("MORPHEUS_ENGINE_DISTORTION=1");				
             }
 
