Improved Intellisense for Rocket code projects

Allow Engine Projects to be created in Rocket but only as stub projects to generate intellisense paths
Add all excluded binaries to list of Precompiled Binaries so that plugin paths are included when it comes to adding intellisense paths

diff --git a/Engine/Source/Programs/UnrealBuildTool/Configuration/UEBuildTarget.cs b/Engine/Source/Programs/UnrealBuildTool/Configuration/UEBuildTarget.cs
--- a/Engine/Source/Programs/UnrealBuildTool/Configuration/UEBuildTarget.cs
+++ b/Engine/Source/Programs/UnrealBuildTool/Configuration/UEBuildTarget.cs
@@ -1870,6 +1870,10 @@ namespace UnrealBuildTool
 					{
 						FilteredBinaries.Add(DLLBinary);
 					}
+					else if(!PrecompiledBinaries.Contains(DLLBinary))
+					{
+						PrecompiledBinaries.Add(DLLBinary);
+					}
 				}
 
 				// Copy the result into the final list
diff --git a/Engine/Source/Programs/UnrealBuildTool/System/ProjectFileGenerator.cs b/Engine/Source/Programs/UnrealBuildTool/System/ProjectFileGenerator.cs
--- a/Engine/Source/Programs/UnrealBuildTool/System/ProjectFileGenerator.cs
+++ b/Engine/Source/Programs/UnrealBuildTool/System/ProjectFileGenerator.cs
@@ -568,7 +568,7 @@ namespace UnrealBuildTool
 					var IntelliSenseTargetFiles = new List<Tuple<ProjectFile, string>>();
 					{
 						// Engine targets
-						if( EngineProject != null && !bGeneratingRocketProjectFiles )
+						if( EngineProject != null)
 						{
 							foreach( var ProjectTarget in EngineProject.ProjectTargets )
 							{
@@ -1713,12 +1713,6 @@ namespace UnrealBuildTool
 					IsEngineTarget = true;
 				}
 
-				if (IsEngineTarget && bGeneratingRocketProjectFiles)
-				{
-					// Rocket project file must never include engine targets.
-					continue;
-				}
-
 				bool WantProjectFileForTarget = true;
 				if(TargetFileRelativeToEngineDirectory.StartsWith(Path.Combine("Source", "Programs"), StringComparison.InvariantCultureIgnoreCase))
 				{
@@ -1808,6 +1802,13 @@ namespace UnrealBuildTool
 						{
 							EngineProject = ProjectFile;
 							BaseFolder = EngineRelativePath;
+							if (bGeneratingRocketProjectFiles)
+							{
+								// Allow engine projects to be created but not built for Rocket
+								EngineProject.IsForeignProject = false;
+								EngineProject.IsGeneratedProject = true;
+								EngineProject.IsStubProject = true;
+							}
 						}
 						else
 						{
