UEPLAT-454 - Add option for fullscreen immersive mode for Android projects
- Fixing typo in manifest file generation preventing it from working

diff --git a/Engine/Source/Programs/UnrealBuildTool/Android/UEDeployAndroid.cs b/Engine/Source/Programs/UnrealBuildTool/Android/UEDeployAndroid.cs
--- a/Engine/Source/Programs/UnrealBuildTool/Android/UEDeployAndroid.cs
+++ b/Engine/Source/Programs/UnrealBuildTool/Android/UEDeployAndroid.cs
@@ -1053,7 +1053,7 @@ namespace UnrealBuildTool.Android
 			Text.AppendLine(string.Format("\t\t<meta-data android:name=\"com.epicgames.ue4.GameActivity.DepthBufferPreference\" android:value=\"{0}\"/>", ConvertDepthBufferIniValue(DepthBufferPreference)));
 			Text.AppendLine(string.Format("\t\t<meta-data android:name=\"com.epicgames.ue4.GameActivity.bPackageDataInsideApk\" android:value=\"{0}\"/>", bPackageDataInsideApk ? "true" : "false"));
             Text.AppendLine(string.Format("\t\t<meta-data android:name=\"com.epicgames.ue4.GameActivity.bVerifyOBBOnStartUp\" android:value=\"{0}\"/>", (bIsForDistribution && !bDisableVerifyOBBOnStartUp) ? "true" : "false"));
-            Text.AppendLine(string.Format("\t\t<meta-data android:name=\"com.epicgames.ue4.GameActivity.bShouldHideUI=\" android:value=\"{0}\"/>", EnableFullScreen ? "true" : "false"));
+            Text.AppendLine(string.Format("\t\t<meta-data android:name=\"com.epicgames.ue4.GameActivity.bShouldHideUI\" android:value=\"{0}\"/>", EnableFullScreen ? "true" : "false"));
             Text.AppendLine(string.Format("\t\t<meta-data android:name=\"com.epicgames.ue4.GameActivity.ProjectName\" android:value=\"{0}\"/>", ProjectName));
 			Text.AppendLine(string.Format("\t\t<meta-data android:name=\"com.epicgames.ue4.GameActivity.bHasOBBFiles\" android:value=\"{0}\"/>", bHasOBBFiles ? "true" : "false"));
             Text.AppendLine("\t\t<meta-data android:name=\"com.google.android.gms.games.APP_ID\"");
