Added some clearer error handling around the RemoteToolChain initialisation.

I have improved error reporting order.
The error we get should now be the actual issue, and not an error that was the result of an original error that just hadnt been reported

diff --git a/Engine/Source/Programs/UnrealBuildTool/System/RPCUtilHelper.cs b/Engine/Source/Programs/UnrealBuildTool/System/RPCUtilHelper.cs
--- a/Engine/Source/Programs/UnrealBuildTool/System/RPCUtilHelper.cs
+++ b/Engine/Source/Programs/UnrealBuildTool/System/RPCUtilHelper.cs
@@ -94,7 +94,7 @@ namespace UnrealBuildTool
 			}
 		}
 
-		static public int Initialize(string InMacName)
+		static public RemoteToolChain.RemoteToolChainErrorCode Initialize(string InMacName)
 		{
 			MacName = InMacName;
 
@@ -109,7 +109,7 @@ namespace UnrealBuildTool
 						if (RemoteToolChain.ResolvedSSHAuthentication == null)
 						{
 							Log.TraceError("SSH authentication required a key, but one was not found. Use Editor to setup remote authentication!");
-							return 100;
+							return RemoteToolChain.RemoteToolChainErrorCode.MissingSSHKey;
 						}
 						// ask for current time, free memory and num CPUs
 						string[] Commands = new string[] 
@@ -122,7 +122,7 @@ namespace UnrealBuildTool
 						if ((Int64)Results["ExitCode"] != 0)
 						{
 							Log.TraceError("Failed to run init commands on {0}. Output = {1}", MacName, Results["CommandOutput"]);
-							return 101;
+							return RemoteToolChain.RemoteToolChainErrorCode.SSHCommandFailed;
 						}
 
 						string[] Lines = ((string)Results["CommandOutput"]).Split("\r\n".ToCharArray(), StringSplitOptions.RemoveEmptyEntries);
@@ -158,16 +158,16 @@ namespace UnrealBuildTool
 				{
 					Log.TraceVerbose("SSH Initialize exception {0}", Ex.ToString());
 					Log.TraceError("Failed to run init commands on {0}", MacName);
-					return 101;
+					return RemoteToolChain.RemoteToolChainErrorCode.SSHCommandFailed;
 				}
 			}
  			else
  			{
 				Log.TraceError("Failed to ping Mac named {0}", MacName);
-				return 102;
+				return RemoteToolChain.RemoteToolChainErrorCode.ServerNotResponding;
  			}
 
-			return 0;
+			return RemoteToolChain.RemoteToolChainErrorCode.NoError;
 		}
 
 		/**
diff --git a/Engine/Source/Programs/UnrealBuildTool/ToolChain/RemoteToolChain.cs b/Engine/Source/Programs/UnrealBuildTool/ToolChain/RemoteToolChain.cs
--- a/Engine/Source/Programs/UnrealBuildTool/ToolChain/RemoteToolChain.cs
+++ b/Engine/Source/Programs/UnrealBuildTool/ToolChain/RemoteToolChain.cs
@@ -14,7 +14,21 @@ using System.Text.RegularExpressions;
 namespace UnrealBuildTool
 {
 	public abstract class RemoteToolChain : UEToolChain
-    {
+	{
+		/**
+		 * Common error codes reported by Remote Tool Chain and its actions.
+		 */
+		public enum RemoteToolChainErrorCode
+		{
+			NoError						= 0,
+			ServerNameNotSpecified		= 1,
+			ServerNotResponding			= 2,
+			MissingDeltaCopyInstall		= 3,
+			MissingRemoteUserName		= 4,
+			MissingSSHKey				= 5,
+			SSHCommandFailed			= 6,
+		};
+
 		protected void RegisterRemoteToolChain(UnrealTargetPlatform InPlatform, CPPTargetPlatform CPPPlatform)
 		{
 			RemoteToolChainPlatform = InPlatform;
@@ -262,8 +276,8 @@ namespace UnrealBuildTool
 		}
 
         // Do any one-time, global initialization for the tool chain
-		static int InitializationErrorCode = 0;
-		private static int InitializeRemoteExecution()
+		static RemoteToolChainErrorCode InitializationErrorCode = RemoteToolChainErrorCode.NoError;
+		private static RemoteToolChainErrorCode InitializeRemoteExecution()
         {
 			if (bHasBeenInitialized)
 			{
@@ -331,14 +345,28 @@ namespace UnrealBuildTool
 				if (string.IsNullOrEmpty(RemoteServerName))
 				{
 					Log.TraceError("Remote compiling requires a server name. Use the editor to set up your remote compilation settings.");
-					InitializationErrorCode = 99;
-					return InitializationErrorCode;
+					return RemoteToolChainErrorCode.ServerNameNotSpecified;
 				}
 
 				if (!bUseRPCUtil)
 				{
+					// Verify the Delta Copy install path
+					ResolvedRSyncExe = ResolveString(RSyncExe, true);
+					ResolvedSSHExe = ResolveString(SSHExe, true);
+
+					if (!File.Exists(ResolvedRSyncExe) || !File.Exists(ResolvedSSHExe))
+					{
+						Log.TraceError("Remote compiling requires Delta Copy to be installed.");
+						return RemoteToolChainErrorCode.MissingDeltaCopyInstall;
+					}
+
 					// we need the RemoteServerName and the Username to find the private key
-					ResolvedRSyncUsername = ResolveString(RSyncUsername, false);
+					ResolvedRSyncUsername = ResolveString(RSyncUsername, false); 
+					if (string.IsNullOrEmpty(ResolvedRSyncUsername))
+					{
+						Log.TraceError("Remote compiling requires a user name. Use the editor to set up your remote compilation settings.");
+						return RemoteToolChainErrorCode.MissingRemoteUserName;
+					}
 
 					bool bFoundOverrideSSHPrivateKey = false;
 
@@ -383,24 +411,21 @@ namespace UnrealBuildTool
 						}
 					}
 
-/*					if (!bFoundOverrideSSHPrivateKey)
-					{
-						throw new BuildException("An SSHKey was required, but one cannot be found. Can't continue...");
-					}*/
-
 					// resolve the rest of the strings
-					ResolvedRSyncExe = ResolveString(RSyncExe, true);
-					ResolvedSSHExe = ResolveString(SSHExe, true);
 					ResolvedRsyncAuthentication = ResolveString(RsyncAuthentication, false);
 					ResolvedSSHAuthentication = ResolveString(SSHAuthentication, false);
 				}
 
 				// start up remote communication and record if it succeeds
-				InitializationErrorCode = RPCUtilHelper.Initialize(RemoteServerName);
-
-				// allow user to set up
-				if (InitializationErrorCode == 100 && !string.IsNullOrEmpty(ResolvedRSyncUsername))
+				InitializationErrorCode = (RemoteToolChainErrorCode)RPCUtilHelper.Initialize(RemoteServerName);
+				if (InitializationErrorCode != RemoteToolChainErrorCode.NoError && InitializationErrorCode != RemoteToolChainErrorCode.MissingSSHKey)
+				{
+					Log.TraceError("Failed to initialize a connection to the Remote Server {0}", RemoteServerName);
+					return InitializationErrorCode;
+				}
+				else if (InitializationErrorCode == RemoteToolChainErrorCode.MissingSSHKey)
 				{
+					// Allow the user to set up a key from here.
 					Process KeyProcess = new Process();
 					KeyProcess.StartInfo.WorkingDirectory = Path.GetFullPath(Path.Combine(BuildConfiguration.RelativeEnginePath, "Build", "BatchFiles"));
 					KeyProcess.StartInfo.FileName = "MakeAndInstallSSHKey.bat";
@@ -420,15 +445,9 @@ namespace UnrealBuildTool
 					// make sure it succeeded if we want to re-init
 					if (KeyProcess.ExitCode == 0)
 					{
-						InitializeRemoteExecution();
+						InitializationErrorCode = InitializeRemoteExecution();
 					}
 				}
-				else if (string.IsNullOrEmpty(ResolvedRSyncUsername))
-				{
-					Log.TraceError("Remote compiling requires a user name. Use the editor to set up your remote compilation settings.");
-					InitializationErrorCode = 98;
-					return InitializationErrorCode;
-				}
 			}
             else
 			{
@@ -445,10 +464,11 @@ namespace UnrealBuildTool
 			base.SetUpGlobalEnvironment();
 
 			// connect to server
-			InitializeRemoteExecution();
-
-			// Setup root directory to use.
-			SetUserDevRootFromServer();
+			if (InitializeRemoteExecution() == RemoteToolChainErrorCode.NoError)
+			{
+				// Setup root directory to use.
+				SetUserDevRootFromServer();
+			}
 		}
 
         /** Converts the passed in path from UBT host to compiler native format. */
