reorganized the new FilteredTraceListener so it can be used by both UBT and UAT

diff --git a/Engine/Source/Programs/AutomationTool/Log.cs b/Engine/Source/Programs/AutomationTool/Log.cs
--- a/Engine/Source/Programs/AutomationTool/Log.cs
+++ b/Engine/Source/Programs/AutomationTool/Log.cs
@@ -36,7 +36,7 @@ namespace AutomationTool
                 bColorConsoleOutput: true,
                 TraceListeners: new TraceListener[]
                 {
-                    new FilteredConsoleTraceListener(),
+                    new UnrealBuildTool.FilteredConsoleTraceListener(),
                     // could return null, but InitLogging handles this gracefully.
                     CreateLogFileListener(out LogFilename),
                     //@todo - this is only used by GUBP nodes. Ideally we don't waste this 20MB if we are not running GUBP.
@@ -264,40 +264,4 @@ namespace AutomationTool
 
 	#endregion
 
-	#region FilteredConsoleTraceListener
-
-	/// <summary>
-	///  Filtered Console Trace Listener
-	///  </summary>
-	class FilteredConsoleTraceListener : TextWriterTraceListener
-	{
-		public FilteredConsoleTraceListener()
-			: base(Console.OpenStandardOutput())
-		{ }
-
-		#region TraceListener Interface
-
-		public override void WriteLine(string message)
-		{
-			// strip off the message code
-			int Code = 0;
-			if (Int32.TryParse(message.Substring(0, 3), System.Globalization.NumberStyles.HexNumber, null, out Code))
-			{
-				// filter the message based on the code
-				if (Code < (int)TraceEventType.Verbose || UnrealBuildTool.Log.bIsVerbose)
-				{
-					// filter based on the message code
-					base.WriteLine(message.Substring(3));
-				}
-			}
-			else
-			{
-				base.WriteLine(message);
-			}
-		}
-
-		#endregion
-	}
-	#endregion
-
 }
diff --git a/Engine/Source/Programs/UnrealBuildTool/System/UnrealBuildTool.cs b/Engine/Source/Programs/UnrealBuildTool/System/UnrealBuildTool.cs
--- a/Engine/Source/Programs/UnrealBuildTool/System/UnrealBuildTool.cs
+++ b/Engine/Source/Programs/UnrealBuildTool/System/UnrealBuildTool.cs
@@ -857,7 +857,7 @@ namespace UnrealBuildTool
                 bColorConsoleOutput: true,
                 TraceListeners: new[] 
                 {
-                    new UnfilteredConsoleTraceListener(),
+                    new FilteredConsoleTraceListener(),
                     !string.IsNullOrEmpty(BuildConfiguration.LogFilename) ? new TextWriterTraceListener(new StreamWriter(new FileStream(BuildConfiguration.LogFilename, FileMode.Create, FileAccess.ReadWrite, FileShare.Read)) { AutoFlush = true }) : null,
                 });
         }
@@ -2755,40 +2755,4 @@ namespace UnrealBuildTool
 			}
 		}
     }
-
-	#region UnfilteredConsoleTraceListener
-
-	/// <summary>
-	///  Unfiltered Console Trace Listener
-	///  </summary>
-	class UnfilteredConsoleTraceListener : TextWriterTraceListener
-	{
-		public UnfilteredConsoleTraceListener()
-			: base(Console.OpenStandardOutput())
-		{ }
-
-		#region TraceListener Interface
-
-		public override void WriteLine(string message)
-		{
-			// strip off the message code
-			int Code = 0;
-			if (Int32.TryParse(message.Substring(0, 3), System.Globalization.NumberStyles.HexNumber, null, out Code))
-			{
-				// filter based on the message code
-				if (Code < (int)TraceEventType.Verbose || Log.bIsVerbose)
-				{
-					base.WriteLine(message.Substring(3));
-				}
-			}
-			else
-			{
-				base.WriteLine(message);
-			}
-		}
-
-		#endregion
-	}
-	#endregion
-
 }
diff --git a/Engine/Source/Programs/UnrealBuildTool/System/Utils.cs b/Engine/Source/Programs/UnrealBuildTool/System/Utils.cs
--- a/Engine/Source/Programs/UnrealBuildTool/System/Utils.cs
+++ b/Engine/Source/Programs/UnrealBuildTool/System/Utils.cs
@@ -1627,4 +1627,41 @@ namespace UnrealBuildTool
 	}
 
 	#endregion
+
+	#region FilteredConsoleTraceListener
+
+	/// <summary>
+	///  Filtered Console Trace Listener
+	///  </summary>
+	public class FilteredConsoleTraceListener : TextWriterTraceListener
+	{
+		public FilteredConsoleTraceListener()
+			: base(Console.OpenStandardOutput())
+		{ }
+
+		#region TraceListener Interface
+
+		public override void WriteLine(string message)
+		{
+			// strip off the message code
+			int Code = 0;
+			if (Int32.TryParse(message.Substring(0, 3), System.Globalization.NumberStyles.HexNumber, null, out Code))
+			{
+				// filter the message based on the code
+				if (Code < (int)TraceEventType.Verbose || Log.bIsVerbose)
+				{
+					// filter based on the message code
+					base.WriteLine(message.Substring(3));
+				}
+			}
+			else
+			{
+				base.WriteLine(message);
+			}
+		}
+
+		#endregion
+	}
+	#endregion
+
 }
