Added left out update for Android x86_64 support (fixes GearVR architecture issue)
#ue4
#android

diff --git a/Engine/Source/Programs/UnrealBuildTool/Android/UEDeployAndroid.cs b/Engine/Source/Programs/UnrealBuildTool/Android/UEDeployAndroid.cs
--- a/Engine/Source/Programs/UnrealBuildTool/Android/UEDeployAndroid.cs
+++ b/Engine/Source/Programs/UnrealBuildTool/Android/UEDeployAndroid.cs
@@ -529,10 +529,14 @@ namespace UnrealBuildTool.Android
 		{
 			string Arch = GetNDKArch(UE4Arch);
 
-			string GccVersion = "4.8";
-			if (!Directory.Exists(Environment.ExpandEnvironmentVariables("%NDKROOT%/sources/cxx-stl/gnu-libstdc++/4.8")))
+			string GccVersion = "4.6";
+			if (Directory.Exists(Environment.ExpandEnvironmentVariables("%NDKROOT%/sources/cxx-stl/gnu-libstdc++/4.9")))
 			{
-				GccVersion = "4.6";
+				GccVersion = "4.9";
+			}
+			else if (Directory.Exists(Environment.ExpandEnvironmentVariables("%NDKROOT%/sources/cxx-stl/gnu-libstdc++/4.8")))
+			{
+				GccVersion = "4.8";
 			}
 
 			// copy it in!
@@ -933,7 +937,7 @@ namespace UnrealBuildTool.Android
 
         
 
-        private string GenerateManifest(string ProjectName, bool bIsForDistribution, bool bPackageDataInsideApk, string GameBuildFilesPath, bool bHasOBBFiles, bool bDisableVerifyOBBOnStartUp, string GPUArch)
+        private string GenerateManifest(string ProjectName, bool bIsForDistribution, bool bPackageDataInsideApk, string GameBuildFilesPath, bool bHasOBBFiles, bool bDisableVerifyOBBOnStartUp, string UE4Arch, string GPUArch)
 		{
 			// ini file to get settings from
 			ConfigCacheIni Ini = GetConfigCacheIni("Engine");
@@ -952,6 +956,8 @@ namespace UnrealBuildTool.Android
 			Ini.GetString("/Script/AndroidRuntimeSettings.AndroidRuntimeSettings", "Orientation", out Orientation);
             bool EnableFullScreen;
             Ini.GetBool("/Script/AndroidRuntimeSettings.AndroidRuntimeSettings", "bFullScreen", out EnableFullScreen);
+			bool EnableEngineCrashHandler;
+			Ini.GetBool("/Script/AndroidRuntimeSettings.AndroidRuntimeSettings", "bEnableEngineCrashHandler", out EnableEngineCrashHandler);
 			List<string> ExtraManifestNodeTags;
 			Ini.GetArray("/Script/AndroidRuntimeSettings.AndroidRuntimeSettings", "ExtraManifestNodeTags", out ExtraManifestNodeTags);
 			List<string> ExtraApplicationNodeTags;
@@ -969,6 +975,12 @@ namespace UnrealBuildTool.Android
 			bool bEnableIAP = false;
 			Ini.GetBool("OnlineSubsystemGooglePlay.Store", "bSupportsInAppPurchasing", out bEnableIAP);
 
+			// disable GearVR if not supported platform (in this case only armv7 for now)
+			if (UE4Arch != "-armv7")
+			{
+				bPackageForGearVR = false;
+			}
+
 			StringBuilder Text = new StringBuilder();
 			Text.AppendLine("<?xml version=\"1.0\" encoding=\"utf-8\"?>");
 			Text.AppendLine("<manifest xmlns:android=\"http://schemas.android.com/apk/res/android\"");
@@ -1002,7 +1014,7 @@ namespace UnrealBuildTool.Android
 			{
 				// normal application settings
 				Text.AppendLine("\t\t          android:theme=\"@android:style/Theme.NoTitleBar.Fullscreen\"");
-				Text.AppendLine("\t\t          android:configChanges=\"orientation|keyboardHidden\"");
+				Text.AppendLine("\t\t          android:configChanges=\"screenSize|orientation|keyboardHidden\"");
 			}
 			else
 			{
@@ -1054,7 +1066,8 @@ namespace UnrealBuildTool.Android
 			Text.AppendLine(string.Format("\t\t<meta-data android:name=\"com.epicgames.ue4.GameActivity.bPackageDataInsideApk\" android:value=\"{0}\"/>", bPackageDataInsideApk ? "true" : "false"));
             Text.AppendLine(string.Format("\t\t<meta-data android:name=\"com.epicgames.ue4.GameActivity.bVerifyOBBOnStartUp\" android:value=\"{0}\"/>", (bIsForDistribution && !bDisableVerifyOBBOnStartUp) ? "true" : "false"));
             Text.AppendLine(string.Format("\t\t<meta-data android:name=\"com.epicgames.ue4.GameActivity.bShouldHideUI\" android:value=\"{0}\"/>", EnableFullScreen ? "true" : "false"));
-            Text.AppendLine(string.Format("\t\t<meta-data android:name=\"com.epicgames.ue4.GameActivity.ProjectName\" android:value=\"{0}\"/>", ProjectName));
+			Text.AppendLine(string.Format("\t\t<meta-data android:name=\"com.epicgames.ue4.GameActivity.bEnableEngineCrashHandler\" android:value=\"{0}\"/>", EnableEngineCrashHandler ? "true" : "false"));
+			Text.AppendLine(string.Format("\t\t<meta-data android:name=\"com.epicgames.ue4.GameActivity.ProjectName\" android:value=\"{0}\"/>", ProjectName));
 			Text.AppendLine(string.Format("\t\t<meta-data android:name=\"com.epicgames.ue4.GameActivity.bHasOBBFiles\" android:value=\"{0}\"/>", bHasOBBFiles ? "true" : "false"));
             Text.AppendLine("\t\t<meta-data android:name=\"com.google.android.gms.games.APP_ID\"");
 			Text.AppendLine("\t\t           android:value=\"@string/app_id\" />");
@@ -1251,7 +1264,7 @@ namespace UnrealBuildTool.Android
             {
                 BuildList = from Arch in Arches
                             from GPUArch in GPUArchitectures
-                            let manifest = GenerateManifest(ProjectName, bForDistribution, bPackageDataInsideApk, GameBuildFilesPath, bDisallowPackagingDataInApk ? false : File.Exists(ObbFileLocation), bDisableVerifyOBBOnStartUp, GPUArch)
+                            let manifest = GenerateManifest(ProjectName, bForDistribution, bPackageDataInsideApk, GameBuildFilesPath, bDisallowPackagingDataInApk ? false : File.Exists(ObbFileLocation), bDisableVerifyOBBOnStartUp, Arch, GPUArch)
                             select Tuple.Create(Arch, GPUArch, manifest);
             }
             else
@@ -1259,7 +1272,7 @@ namespace UnrealBuildTool.Android
                 BuildList = from Arch in Arches
                             from GPUArch in GPUArchitectures
                             let manifestFile = Path.Combine(IntermediateAndroidPath, Arch + "_" + GPUArch + "_AndroidManifest.xml")
-                            let manifest = GenerateManifest(ProjectName, bForDistribution, bPackageDataInsideApk, GameBuildFilesPath, bDisallowPackagingDataInApk ? false : File.Exists(ObbFileLocation), bDisableVerifyOBBOnStartUp, GPUArch)
+                            let manifest = GenerateManifest(ProjectName, bForDistribution, bPackageDataInsideApk, GameBuildFilesPath, bDisallowPackagingDataInApk ? false : File.Exists(ObbFileLocation), bDisableVerifyOBBOnStartUp, Arch, GPUArch)
                             let OldManifest = File.Exists(manifestFile) ? File.ReadAllText(manifestFile) : ""
                             where manifest != OldManifest
                             select Tuple.Create(Arch, GPUArch, manifest);
