fix for deleting the PCLaunch.rc.res in rocket
UE-16971
#ubt

diff --git a/Engine/Source/Programs/UnrealBuildTool/Configuration/UEBuildTarget.cs b/Engine/Source/Programs/UnrealBuildTool/Configuration/UEBuildTarget.cs
--- a/Engine/Source/Programs/UnrealBuildTool/Configuration/UEBuildTarget.cs
+++ b/Engine/Source/Programs/UnrealBuildTool/Configuration/UEBuildTarget.cs
@@ -1085,7 +1085,10 @@ namespace UnrealBuildTool
 				{
 					if(Directory.Exists(Binary.Config.IntermediateDirectory))
 					{
-						CleanDirectory(Binary.Config.IntermediateDirectory);
+						if (!UnrealBuildTool.RunningRocket() || !Utils.IsFileUnderDirectory(Binary.Config.IntermediateDirectory, BuildConfiguration.RelativeEnginePath))
+						{
+							CleanDirectory(Binary.Config.IntermediateDirectory);
+						}
 					}
 				}
 
