Fix for iOS building on PC
Fix for iOS loading ini files for every cpp file

diff --git a/Engine/Source/Programs/UnrealBuildTool/Configuration/UEBuildPlatform.cs b/Engine/Source/Programs/UnrealBuildTool/Configuration/UEBuildPlatform.cs
--- a/Engine/Source/Programs/UnrealBuildTool/Configuration/UEBuildPlatform.cs
+++ b/Engine/Source/Programs/UnrealBuildTool/Configuration/UEBuildPlatform.cs
@@ -284,6 +284,8 @@ namespace UnrealBuildTool
 		// a mapping of a group to the platforms in the group (ie, Microsoft contains Win32 and Win64)
 		static Dictionary<UnrealPlatformGroup, List<UnrealTargetPlatform>> PlatformGroupDictionary = new Dictionary<UnrealPlatformGroup, List<UnrealTargetPlatform>>();
 
+		private bool bInitializedProject = false;
+
 		/**
 		 * Attempt to convert a string to an UnrealTargetPlatform enum entry
 		 * 
@@ -750,83 +752,88 @@ namespace UnrealBuildTool
 		 */
 		public virtual void SetUpProjectEnvironment(UnrealTargetPlatform InPlatform)
 		{
-			ConfigCacheIni Ini = new ConfigCacheIni(InPlatform, "Engine", UnrealBuildTool.GetUProjectPath());
-			bool bValue = UEBuildConfiguration.bCompileAPEX;
-			if (Ini.GetBool("/Script/BuildSettings.BuildSettings", "bCompileApex", out bValue))
+			if (!bInitializedProject)
 			{
-				UEBuildConfiguration.bCompileAPEX = bValue;
-			}
+				ConfigCacheIni Ini = new ConfigCacheIni(InPlatform, "Engine", UnrealBuildTool.GetUProjectPath());
+				bool bValue = UEBuildConfiguration.bCompileAPEX;
+				if (Ini.GetBool("/Script/BuildSettings.BuildSettings", "bCompileApex", out bValue))
+				{
+					UEBuildConfiguration.bCompileAPEX = bValue;
+				}
 
-			bValue = UEBuildConfiguration.bCompileBox2D;
-			if (Ini.GetBool("/Script/BuildSettings.BuildSettings", "bCompileBox2D", out bValue))
-			{
-				UEBuildConfiguration.bCompileBox2D = bValue;
-			}
+				bValue = UEBuildConfiguration.bCompileBox2D;
+				if (Ini.GetBool("/Script/BuildSettings.BuildSettings", "bCompileBox2D", out bValue))
+				{
+					UEBuildConfiguration.bCompileBox2D = bValue;
+				}
 
-			bValue = UEBuildConfiguration.bCompileICU;
-			if (Ini.GetBool("/Script/BuildSettings.BuildSettings", "bCompileICU", out bValue))
-			{
-				UEBuildConfiguration.bCompileICU = bValue;
-			}
+				bValue = UEBuildConfiguration.bCompileICU;
+				if (Ini.GetBool("/Script/BuildSettings.BuildSettings", "bCompileICU", out bValue))
+				{
+					UEBuildConfiguration.bCompileICU = bValue;
+				}
 
-			bValue = UEBuildConfiguration.bCompileSimplygon;
-			if (Ini.GetBool("/Script/BuildSettings.BuildSettings", "bCompileSimplygon", out bValue))
-			{
-				UEBuildConfiguration.bCompileSimplygon = bValue;
-			}
+				bValue = UEBuildConfiguration.bCompileSimplygon;
+				if (Ini.GetBool("/Script/BuildSettings.BuildSettings", "bCompileSimplygon", out bValue))
+				{
+					UEBuildConfiguration.bCompileSimplygon = bValue;
+				}
 
-			bValue = UEBuildConfiguration.bCompileLeanAndMeanUE;
-			if (Ini.GetBool("/Script/BuildSettings.BuildSettings", "bCompileLeanAndMeanUE", out bValue))
-			{
-				UEBuildConfiguration.bCompileLeanAndMeanUE = bValue;
-			}
+				bValue = UEBuildConfiguration.bCompileLeanAndMeanUE;
+				if (Ini.GetBool("/Script/BuildSettings.BuildSettings", "bCompileLeanAndMeanUE", out bValue))
+				{
+					UEBuildConfiguration.bCompileLeanAndMeanUE = bValue;
+				}
 
-			bValue = UEBuildConfiguration.bIncludeADO;
-			if (Ini.GetBool("/Script/BuildSettings.BuildSettings", "bIncludeADO", out bValue))
-			{
-				UEBuildConfiguration.bIncludeADO = bValue;
-			}
+				bValue = UEBuildConfiguration.bIncludeADO;
+				if (Ini.GetBool("/Script/BuildSettings.BuildSettings", "bIncludeADO", out bValue))
+				{
+					UEBuildConfiguration.bIncludeADO = bValue;
+				}
 
-			bValue = UEBuildConfiguration.bCompileRecast;
-			if (Ini.GetBool("/Script/BuildSettings.BuildSettings", "bCompileRecast", out bValue))
-			{
-				UEBuildConfiguration.bCompileRecast = bValue;
-			}
+				bValue = UEBuildConfiguration.bCompileRecast;
+				if (Ini.GetBool("/Script/BuildSettings.BuildSettings", "bCompileRecast", out bValue))
+				{
+					UEBuildConfiguration.bCompileRecast = bValue;
+				}
 
-			bValue = UEBuildConfiguration.bCompileSpeedTree;
-			if (Ini.GetBool("/Script/BuildSettings.BuildSettings", "bCompileSpeedTree", out bValue))
-			{
-				UEBuildConfiguration.bCompileSpeedTree = bValue;
-			}
+				bValue = UEBuildConfiguration.bCompileSpeedTree;
+				if (Ini.GetBool("/Script/BuildSettings.BuildSettings", "bCompileSpeedTree", out bValue))
+				{
+					UEBuildConfiguration.bCompileSpeedTree = bValue;
+				}
 
-			bValue = UEBuildConfiguration.bCompileWithPluginSupport;
-			if (Ini.GetBool("/Script/BuildSettings.BuildSettings", "bCompileWithPluginSupport", out bValue))
-			{
-				UEBuildConfiguration.bCompileWithPluginSupport = bValue;
-			}
+				bValue = UEBuildConfiguration.bCompileWithPluginSupport;
+				if (Ini.GetBool("/Script/BuildSettings.BuildSettings", "bCompileWithPluginSupport", out bValue))
+				{
+					UEBuildConfiguration.bCompileWithPluginSupport = bValue;
+				}
 
-			bValue = UEBuildConfiguration.bCompilePhysXVehicle;
-			if (Ini.GetBool("/Script/BuildSettings.BuildSettings", "bCompilePhysXVehicle", out bValue))
-			{
-				UEBuildConfiguration.bCompilePhysXVehicle = bValue;
-			}
+				bValue = UEBuildConfiguration.bCompilePhysXVehicle;
+				if (Ini.GetBool("/Script/BuildSettings.BuildSettings", "bCompilePhysXVehicle", out bValue))
+				{
+					UEBuildConfiguration.bCompilePhysXVehicle = bValue;
+				}
 
-			bValue = UEBuildConfiguration.bCompileFreeType;
-			if (Ini.GetBool("/Script/BuildSettings.BuildSettings", "bCompileFreeType", out bValue))
-			{
-				UEBuildConfiguration.bCompileFreeType = bValue;
-			}
+				bValue = UEBuildConfiguration.bCompileFreeType;
+				if (Ini.GetBool("/Script/BuildSettings.BuildSettings", "bCompileFreeType", out bValue))
+				{
+					UEBuildConfiguration.bCompileFreeType = bValue;
+				}
 
-			bValue = UEBuildConfiguration.bCompileForSize;
-			if (Ini.GetBool("/Script/BuildSettings.BuildSettings", "bCompileForSize", out bValue))
-			{
-				UEBuildConfiguration.bCompileForSize = bValue;
-			}
+				bValue = UEBuildConfiguration.bCompileForSize;
+				if (Ini.GetBool("/Script/BuildSettings.BuildSettings", "bCompileForSize", out bValue))
+				{
+					UEBuildConfiguration.bCompileForSize = bValue;
+				}
 
-			bValue = UEBuildConfiguration.bCompileCEF3;
-			if (Ini.GetBool("/Script/BuildSettings.BuildSettings", "bCompileCEF3", out bValue))
-			{
-				UEBuildConfiguration.bCompileCEF3 = bValue;
+				bValue = UEBuildConfiguration.bCompileCEF3;
+				if (Ini.GetBool("/Script/BuildSettings.BuildSettings", "bCompileCEF3", out bValue))
+				{
+					UEBuildConfiguration.bCompileCEF3 = bValue;
+				}
+
+				bInitializedProject = true;
 			}
 		}
 
diff --git a/Engine/Source/Programs/UnrealBuildTool/IOS/UEBuildIOS.cs b/Engine/Source/Programs/UnrealBuildTool/IOS/UEBuildIOS.cs
--- a/Engine/Source/Programs/UnrealBuildTool/IOS/UEBuildIOS.cs
+++ b/Engine/Source/Programs/UnrealBuildTool/IOS/UEBuildIOS.cs
@@ -27,6 +27,8 @@ namespace UnrealBuildTool
 		[XmlConfig]
 		public static string ShippingArchitectures = "armv7,arm64";
 
+		private bool bInitializedProject = false;
+
 		public string GetRunTimeVersion()
 		{
 			return RunTimeIOSVersion;
@@ -141,99 +143,104 @@ namespace UnrealBuildTool
 
 		public override void SetUpProjectEnvironment(UnrealTargetPlatform InPlatform)
 		{
-			base.SetUpProjectEnvironment(InPlatform);
-
-			// update the configuration based on the project file
-			// look in ini settings for what platforms to compile for
-			ConfigCacheIni Ini = new ConfigCacheIni(InPlatform, "Engine", UnrealBuildTool.GetUProjectPath());
-			string MinVersion = "IOS_6";
-			if (Ini.GetString("/Script/IOSRuntimeSettings.IOSRuntimeSettings", "MinimumiOSVersion", out MinVersion))
+			if (!bInitializedProject)
 			{
-				switch (MinVersion)
+				base.SetUpProjectEnvironment(InPlatform);
+
+				// update the configuration based on the project file
+				// look in ini settings for what platforms to compile for
+				ConfigCacheIni Ini = new ConfigCacheIni(InPlatform, "Engine", UnrealBuildTool.GetUProjectPath());
+				string MinVersion = "IOS_6";
+				if (Ini.GetString("/Script/IOSRuntimeSettings.IOSRuntimeSettings", "MinimumiOSVersion", out MinVersion))
 				{
-					case "IOS_61":
-						RunTimeIOSVersion = "6.1";
-						break;
-					case "IOS_7":
-						RunTimeIOSVersion = "7.0";
-						break;
-					case "IOS_8":
-						RunTimeIOSVersion = "8.0";
-						break;
+					switch (MinVersion)
+					{
+						case "IOS_61":
+							RunTimeIOSVersion = "6.1";
+							break;
+						case "IOS_7":
+							RunTimeIOSVersion = "7.0";
+							break;
+						case "IOS_8":
+							RunTimeIOSVersion = "8.0";
+							break;
+					}
 				}
-			}
 
-			bool biPhoneAllowed = true;
-			bool biPadAllowed = true;
-			Ini.GetBool("/Script/IOSRuntimeSettings.IOSRuntimeSettings", "bSupportsIPhone", out biPhoneAllowed);
-			Ini.GetBool("/Script/IOSRuntimeSettings.IOSRuntimeSettings", "bSupportsIPad", out biPadAllowed);
-			if (biPhoneAllowed && biPadAllowed)
-			{
-				RunTimeIOSDevices = "1,2";
-			}
-			else if (biPadAllowed)
-			{
-				RunTimeIOSDevices = "2";
-			}
-			else if (biPhoneAllowed)
-			{
-				RunTimeIOSDevices = "1";
-			}
+				bool biPhoneAllowed = true;
+				bool biPadAllowed = true;
+				Ini.GetBool("/Script/IOSRuntimeSettings.IOSRuntimeSettings", "bSupportsIPhone", out biPhoneAllowed);
+				Ini.GetBool("/Script/IOSRuntimeSettings.IOSRuntimeSettings", "bSupportsIPad", out biPadAllowed);
+				if (biPhoneAllowed && biPadAllowed)
+				{
+					RunTimeIOSDevices = "1,2";
+				}
+				else if (biPadAllowed)
+				{
+					RunTimeIOSDevices = "2";
+				}
+				else if (biPhoneAllowed)
+				{
+					RunTimeIOSDevices = "1";
+				}
 
-			List<string> ProjectArches = new List<string>();
-			bool bBuild = true;
-			if (Ini.GetBool("/Script/IOSRuntimeSettings.IOSRuntimeSettings", "bDevForArmV7", out bBuild) && bBuild)
-			{
-				ProjectArches.Add("armv7");
-			}
-			if (Ini.GetBool("/Script/IOSRuntimeSettings.IOSRuntimeSettings", "bDevForArm64", out bBuild) && bBuild)
-			{
-				ProjectArches.Add("arm64");
-			}
-			if (Ini.GetBool("/Script/IOSRuntimeSettings.IOSRuntimeSettings", "bDevForArmV7S", out bBuild) && bBuild)
-			{
-				ProjectArches.Add("armv7s");
-			}
+				List<string> ProjectArches = new List<string>();
+				bool bBuild = true;
+				if (Ini.GetBool("/Script/IOSRuntimeSettings.IOSRuntimeSettings", "bDevForArmV7", out bBuild) && bBuild)
+				{
+					ProjectArches.Add("armv7");
+				}
+				if (Ini.GetBool("/Script/IOSRuntimeSettings.IOSRuntimeSettings", "bDevForArm64", out bBuild) && bBuild)
+				{
+					ProjectArches.Add("arm64");
+				}
+				if (Ini.GetBool("/Script/IOSRuntimeSettings.IOSRuntimeSettings", "bDevForArmV7S", out bBuild) && bBuild)
+				{
+					ProjectArches.Add("armv7s");
+				}
 
-			// force armv7 if something went wrong
-			if (ProjectArches.Count == 0)
-			{
-				ProjectArches.Add("armv7");
-			}
-			NonShippingArchitectures = ProjectArches[0];
-			for (int Index = 1; Index < ProjectArches.Count; ++Index)
-			{
-				NonShippingArchitectures += "," + ProjectArches[Index];
-			}
+				// force armv7 if something went wrong
+				if (ProjectArches.Count == 0)
+				{
+					ProjectArches.Add("armv7");
+				}
+				NonShippingArchitectures = ProjectArches[0];
+				for (int Index = 1; Index < ProjectArches.Count; ++Index)
+				{
+					NonShippingArchitectures += "," + ProjectArches[Index];
+				}
 
-			ProjectArches.Clear();
-			if (Ini.GetBool("/Script/IOSRuntimeSettings.IOSRuntimeSettings", "bShipForArmV7", out bBuild) && bBuild)
-			{
-				ProjectArches.Add("armv7");
-			}
-			if (Ini.GetBool("/Script/IOSRuntimeSettings.IOSRuntimeSettings", "bShipForArm64", out bBuild) && bBuild)
-			{
-				ProjectArches.Add("arm64");
-			}
-			if (Ini.GetBool("/Script/IOSRuntimeSettings.IOSRuntimeSettings", "bShipForArmV7S", out bBuild) && bBuild)
-			{
-				ProjectArches.Add("armv7s");
-			}
+				ProjectArches.Clear();
+				if (Ini.GetBool("/Script/IOSRuntimeSettings.IOSRuntimeSettings", "bShipForArmV7", out bBuild) && bBuild)
+				{
+					ProjectArches.Add("armv7");
+				}
+				if (Ini.GetBool("/Script/IOSRuntimeSettings.IOSRuntimeSettings", "bShipForArm64", out bBuild) && bBuild)
+				{
+					ProjectArches.Add("arm64");
+				}
+				if (Ini.GetBool("/Script/IOSRuntimeSettings.IOSRuntimeSettings", "bShipForArmV7S", out bBuild) && bBuild)
+				{
+					ProjectArches.Add("armv7s");
+				}
 
-			// force armv7 if something went wrong
-			if (ProjectArches.Count == 0)
-			{
-				ProjectArches.Add("armv7");
-				ProjectArches.Add("arm64");
-			}
-			ShippingArchitectures = ProjectArches[0];
-			for (int Index = 1; Index < ProjectArches.Count; ++Index)
-			{
-				ShippingArchitectures += "," + ProjectArches[Index];
-			}
+				// force armv7 if something went wrong
+				if (ProjectArches.Count == 0)
+				{
+					ProjectArches.Add("armv7");
+					ProjectArches.Add("arm64");
+				}
+				ShippingArchitectures = ProjectArches[0];
+				for (int Index = 1; Index < ProjectArches.Count; ++Index)
+				{
+					ShippingArchitectures += "," + ProjectArches[Index];
+				}
+
+				// determine if we need to generate the dsym
+				Ini.GetBool("/Script/IOSRuntimeSettings.IOSRuntimeSettings", "bGenerateSYMFile", out BuildConfiguration.bGeneratedSYMFile);
 
-			// determine if we need to generate the dsym
-			Ini.GetBool("/Script/IOSRuntimeSettings.IOSRuntimeSettings", "bGenerateSYMFile", out BuildConfiguration.bGeneratedSYMFile);
+				bInitializedProject = true;
+			}
 		}
 
 		/**
