CIS fix for Mac SDK

diff --git a/Engine/Source/Programs/UnrealBuildTool/ToolChain/AppleToolChain.cs b/Engine/Source/Programs/UnrealBuildTool/ToolChain/AppleToolChain.cs
--- a/Engine/Source/Programs/UnrealBuildTool/ToolChain/AppleToolChain.cs
+++ b/Engine/Source/Programs/UnrealBuildTool/ToolChain/AppleToolChain.cs
@@ -65,12 +65,6 @@ namespace UnrealBuildTool
 					{
 						// on the Mac, we can just get the directory name
 						SubDirs = System.IO.Directory.GetDirectories(BaseSDKDir);
-
-						// make sure we have a valid SDK directory
-						if (!Directory.Exists(Path.Combine(BaseSDKDir, OSPrefix + PlatformSDKVersion + ".sdk")))
-						{
-							throw new BuildException("Invalid SDK {0}{1}.sdk, not found in {2}", OSPrefix, PlatformSDKVersion, BaseSDKDir);
-						}
 					}
 					else
 					{
@@ -136,6 +130,12 @@ namespace UnrealBuildTool
 				}
 			}
 
+			// make sure we have a valid SDK directory
+			if (Utils.IsRunningOnMono && !Directory.Exists(Path.Combine(BaseSDKDir, OSPrefix + PlatformSDKVersion + ".sdk")))
+			{
+				throw new BuildException("Invalid SDK {0}{1}.sdk, not found in {2}", OSPrefix, PlatformSDKVersion, BaseSDKDir);
+			}
+
 			if (bVerbose && !ProjectFileGenerator.bGenerateProjectFiles)
 			{
 				if (BuildHostPlatform.Current.Platform != UnrealTargetPlatform.Mac)
