Allow targets to specify whether they want to build all plugins (rather than just the ones they have enabled), and turn it on for the vanilla UE4Editor target. Content-only projects may want to enable plugins, but they have no target to build.

diff --git a/Engine/Source/Programs/UnrealBuildTool/Configuration/UEBuildTarget.cs b/Engine/Source/Programs/UnrealBuildTool/Configuration/UEBuildTarget.cs
--- a/Engine/Source/Programs/UnrealBuildTool/Configuration/UEBuildTarget.cs
+++ b/Engine/Source/Programs/UnrealBuildTool/Configuration/UEBuildTarget.cs
@@ -2625,7 +2625,7 @@ namespace UnrealBuildTool
 			}
 
 			// Set the list of plugins that should be built
-			if (bPrecompile && TargetType != TargetRules.TargetType.Program)
+			if ((bPrecompile && TargetType != TargetRules.TargetType.Program) || Rules.bBuildAllPlugins)
 			{
 				BuildPlugins = new List<PluginInfo>(ValidPlugins);
 			}
diff --git a/Engine/Source/Programs/UnrealBuildTool/System/RulesCompiler.cs b/Engine/Source/Programs/UnrealBuildTool/System/RulesCompiler.cs
--- a/Engine/Source/Programs/UnrealBuildTool/System/RulesCompiler.cs
+++ b/Engine/Source/Programs/UnrealBuildTool/System/RulesCompiler.cs
@@ -471,6 +471,12 @@ namespace UnrealBuildTool
 		/// </summary>
 		public UnrealTargetConfiguration UndecoratedConfiguration = UnrealTargetConfiguration.Development;
 
+		/// <summary>
+		/// Build all the plugins that we can find, even if they're not enabled. This is particularly useful for content-only projects, 
+		/// where you're building the UE4Editor target but running it with a game that enables a plugin.
+		/// </summary>
+		public bool bBuildAllPlugins = false;
+
 		/// <summary>
 		/// A list of additional plugins which need to be included in this target. This allows referencing non-optional plugin modules
 		/// which cannot be disabled, and allows building against specific modules in program targets which do not fit the categories
diff --git a/Engine/Source/UE4Editor.Target.cs b/Engine/Source/UE4Editor.Target.cs
--- a/Engine/Source/UE4Editor.Target.cs
+++ b/Engine/Source/UE4Editor.Target.cs
@@ -8,6 +8,7 @@ public class UE4EditorTarget : TargetRules
 	public UE4EditorTarget( TargetInfo Target )
 	{
 		Type = TargetType.Editor;
+		bBuildAllPlugins = true;
 	}
 
 	public override void SetupBinaries(
