Enable ld.gold on r10e (clang 3.6) and add in arguments for arm64 compiling
#ue4
#android

diff --git a/Engine/Source/Programs/UnrealBuildTool/Android/AndroidToolChain.cs b/Engine/Source/Programs/UnrealBuildTool/Android/AndroidToolChain.cs
--- a/Engine/Source/Programs/UnrealBuildTool/Android/AndroidToolChain.cs
+++ b/Engine/Source/Programs/UnrealBuildTool/Android/AndroidToolChain.cs
@@ -407,7 +407,25 @@ namespace UnrealBuildTool
 			}
 			else if (Architecture == "-arm64")
 			{
-				// todo
+				Result += " -funwind-tables";			// Just generates any needed static data, affects no code 
+				Result += " -fstack-protector";			// Emits extra code to check for buffer overflows
+				Result += " -fno-strict-aliasing";		// Prevents unwanted or invalid optimizations that could produce incorrect code
+				Result += " -fpic";						// Generates position-independent code (PIC) suitable for use in a shared library
+				Result += " -fno-exceptions";			// Do not enable exception handling, generates extra code needed to propagate exceptions
+				Result += " -fno-rtti";					// 
+				Result += " -fno-short-enums";			// Do not allocate to an enum type only as many bytes as it needs for the declared range of possible values
+
+				Result += " -march=armv8-a";
+				//Result += " -mfloat-abi=softfp";
+				//Result += " -mfpu=vfpv3-d16";			//@todo android: UE3 was just vfp. arm7a should all support v3 with 16 registers
+
+				// Some switches interfere with on-device debugging
+				if (CompileEnvironment.Config.Target.Configuration != CPPTargetConfiguration.Debug)
+				{
+					Result += " -ffunction-sections";   // Places each function in its own section of the output file, linker may be able to perform opts to improve locality of reference
+				}
+
+				Result += " -fsigned-char";				// Treat chars as signed //@todo android: any concerns about ABI compatibility with libs here?
 			}
 			else if (Architecture == "-x86")
 			{
@@ -522,6 +540,11 @@ namespace UnrealBuildTool
 				Result += " -Wl,--fix-cortex-a8";		// required to route around a CPU bug in some Cortex-A8 implementations
 			}
 
+			if (ClangVersionFloat >= 3.6f)
+			{
+				Result += " -fuse-ld=gold";				// ld.gold is available in r10e (clang 3.6)
+			}
+
             // verbose output from the linker
             // Result += " -v";
 
