Removed some debug messages displayed during Xcode project generation

diff --git a/Engine/Source/Programs/UnrealBuildTool/System/XcodeProjectFileGenerator.cs b/Engine/Source/Programs/UnrealBuildTool/System/XcodeProjectFileGenerator.cs
--- a/Engine/Source/Programs/UnrealBuildTool/System/XcodeProjectFileGenerator.cs
+++ b/Engine/Source/Programs/UnrealBuildTool/System/XcodeProjectFileGenerator.cs
@@ -1210,9 +1210,6 @@ namespace UnrealBuildTool
 					DepotRoot = DepotRoot.Substring(0, DepotRoot.LastIndexOf(Path.DirectorySeparatorChar));
 					DepotRoot = DepotRoot.TrimEnd(Path.DirectorySeparatorChar);
 
-					Log.TraceInformation(string.Format("Target.TargetFilePath={0}", Target.TargetFilePath));
-					Log.TraceInformation(string.Format("DepotRoot={0}", DepotRoot));
-
                     if (Utils.IsFileUnderDirectory(Target.TargetFilePath, DepotRoot))
                     {
                         string PreviousPath = Target.TargetFilePath;
@@ -1227,9 +1224,6 @@ namespace UnrealBuildTool
                         {
                             TargetBinariesFolder = PreviousPath;
                         }
-
-						Log.TraceInformation(string.Format("TargetBinariesFolder={0}", TargetBinariesFolder));
-						Log.TraceInformation(string.Format("ExeBaseName={0}", ExeBaseName));
                     }
                 }
 
