Mac build scripts now properly handle renamed launcher.

diff --git a/Engine/Source/Programs/UnrealBuildTool/Mac/MacToolChain.cs b/Engine/Source/Programs/UnrealBuildTool/Mac/MacToolChain.cs
--- a/Engine/Source/Programs/UnrealBuildTool/Mac/MacToolChain.cs
+++ b/Engine/Source/Programs/UnrealBuildTool/Mac/MacToolChain.cs
@@ -462,21 +462,21 @@ namespace UnrealBuildTool
 
 		private string LoadLauncherDisplayVersion()
 		{
-			string[] VersionHeader = Utils.ReadAllText("../Source/Programs/NoRedist/UnrealEngineLauncher/Private/Version.h").Replace("\r\n", "\n").Replace("\t", " ").Split('\n');
+            string[] VersionHeader = Utils.ReadAllText("../../Portal/Source/Layers/DataAccess/Public/Version.h").Replace("\r\n", "\n").Replace("\t", " ").Split('\n');
 			string LauncherVersionMajor = "1";
 			string LauncherVersionMinor = "0";
 			string LauncherVersionPatch = "0";
 			foreach (string Line in VersionHeader)
 			{
-				if (Line.StartsWith("#define LAUNCHER_MAJOR_VERSION "))
+                if (Line.StartsWith("#define PORTAL_MAJOR_VERSION "))
 				{
 					LauncherVersionMajor = Line.Split(' ')[2];
 				}
-				else if (Line.StartsWith("#define LAUNCHER_MINOR_VERSION "))
+                else if (Line.StartsWith("#define PORTAL_MINOR_VERSION "))
 				{
 					LauncherVersionMinor = Line.Split(' ')[2];
 				}
-				else if (Line.StartsWith("#define LAUNCHER_PATCH_VERSION "))
+                else if (Line.StartsWith("#define PORTAL_PATCH_VERSION "))
 				{
 					LauncherVersionPatch = Line.Split(' ')[2];
 				}
@@ -948,9 +948,8 @@ namespace UnrealBuildTool
 					AppendMacLine(FinalizeAppBundleScript, "sh \"{0}\" \"{1}\"", ConvertPath(DylibCopyScriptPath).Replace("$", "\\$"), ExeName);
 
 					string IconName = "UE4";
-					string BundleVersion = ExeName.StartsWith("UnrealEngineLauncher") ? LoadLauncherDisplayVersion() : LoadEngineDisplayVersion();
-					string EngineSourcePath = ConvertPath(Directory.GetCurrentDirectory()).Replace("$", "\\$");
-
+					string BundleVersion = ExeName.StartsWith("EpicGamesLauncher") ? LoadLauncherDisplayVersion() : LoadEngineDisplayVersion();
+                    string EngineSourcePath = ConvertPath(Directory.GetCurrentDirectory()).Replace("$", "\\$");
 					string UProjectFilePath = UProjectInfo.GetProjectFilePath(GameName);
 					string CustomResourcesPath = "";
 					string CustomBuildPath = "";
@@ -969,8 +968,9 @@ namespace UnrealBuildTool
 					}
 					else
 					{
+                        string ResourceParentFolderName = ExeName.StartsWith("EpicGamesLauncher") ? "Application" : GameName;
 						string FullUProjectFilePath = Path.GetFullPath(UProjectFilePath);
-						CustomResourcesPath = Path.GetDirectoryName(FullUProjectFilePath) + "/Source/" + GameName + "/Resources/Mac";
+						CustomResourcesPath = Path.GetDirectoryName(FullUProjectFilePath) + "/Source/" + ResourceParentFolderName + "/Resources/Mac";
 						CustomBuildPath = Path.GetDirectoryName(FullUProjectFilePath) + "/Build/Mac";
 					}
 
