Update project generator for AndroidWorks compatibility
Reordered vcxproj to better follow:
http://blogs.msdn.com/b/visualstudio/archive/2010/05/14/a-guide-to-vcxproj-and-props-file-structure.aspx
#ue4
#android

diff --git a/Engine/Source/Programs/UnrealBuildTool/Android/AndroidProjectGenerator.cs b/Engine/Source/Programs/UnrealBuildTool/Android/AndroidProjectGenerator.cs
--- a/Engine/Source/Programs/UnrealBuildTool/Android/AndroidProjectGenerator.cs
+++ b/Engine/Source/Programs/UnrealBuildTool/Android/AndroidProjectGenerator.cs
@@ -15,6 +15,7 @@ namespace UnrealBuildTool
 	{
 		static bool CheckedForNsight = false;		// whether we have checked for a recent enough version of Nsight yet
 		static bool NsightInstalled = false;		// true if a recent enough version of Nsight is installed
+		static int NsightVersionCode = 0;			// version code matching detected Nsight
 
 		/**
 		 *	Check to see if a recent enough version of Nsight is installed.
@@ -53,14 +54,22 @@ namespace UnrealBuildTool
 			// grab the version info from the DLL
 			FileVersionInfo NsightVersion = FileVersionInfo.GetVersionInfo(NsightDllPath);
 
-			if (NsightVersion.ProductMajorPart >= 2)
+			if (NsightVersion.ProductMajorPart >= 3)
+			{
+				// Nsight 3.0+ should be valid
+				NsightVersionCode = 9;
+				NsightInstalled = true;
+			}
+			else if (NsightVersion.ProductMajorPart == 2)
 			{
 				// Nsight 2.0+ should be valid
+				NsightVersionCode = 6;
 				NsightInstalled = true;
 			}
 			else if ((NsightVersion.ProductMajorPart == 1) && (NsightVersion.ProductMinorPart >= 5))
 			{
 				// Nsight 1.5+ should be valid
+				NsightVersionCode = 6;
 				NsightInstalled = true;
 			}
 
@@ -128,7 +137,7 @@ namespace UnrealBuildTool
 			}
 
 			return 	"	<PropertyGroup Label=\"NsightTegraProject\">" + ProjectFileGenerator.NewLine +
-					"		<NsightTegraProjectRevisionNumber>6</NsightTegraProjectRevisionNumber>" + ProjectFileGenerator.NewLine +
+					"		<NsightTegraProjectRevisionNumber>" + NsightVersionCode.ToString() + "</NsightTegraProjectRevisionNumber>" + ProjectFileGenerator.NewLine +
 					"	</PropertyGroup>" + ProjectFileGenerator.NewLine;
 		}
 
@@ -149,6 +158,25 @@ namespace UnrealBuildTool
 			return "ExternalBuildSystem";
 		}
 
+		/**
+		 *	Return the platform toolset string to write into the project configuration
+		 *	
+		 *	@param	InPlatform			The UnrealTargetPlatform being built
+		 *	@param	InConfiguration		The UnrealTargetConfiguration being built
+		 *	
+		 *	@return	string				The custom configuration section for the project file; Empty string if it doesn't require one
+		 */
+		public override string GetVisualStudioPlatformToolsetString(UnrealTargetPlatform InPlatform, UnrealTargetConfiguration InConfiguration, VCProjectFile InProjectFile)
+		{
+			if (!IsNsightInstalled())
+			{
+				return "\t\t<PlatformToolset>" + VCProjectFileGenerator.ProjectFilePlatformToolsetVersionString + "</PlatformToolset>" + ProjectFileGenerator.NewLine;
+			}
+
+			return "\t\t<PlatformToolset>" + VCProjectFileGenerator.ProjectFilePlatformToolsetVersionString + "</PlatformToolset>" + ProjectFileGenerator.NewLine
+				+ "\t\t<AndroidNativeAPI>UseTarget</AndroidNativeAPI>" + ProjectFileGenerator.NewLine;
+		}
+
 		/**
 		 * Return any custom paths for VisualStudio this platform requires
 		 * This include ReferencePath, LibraryPath, LibraryWPath, IncludePath and ExecutablePath.
@@ -196,12 +224,12 @@ namespace UnrealBuildTool
 			AdditionalLibDirs += @";$(AdditionalLibraryDirectories)";
 
 			string PathsLines = 
-				"		<IncludePath/>" + ProjectFileGenerator.NewLine +
-				"		<ReferencePath/>" + ProjectFileGenerator.NewLine +
-				"		<LibraryPath/>" + ProjectFileGenerator.NewLine +
-				"		<LibraryWPath/>" + ProjectFileGenerator.NewLine +
-				"		<SourcePath/>" + ProjectFileGenerator.NewLine +
-				"		<ExcludePath/>" + ProjectFileGenerator.NewLine +
+				"		<IncludePath />" + ProjectFileGenerator.NewLine +
+				"		<ReferencePath />" + ProjectFileGenerator.NewLine +
+				"		<LibraryPath />" + ProjectFileGenerator.NewLine +
+				"		<LibraryWPath />" + ProjectFileGenerator.NewLine +
+				"		<SourcePath />" + ProjectFileGenerator.NewLine +
+				"		<ExcludePath />" + ProjectFileGenerator.NewLine +
 				"		<AndroidAttach>False</AndroidAttach>" + ProjectFileGenerator.NewLine +
 				"		<DebuggerFlavor>AndroidDebugger</DebuggerFlavor>" + ProjectFileGenerator.NewLine +
 				"		<OverrideAPKPath>" + APKPath + "</OverrideAPKPath>" + ProjectFileGenerator.NewLine +
diff --git a/Engine/Source/Programs/UnrealBuildTool/System/VCProject.cs b/Engine/Source/Programs/UnrealBuildTool/System/VCProject.cs
--- a/Engine/Source/Programs/UnrealBuildTool/System/VCProject.cs
+++ b/Engine/Source/Programs/UnrealBuildTool/System/VCProject.cs
@@ -493,25 +493,66 @@
 				VCProjectFileContent.Append( AdditionalPropertyGroups );
 			}
 
+			// Project globals (project GUID, project type, SCC bindings, etc)
+			{
+				VCProjectFileContent.Append(
+					"	<PropertyGroup Label=\"Globals\">" + ProjectFileGenerator.NewLine +
+					"		<ProjectGuid>" + ProjectGUID.ToString("B").ToUpperInvariant() + "</ProjectGuid>" + ProjectFileGenerator.NewLine +
+					"		<Keyword>MakeFileProj</Keyword>" + ProjectFileGenerator.NewLine +
+					"		<RootNamespace>" + ProjectName + "</RootNamespace>" + ProjectFileGenerator.NewLine);
+
+				VCProjectFileContent.Append(
+					"	</PropertyGroup>" + ProjectFileGenerator.NewLine);
+			}
+
+			VCProjectFileContent.Append(
+				"	<Import Project=\"$(VCTargetsPath)\\Microsoft.Cpp.Default.props\" />" + ProjectFileGenerator.NewLine);
+
+			// Write each project configuration PreDefaultProps section
+			foreach (var ConfigurationTuple in ProjectConfigurationNameAndConfigurations)
+			{
+				var ProjectConfigurationName = ConfigurationTuple.Item1;
+				var TargetConfiguration = ConfigurationTuple.Item2;
+				foreach (var PlatformTuple in ProjectPlatformNameAndPlatforms)
+				{
+					var ProjectPlatformName = PlatformTuple.Item1;
+					var TargetPlatform = PlatformTuple.Item2;
+					WritePreDefaultPropsConfiguration(TargetPlatform, TargetConfiguration, ProjectPlatformName, ProjectConfigurationName, VCProjectFileContent);
+				}
+			}
+
+
+			VCProjectFileContent.Append(
+				"	<Import Project=\"$(VCTargetsPath)\\Microsoft.Cpp.props\" />" + ProjectFileGenerator.NewLine +
+				"	<ImportGroup Label=\"ExtensionSettings\" />" + ProjectFileGenerator.NewLine +
+				"	<PropertyGroup Label=\"UserMacros\" />" + ProjectFileGenerator.NewLine
+				);
+
+			// Write each project configuration
+			foreach( var Combination in ProjectConfigAndTargetCombinations )
+			{
+				WriteConfiguration( ProjectName, Combination, VCProjectFileContent, bGenerateUserFileContent? VCUserFileContent : null );
+			}
+
 			// Source folders and files
 			{
 				var LocalAliasedFiles = new List<AliasedFile>(AliasedFiles);
 
-				foreach( var CurFile in SourceFiles )
+				foreach (var CurFile in SourceFiles)
 				{
 					// We want all source file and directory paths in the project files to be relative to the project file's
 					// location on the disk.  Convert the path to be relative to the project file directory
-					var ProjectRelativeSourceFile = Utils.MakePathRelativeTo( CurFile.FilePath, Path.GetDirectoryName( ProjectFilePath ) );
+					var ProjectRelativeSourceFile = Utils.MakePathRelativeTo(CurFile.FilePath, Path.GetDirectoryName(ProjectFilePath));
 
 					// By default, files will appear relative to the project file in the solution.  This is kind of the normal Visual
 					// Studio way to do things, but because our generated project files are emitted to intermediate folders, if we always
 					// did this it would yield really ugly paths int he solution explorer
-					string FilterRelativeSourceDirectory = Path.GetDirectoryName( ProjectRelativeSourceFile );
+					string FilterRelativeSourceDirectory = Path.GetDirectoryName(ProjectRelativeSourceFile);
 
 					// Use the specified relative base folder
-					if( CurFile.RelativeBaseFolder != null )	// NOTE: We are looking for null strings, not empty strings!
+					if (CurFile.RelativeBaseFolder != null)	// NOTE: We are looking for null strings, not empty strings!
 					{
-						FilterRelativeSourceDirectory = Path.GetDirectoryName( Utils.MakePathRelativeTo( CurFile.FilePath, CurFile.RelativeBaseFolder ) );
+						FilterRelativeSourceDirectory = Path.GetDirectoryName(Utils.MakePathRelativeTo(CurFile.FilePath, CurFile.RelativeBaseFolder));
 					}
 
 					LocalAliasedFiles.Add(new AliasedFile(ProjectRelativeSourceFile, FilterRelativeSourceDirectory));
@@ -556,49 +597,10 @@
 				}
 
 				VCProjectFileContent.Append(
-					"	</ItemGroup>" + ProjectFileGenerator.NewLine );
+					"	</ItemGroup>" + ProjectFileGenerator.NewLine);
 
 				VCFiltersFileContent.Append(
-					"	</ItemGroup>" + ProjectFileGenerator.NewLine );
+					"	</ItemGroup>" + ProjectFileGenerator.NewLine);
 			}
 
-
-			// Project globals (project GUID, project type, SCC bindings, etc)
-			{
-				VCProjectFileContent.Append(
-					"	<PropertyGroup Label=\"Globals\">" + ProjectFileGenerator.NewLine +
-					"		<ProjectGuid>" + ProjectGUID.ToString( "B" ).ToUpperInvariant() + "</ProjectGuid>" + ProjectFileGenerator.NewLine +
-					"		<Keyword>MakeFileProj</Keyword>" + ProjectFileGenerator.NewLine +
-					"		<RootNamespace>" + ProjectName + "</RootNamespace>" + ProjectFileGenerator.NewLine);
-
-				VCProjectFileContent.Append(
-					"	</PropertyGroup>" + ProjectFileGenerator.NewLine);
-			}
-
-			// Write each project configuration PreDefaultProps section
-			foreach (var ConfigurationTuple in ProjectConfigurationNameAndConfigurations)
-			{
-				var ProjectConfigurationName = ConfigurationTuple.Item1;
-				var TargetConfiguration = ConfigurationTuple.Item2;
-				foreach (var PlatformTuple in ProjectPlatformNameAndPlatforms)
-				{
-					var ProjectPlatformName = PlatformTuple.Item1;
-					var TargetPlatform = PlatformTuple.Item2;
-					WritePreDefaultPropsConfiguration( TargetPlatform, TargetConfiguration, ProjectPlatformName, ProjectConfigurationName, VCProjectFileContent );
-				}
-			}
-
-			VCProjectFileContent.Append(
-				"	<Import Project=\"$(VCTargetsPath)\\Microsoft.Cpp.Default.props\" />" + ProjectFileGenerator.NewLine +
-				"	<Import Project=\"$(VCTargetsPath)\\Microsoft.Cpp.props\" />" + ProjectFileGenerator.NewLine +
-				"	<ImportGroup Label=\"ExtensionSettings\" />" + ProjectFileGenerator.NewLine +
-				"	<PropertyGroup Label=\"UserMacros\" />" + ProjectFileGenerator.NewLine
-				);
-
-			// Write each project configuration
-			foreach( var Combination in ProjectConfigAndTargetCombinations )
-			{
-				WriteConfiguration( ProjectName, Combination, VCProjectFileContent, bGenerateUserFileContent? VCUserFileContent : null );
-			}
-
 			// For Rocket, include engine source in the source search paths. We never build it locally, so the debugger can't find it.
