UEB-124 - Invoking UAT/UBT should not wait for mutex
* UAT/UBT will now fail immediately instead.

diff --git a/Engine/Source/Programs/AutomationTool/Utils.cs b/Engine/Source/Programs/AutomationTool/Utils.cs
--- a/Engine/Source/Programs/AutomationTool/Utils.cs
+++ b/Engine/Source/Programs/AutomationTool/Utils.cs
@@ -704,9 +704,7 @@ namespace AutomationTool
 			{
 				if (!bCreatedMutex)
 				{
-					Log.WriteLine(TraceEventType.Warning, "Another instance of {0} is already running. Waiting until it exists.", ExecutingAssemblyLocation);
-					// If this instance didn't create the mutex, wait for the existing mutex to be released by the mutex's creator.
-					SingleInstanceMutex.WaitOne();
+                    throw new AutomationException("Another instance of {0} is already running.", ExecutingAssemblyLocation);
 				}
 				else
 				{
diff --git a/Engine/Source/Programs/UnrealBuildTool/System/UnrealBuildTool.cs b/Engine/Source/Programs/UnrealBuildTool/System/UnrealBuildTool.cs
--- a/Engine/Source/Programs/UnrealBuildTool/System/UnrealBuildTool.cs
+++ b/Engine/Source/Programs/UnrealBuildTool/System/UnrealBuildTool.cs
@@ -18,9 +18,6 @@ namespace UnrealBuildTool
         /** Time at which control fist passes to UBT. */
         static public DateTime StartTime;
 
-        /** Time we waited for the mutex. Tracked for proper active execution time collection. */
-        static public TimeSpan MutexWaitTime = TimeSpan.Zero;
-
         static public string BuildGuid = Guid.NewGuid().ToString();
 
         /** Total time spent building projects. */
@@ -960,10 +957,7 @@ namespace UnrealBuildTool
                     SingleInstanceMutex = new Mutex(true, MutexName, out bCreatedMutex);
                     if (!bCreatedMutex)
                     {
-                        // If this instance didn't create the mutex, wait for the existing mutex to be released by the mutex's creator.
-                        DateTime MutexWaitStartTime = DateTime.UtcNow;
-                        SingleInstanceMutex.WaitOne();
-                        MutexWaitTime = DateTime.UtcNow - MutexWaitStartTime;
+                        throw new BuildException("Mutex {0} already set, indicating that a conflicting instance of {1} is already running.", MutexName, ExecutingAssembly.GetFilename());
                     }
                 }
 
@@ -1313,7 +1307,7 @@ namespace UnrealBuildTool
                         }
                     }
                     // Print some performance info
-                    var BuildDuration = (DateTime.UtcNow - StartTime - MutexWaitTime).TotalSeconds;
+                    var BuildDuration = (DateTime.UtcNow - StartTime).TotalSeconds;
                     if (BuildConfiguration.bPrintPerformanceInfo)
                     {
                         Log.TraceInformation("GetIncludes time: " + CPPEnvironment.TotalTimeSpentGettingIncludes + "s (" + CPPEnvironment.TotalIncludesRequested + " includes)" );
@@ -1330,7 +1324,6 @@ namespace UnrealBuildTool
 
                     Telemetry.SendEvent("PerformanceInfo.2",
                         "TotalExecutionTimeSec", BuildDuration.ToString("0.00"),
-                        "MutexWaitTimeSec", MutexWaitTime.TotalSeconds.ToString("0.00"),
                         "TotalTimeSpentGettingIncludesSec", CPPEnvironment.TotalTimeSpentGettingIncludes.ToString("0.00"),
                         "TotalIncludesRequested", CPPEnvironment.TotalIncludesRequested.ToString(),
                         "DirectIncludeCacheMissesTotalTimeSec", CPPEnvironment.DirectIncludeCacheMissesTotalTime.ToString("0.00"),
@@ -1359,7 +1352,7 @@ namespace UnrealBuildTool
             Telemetry.Shutdown();
 
             // Print some performance info
-            Log.TraceVerbose("Execution time: {0}", (DateTime.UtcNow - StartTime - MutexWaitTime).TotalSeconds);
+            Log.TraceVerbose("Execution time: {0}", (DateTime.UtcNow - StartTime).TotalSeconds);
 
 			if (ExtendedErrorCode != 0)
 			{
@@ -2000,7 +1993,7 @@ namespace UnrealBuildTool
             }
 
             // Figure out how long we took to execute.
-            double BuildDuration = (DateTime.UtcNow - StartTime - MutexWaitTime).TotalSeconds;
+            double BuildDuration = (DateTime.UtcNow - StartTime).TotalSeconds;
             if (ExecutorName == "Local" || ExecutorName == "Distcc" || ExecutorName == "SNDBS")
             {
                 Log.WriteLineIf(BuildConfiguration.bLogDetailedActionStats || BuildConfiguration.bPrintDebugInfo,
