Removed the workaround for the PCH problem that was solved in second Xcode 7 beta

diff --git a/Engine/Source/Programs/UnrealBuildTool/Mac/MacToolChain.cs b/Engine/Source/Programs/UnrealBuildTool/Mac/MacToolChain.cs
--- a/Engine/Source/Programs/UnrealBuildTool/Mac/MacToolChain.cs
+++ b/Engine/Source/Programs/UnrealBuildTool/Mac/MacToolChain.cs
@@ -107,7 +107,6 @@ namespace UnrealBuildTool
 
 			if (MacOSSDKVersionFloat < 10.9f && MacOSSDKVersionFloat >= 10.11f)
 			{
-				Log.TraceInformationOnce("Got into 10.11 land");
 				Result += " -Wno-inconsistent-missing-override"; // too many missing overrides...
 				Result += " -Wno-unused-local-typedef"; // PhysX has some, hard to remove
 			}
diff --git a/Engine/Source/Programs/UnrealBuildTool/Mac/UEBuildMac.cs b/Engine/Source/Programs/UnrealBuildTool/Mac/UEBuildMac.cs
--- a/Engine/Source/Programs/UnrealBuildTool/Mac/UEBuildMac.cs
+++ b/Engine/Source/Programs/UnrealBuildTool/Mac/UEBuildMac.cs
@@ -163,13 +163,6 @@ namespace UnrealBuildTool
 
 		public override void ValidateBuildConfiguration(CPPTargetConfiguration Configuration, CPPTargetPlatform Platform, bool bCreateDebugInfo)
 		{
-			// Xcode 7 doesn't work currently with PCH files
-			if (MacToolChain.MacOSSDKVersion == "10.11")
-			{
-				Log.TraceInformation("=== Disabling PCH files due to Xcode 7 compiler bug! ===");
-				BuildConfiguration.bUsePCHFiles = false;
-			}
-
 			if (BuildHostPlatform.Current.Platform != UnrealTargetPlatform.Mac)
 			{
 				// @todo: Temporarily disable precompiled header files when building remotely due to errors
