UBT: Fix problem with ordering modules when compiling for hot-reload.

diff --git a/Engine/Source/Programs/UnrealBuildTool/Configuration/UEBuildTarget.cs b/Engine/Source/Programs/UnrealBuildTool/Configuration/UEBuildTarget.cs
--- a/Engine/Source/Programs/UnrealBuildTool/Configuration/UEBuildTarget.cs
+++ b/Engine/Source/Programs/UnrealBuildTool/Configuration/UEBuildTarget.cs
@@ -771,6 +771,10 @@ namespace UnrealBuildTool
 		[NonSerialized]
 		protected List<OnlyModule> OnlyModules = new List<OnlyModule>();
 
+		/** Kept to determine the correct module parsing order when filtering modules. */
+		[NonSerialized]
+		protected List<UEBuildBinary> NonFilteredModules = new List<UEBuildBinary>();
+
 		/** true if target should be compiled in monolithic mode, false if not */
 		protected bool bCompileMonolithic = false;
 
@@ -1465,6 +1469,32 @@ namespace UnrealBuildTool
 			}
 		}
 
+		/// <summary>
+		/// Gathers dependency modules for given binaries list.
+		/// </summary>
+		/// <param name="Binaries">Binaries list.</param>
+		/// <returns>Dependency modules set.</returns>
+		static HashSet<UEBuildModuleCPP> GatherDependencyModules(List<UEBuildBinary> Binaries)
+		{
+			var Output = new HashSet<UEBuildModuleCPP>();
+
+			foreach (var Binary in Binaries)
+			{
+				var DependencyModules = Binary.GetAllDependencyModules(bIncludeDynamicallyLoaded: false, bForceCircular: false);
+				foreach (var Module in DependencyModules.OfType<UEBuildModuleCPP>())
+				{
+					if (!Module.bIncludedInTarget)
+					{
+						throw new BuildException("Expecting module {0} to have bIncludeInTarget set", Module.Name);
+					}
+
+					Output.Add(Module);
+				}
+			}
+
+			return Output;
+		}
+
 		/** Builds the target, appending list of output files and returns building result. */
 		public ECompilationResult Build(IUEToolChain TargetToolChain, out List<FileItem> OutputItems, out List<UHTModuleInfo> UObjectModules, out string EULAViolationWarning)
 		{
@@ -1568,19 +1598,14 @@ namespace UnrealBuildTool
 				// Reconstruct a full list of binaries. Binaries which aren't compiled are stripped out of AppBinaries, but we still need to scan them for UHT.
 				List<UEBuildBinary> AllAppBinaries = AppBinaries.Union(PrecompiledBinaries).ToList();
 	
-				var ModulesToGenerateHeadersFor = new HashSet<UEBuildModuleCPP>();
-				foreach( var Binary in AllAppBinaries )
+				var ModulesToGenerateHeadersFor = GatherDependencyModules(AllAppBinaries);
+
+				if (OnlyModules.Count > 0)
 				{
-					var DependencyModules = Binary.GetAllDependencyModules(bIncludeDynamicallyLoaded: false, bForceCircular: false);
-					foreach( var Module in DependencyModules.OfType<UEBuildModuleCPP>() )
-					{
-						if( !Module.bIncludedInTarget )
-						{
-							throw new BuildException( "Expecting module {0} to have bIncludeInTarget set", Module.Name );
-						}
+					var CorrectlyOrderedModules = GatherDependencyModules(NonFilteredModules);
 
-						ModulesToGenerateHeadersFor.Add( Module );
-					}
+					CorrectlyOrderedModules.RemoveWhere((Module) => !ModulesToGenerateHeadersFor.Contains(Module));
+					ModulesToGenerateHeadersFor = CorrectlyOrderedModules;
 				}
 
 				foreach( var Module in ModulesToGenerateHeadersFor )
@@ -1787,6 +1812,7 @@ namespace UnrealBuildTool
 			// If we're building a single module, then find the binary for that module and add it to our target
 			if (OnlyModules.Count > 0)
 			{
+				NonFilteredModules = AppBinaries;
 				AppBinaries = GetFilteredOnlyModules(AppBinaries, OnlyModules);
 				if (AppBinaries.Count == 0)
 				{
