Fix a problem running UnrealHeaderTool when there are stale binaries from another configuration locally. Now uses the UHT receipt to determine which binaries to check instead, and deletes them all if there's a mismatch.

Code previously searched for any UnrealHeaderTool*.dll binaries, which would pick up things like a stale UnrealHeaderTool-ScriptGenerator-Win64-Debug.dll. If that file was the first found, and chosen as the API version to compare everything else to, all the up-to-date DLLs would be deleted. Using UBT makefiles was preventing UHT being rebuilt, which would leave users in a situation where they would only be able to do one good build after generating project files.

diff --git a/Engine/Source/Programs/UnrealBuildTool/System/ExternalExecution.cs b/Engine/Source/Programs/UnrealBuildTool/System/ExternalExecution.cs
--- a/Engine/Source/Programs/UnrealBuildTool/System/ExternalExecution.cs
+++ b/Engine/Source/Programs/UnrealBuildTool/System/ExternalExecution.cs
@@ -219,131 +219,127 @@ namespace UnrealBuildTool
 			return HeaderToolPath;
 		}
 
-		class VersionedBinary
+		/// <summary>
+		/// Gets the latest write time of any of the UnrealHeaderTool binaries (including DLLs and Plugins) or DateTime.MaxValue if UnrealHeaderTool does not exist
+		/// </summary>
+		/// <returns>
+		/// Latest timestamp of UHT binaries or DateTime.MaxValue if UnrealHeaderTool is out of date and needs to be rebuilt.
+		/// </returns>
+		static bool GetHeaderToolTimestamp(out DateTime Timestamp)
 		{
-			public VersionedBinary(string InFilename, int InVersion)
+			using (var TimestampTimer = new ScopedTimer("GetHeaderToolTimestamp"))
 			{
-				Filename = InFilename;
-				Version = InVersion;
+				// Try to read the receipt for UHT.
+				string ReceiptPath = BuildReceipt.GetDefaultPath(BuildConfiguration.RelativeEnginePath, "UnrealHeaderTool", BuildHostPlatform.Current.Platform, UnrealTargetConfiguration.Development, null);
+				if(!File.Exists(ReceiptPath))
+				{
+					Timestamp = DateTime.MaxValue;
+					return false;
+				}
+
+				BuildReceipt Receipt;
+				try
+				{
+					Receipt = BuildReceipt.Read(ReceiptPath);
+				}
+				catch(Exception)
+				{
+					Timestamp = DateTime.MaxValue;
+					return false;
+				}
+				Receipt.ExpandPathVariables(BuildConfiguration.RelativeEnginePath, BuildConfiguration.RelativeEnginePath);
+
+				// Check all the binaries exist, and that all the DLLs are built against the right version
+				if(!CheckBinariesExist(Receipt) || !CheckDynamicLibaryVersionsMatch(Receipt))
+				{
+					Timestamp = DateTime.MaxValue;
+					return false;
+				}
+
+				// Return the timestamp for all the binaries
+				Timestamp = GetTimestampFromBinaries(Receipt);
+				return true;
 			}
-			public string Filename;
-			public int Version;
 		}
 
 		/// <summary>
-		/// Finds all UnrealHeaderTool plugins in the plugin directory
+		/// Checks if all the files in a receipt are present and that all the DLLs are at the same version
 		/// </summary>
-		static void RecursivelyCollectHeaderToolPlugins(string RootPath, string Pattern, string Platform, List<VersionedBinary> PluginBinaries)
+		/// <returns>
+		/// True if all the files are valid.
+		/// </returns>
+		static bool CheckBinariesExist(BuildReceipt Receipt)
 		{
-			var SubDirectories = Directory.GetDirectories(RootPath);
-			foreach (var Dir in SubDirectories)
+			bool bExist = true;
+			foreach(BuildProduct BuildProduct in Receipt.BuildProducts)
 			{
-				if (Dir.IndexOf("Intermediate", StringComparison.InvariantCultureIgnoreCase) < 0)
+				if(BuildProduct.Type == BuildProductType.Executable || BuildProduct.Type == BuildProductType.DynamicLibrary)
 				{
-					RecursivelyCollectHeaderToolPlugins(Dir, Pattern, Platform, PluginBinaries);
-					if (Dir.EndsWith("Binaries", StringComparison.InvariantCultureIgnoreCase))
+					if(!File.Exists(BuildProduct.Path))
 					{
-						// No need to search the other folders
-						break;
+						Log.TraceWarning("Missing binary: {0}", BuildProduct.Path);
+						bExist = false;
 					}
 				}
 			}
-			var Binaries = Directory.GetFiles(RootPath, Pattern);
-			foreach (var Binary in Binaries)
-			{
-				if (Binary.Contains(Platform))
-				{
-					PluginBinaries.Add(new VersionedBinary(Binary, BuildHostPlatform.Current.GetDllApiVersion(Binary)));
-				}
-			}
+			return bExist;
 		}
 
 		/// <summary>
-		/// Gets all UnrealHeaderTool binaries (including DLLs if it was not build monolithically)
+		/// Checks if all the files in a receipt have the same version
 		/// </summary>
-		static VersionedBinary[] GetHeaderToolBinaries()
+		/// <returns>
+		/// True if all the files are valid.
+		/// </returns>
+		static bool CheckDynamicLibaryVersionsMatch(BuildReceipt Receipt)
 		{
-			var Binaries = new List<VersionedBinary>();
-			var HeaderToolExe = GetHeaderToolPath();
-			if (File.Exists(HeaderToolExe))
+			List<Tuple<string, int>> BinaryVersions = new List<Tuple<string,int>>();
+			foreach(BuildProduct BuildProduct in Receipt.BuildProducts)
 			{
-				Binaries.Add(new VersionedBinary(HeaderToolExe, -1));
-
-				var HeaderToolLocation = Path.GetDirectoryName(HeaderToolExe);
-				var Platform = BuildHostPlatform.Current.Platform;
-				var DLLExtension = UEBuildPlatform.GetBuildPlatform(Platform).GetBinaryExtension(UEBuildBinaryType.DynamicLinkLibrary);
-				var DLLSearchPattern = "UnrealHeaderTool-*" + DLLExtension;
-				var HeaderToolDLLs = Directory.GetFiles(HeaderToolLocation, DLLSearchPattern, SearchOption.TopDirectoryOnly);
-
-
-				foreach (var Binary in HeaderToolDLLs)
+				if(BuildProduct.Type == BuildProductType.DynamicLibrary)
 				{
-					Binaries.Add(new VersionedBinary(Binary, BuildHostPlatform.Current.GetDllApiVersion(Binary)));
+					int Version = BuildHostPlatform.Current.GetDllApiVersion(BuildProduct.Path);
+					BinaryVersions.Add(new Tuple<string,int>(BuildProduct.Path, Version));
 				}
+			}
 
-				var PluginDirectory = Path.Combine("..", "Plugins");
-				RecursivelyCollectHeaderToolPlugins(PluginDirectory, DLLSearchPattern, Platform.ToString(), Binaries);
+			bool bMatch = true;
+			if(BinaryVersions.Count > 0 && !BinaryVersions.All(x => x.Item2 == BinaryVersions[0].Item2))
+			{
+				Log.TraceWarning("Detected mismatch in binary versions:");
+				foreach(Tuple<string, int> BinaryVersion in BinaryVersions)
+				{
+					Log.TraceWarning("  {0} has API version {1}", BinaryVersion.Item1, BinaryVersion.Item2);
+					File.Delete(BinaryVersion.Item1);
+				}
+				bMatch = false;
 			}
-			return Binaries.ToArray();
+			return bMatch;
 		}
 
 		/// <summary>
-		/// Gets the latest write time of any of the UnrealHeaderTool binaries (including DLLs and Plugins) or DateTime.MaxValue if UnrealHeaderTool does not exist
+		/// Checks if all the files in a receipt are present and that all the DLLs are at the same version
 		/// </summary>
 		/// <returns>
-		/// Latest timestamp of UHT binaries or DateTime.MaxValue if UnrealHeaderTool is out of date and needs to be rebuilt.
+		/// True if all the files are valid.
 		/// </returns>
-		static DateTime CheckIfUnrealHeaderToolIsUpToDate()
+		static DateTime GetTimestampFromBinaries(BuildReceipt Receipt)
 		{
-			var LatestWriteTime = DateTime.MinValue;
-			int? MinVersion = null;
-			using (var TimestampTimer = new ScopedTimer("GetHeaderToolTimestamp"))
+			DateTime LatestWriteTime = DateTime.MinValue;
+			foreach(BuildProduct BuildProduct in Receipt.BuildProducts)
 			{
-				var HeaderToolBinaries = GetHeaderToolBinaries();				
-				// Find the latest write time for all UnrealHeaderTool binaries
-				foreach (var Binary in HeaderToolBinaries)
-				{
-					var BinaryInfo = new FileInfo(Binary.Filename);
-					if (BinaryInfo.Exists)
-					{
-						// Latest write time
-						if (BinaryInfo.LastWriteTime > LatestWriteTime)
-						{
-							LatestWriteTime = BinaryInfo.LastWriteTime;
-						}
-						// Minimum version
-						if (Binary.Version > -1)
-						{
-							MinVersion = MinVersion.HasValue ? Math.Min(MinVersion.Value, Binary.Version) : Binary.Version;
-						}
-					}
-				}
-				if (MinVersion.HasValue)
+				if(BuildProduct.Type == BuildProductType.Executable || BuildProduct.Type == BuildProductType.DynamicLibrary)
 				{
-					// If we were able to retrieve the minimal API version, go through all binaries one more time
-					// and delete all binaries that do not match the minimum version (which for local builds would be 0, but it will
-					// also detect bad or partial syncs)
-					foreach (var Binary in HeaderToolBinaries)
+					DateTime WriteTime = File.GetLastWriteTime(BuildProduct.Path);
+					if(WriteTime > LatestWriteTime)
 					{
-						if (Binary.Version > -1)
-						{
-							if (Binary.Version != MinVersion.Value)
-							{
-								// Bad sync
-								File.Delete(Binary.Filename);
-								LatestWriteTime = DateTime.MaxValue;
-								Log.TraceWarning("Detected mismatched version in UHT binary {0} (API Version {1}, expected: {2})", Path.GetFileName(Binary.Filename), Binary.Version, MinVersion.Value);
-							}
-						}
+						LatestWriteTime = WriteTime;
 					}
 				}
 			}
-			// If UHT doesn't exist or is out of date/mismatched, force regenerate.
-			return LatestWriteTime > DateTime.MinValue ? LatestWriteTime : DateTime.MaxValue;
+			return LatestWriteTime;
 		}
 
-
-
 		/// <summary>
 		/// Gets the timestamp of CoreUObject.generated.cpp file.
 		/// </summary>
@@ -379,11 +375,8 @@ namespace UnrealBuildTool
 		 * 
 		 * @return					True if the code files are out of date
 		 * */
-		private static bool AreGeneratedCodeFilesOutOfDate(List<UHTModuleInfo> UObjectModules)
+		private static bool AreGeneratedCodeFilesOutOfDate(List<UHTModuleInfo> UObjectModules, DateTime HeaderToolTimestamp)
 		{
-			// Get UnrealHeaderTool timestamp. If it's newer than generated headers, they need to be rebuilt too.
-			var HeaderToolTimestamp = CheckIfUnrealHeaderToolIsUpToDate();
-
 			// Get CoreUObject.generated.cpp timestamp.  If the source files are older than the CoreUObject generated code, we'll
 			// need to regenerate code for the module
 			DateTime? CoreGeneratedTimestamp = null;
@@ -572,9 +565,12 @@ namespace UnrealBuildTool
 				var ToolChain = UEToolChain.GetPlatformToolChain(CppPlatform);
 				var RootLocalPath  = Path.GetFullPath(ProjectFileGenerator.RootRelativePath);
 
+				// check if UHT is out of date
+				DateTime HeaderToolTimestamp = DateTime.MaxValue;
+				bool bHaveHeaderTool = !bIsBuildingUHT && GetHeaderToolTimestamp(out HeaderToolTimestamp);
 
 				// ensure the headers are up to date
-				bool bUHTNeedsToRun = (UEBuildConfiguration.bForceHeaderGeneration == true || AreGeneratedCodeFilesOutOfDate(UObjectModules));
+				bool bUHTNeedsToRun = (UEBuildConfiguration.bForceHeaderGeneration == true || !bHaveHeaderTool || AreGeneratedCodeFilesOutOfDate(UObjectModules, HeaderToolTimestamp));
 				if( bUHTNeedsToRun || UnrealBuildTool.IsGatheringBuild )
 				{
 					// Since code files are definitely out of date, we'll now finish computing information about the UObject modules for UHT.  We
@@ -610,7 +606,7 @@ namespace UnrealBuildTool
 					if (UnrealBuildTool.RunningRocket() == false && 
 						UEBuildConfiguration.bDoNotBuildUHT == false &&
 						UEBuildConfiguration.bHotReloadFromIDE == false &&
-						!( !UnrealBuildTool.IsGatheringBuild && UnrealBuildTool.IsAssemblingBuild ) )	// If running in "assembler only" mode, we assume UHT is already up to date for much faster iteration!
+						!( bHaveHeaderTool && !UnrealBuildTool.IsGatheringBuild && UnrealBuildTool.IsAssemblingBuild ) )	// If running in "assembler only" mode, we assume UHT is already up to date for much faster iteration!
 					{
 						// If it is out of date or not there it will be built.
 						// If it is there and up to date, it will add 0.8 seconds to the build time.
