Fixed paths to Mac apps in Xcode schemes

diff --git a/Engine/Source/Programs/UnrealBuildTool/System/XcodeProjectFileGenerator.cs b/Engine/Source/Programs/UnrealBuildTool/System/XcodeProjectFileGenerator.cs
--- a/Engine/Source/Programs/UnrealBuildTool/System/XcodeProjectFileGenerator.cs
+++ b/Engine/Source/Programs/UnrealBuildTool/System/XcodeProjectFileGenerator.cs
@@ -1227,7 +1227,7 @@ namespace UnrealBuildTool
                     }
                 }
 
-				if (ExeBaseName == Target.TargetName && ExeBaseName.EndsWith("Game") && ExeExtension == ".app")
+				if (ExeBaseName == Target.TargetName && !Target.TargetFilePath.StartsWith(TargetBinariesFolder) && !ExeBaseName.EndsWith("Editor") && ExeExtension == ".app")
 				{
 					TargetBinariesFolder = TargetBinariesFolder.Replace("/Engine", "/") + ExeBaseName + "/Binaries/";
 				}
@@ -1417,18 +1417,6 @@ namespace UnrealBuildTool
 						WriteSchemeFile(XcodeProjectPath, Target, ".app", ProjectBinariesDir + Target.TargetName, new List<string>(new string[] { "&quot;" + UnrealBuildTool.GetUProjectFile() + "&quot;" }));
 					}
 				}
-				else if (Target.TargetName.EndsWith("Game"))
-				{
-					string ExeBaseName = Target.TargetName;
-					List<string> Args = null;
-					if (Target.TargetPlatform == UnrealTargetPlatform.Mac)
-					{
-						Args = new List<string>(new string[] { Target.TargetName });
-						ExeBaseName = "UE4";
-					}
-
-					WriteSchemeFile(XcodeProjectPath, Target, ".app", ExeBaseName, Args);
-				}
 				else if (Target.TargetName.EndsWith("Editor") && Target.TargetName != "UE4Editor")
 				{
 					string GameName = Target.TargetName.Replace("Editor", "");
