Fix parsing of MODULE_API_VERSION on Mac. Now takes ENGINE_IS_LICENSEE_VERSION into account, and handles tabs at the start of the line correctly.

diff --git a/Engine/Source/Programs/UnrealBuildTool/Mac/MacToolChain.cs b/Engine/Source/Programs/UnrealBuildTool/Mac/MacToolChain.cs
--- a/Engine/Source/Programs/UnrealBuildTool/Mac/MacToolChain.cs
+++ b/Engine/Source/Programs/UnrealBuildTool/Mac/MacToolChain.cs
@@ -496,6 +496,19 @@ namespace UnrealBuildTool
 			return 0;
 		}
 
+		private int LoadIsLicenseeVersionValue()
+		{
+			string[] VersionHeader = Utils.ReadAllText("../Source/Runtime/Launch/Resources/Version.h").Replace("\r\n", "\n").Replace("\t", " ").Split('\n');
+			foreach (string Line in VersionHeader)
+			{
+				if (Line.StartsWith("#define ENGINE_IS_LICENSEE_VERSION "))
+				{
+					return int.Parse(Line.Split(' ')[2]);
+				}
+			}
+			return 0;
+		}
+
 		private string LoadEngineAPIVersion()
 		{
 			int CL = 0;
@@ -503,20 +516,25 @@ namespace UnrealBuildTool
 			int BuiltFromChangelist = LoadBuiltFromChangelistValue();
 			if (BuiltFromChangelist > 0)
 			{
+				bool bFirstApiVersionDefine = true;
 				foreach (string Line in File.ReadAllLines("../Source/Runtime/Core/Public/Modules/ModuleVersion.h"))
 				{
-					string[] Tokens = Line.Split(' ', '\t');
-					if (Tokens[0] == "#define" && Tokens[1] == "MODULE_API_VERSION")
+					string[] Tokens = Line.Split(new char[]{ ' ', '\t' }, StringSplitOptions.RemoveEmptyEntries);
+					if (Tokens.Length >= 3 && Tokens[0] == "#define" && Tokens[1] == "MODULE_API_VERSION")
 					{
-						if(Tokens[2] == "BUILT_FROM_CHANGELIST")
-						{
-							CL = LoadEngineCL();
-						}
-						else
+						if (!bFirstApiVersionDefine || LoadIsLicenseeVersionValue() != 0)
 						{
-							CL = int.Parse(Tokens[2]);
+							if (Tokens[2] == "BUILT_FROM_CHANGELIST")
+							{
+								CL = LoadEngineCL();
+							}
+							else
+							{
+								CL = int.Parse(Tokens[2]);
+							}
+							break;
 						}
-						break;
+						bFirstApiVersionDefine = false;
 					}
 				}
 			}
