AVX support in UnrealBuildTool.

diff --git a/Engine/Source/Programs/UnrealBuildTool/Configuration/UEBuildModule.cs b/Engine/Source/Programs/UnrealBuildTool/Configuration/UEBuildModule.cs
--- a/Engine/Source/Programs/UnrealBuildTool/Configuration/UEBuildModule.cs
+++ b/Engine/Source/Programs/UnrealBuildTool/Configuration/UEBuildModule.cs
@@ -1071,6 +1071,9 @@ namespace UnrealBuildTool
 		/** Use run time type information */
 		public bool bUseRTTI = false;
 
+		/** Use AVX instructions */
+		public bool bUseAVX = false;
+
 		/** Enable buffer security checks. This should usually be enabled as it prevents severe security risks. */
 		public bool bEnableBufferSecurityChecks = true;
 
@@ -1176,7 +1179,8 @@ namespace UnrealBuildTool
 			bool InEnableExceptions,
 			bool InEnableShadowVariableWarnings,
 			bool bInBuildSourceFiles,
-			string InBuildCsFilename
+			string InBuildCsFilename,
+			bool bInUseAVX
 			)
 			: base(	InTarget,
 					InName, 
@@ -1224,6 +1228,7 @@ namespace UnrealBuildTool
 			bAllowSharedPCH                        = InAllowSharedPCH;
 			SharedPCHHeaderFile                    = InSharedPCHHeaderFile;
 			bUseRTTI                               = InUseRTTI;
+			bUseAVX                                = bInUseAVX;
 			bEnableBufferSecurityChecks 		   = InEnableBufferSecurityChecks;
 			bFasterWithoutUnity                    = InFasterWithoutUnity;
 			MinFilesUsingPrecompiledHeaderOverride = InMinFilesUsingPrecompiledHeaderOverride;
@@ -1883,6 +1888,7 @@ namespace UnrealBuildTool
 			Result.Config.bFasterWithoutUnity                    = bFasterWithoutUnity;
 			Result.Config.OptimizeCode                           = OptimizeCode;
 			Result.Config.bUseRTTI                               = bUseRTTI;
+			Result.Config.bUseAVX                                = bUseAVX;
 			Result.Config.bEnableBufferSecurityChecks            = bEnableBufferSecurityChecks;
 			Result.Config.bFasterWithoutUnity                    = bFasterWithoutUnity;
 			Result.Config.MinFilesUsingPrecompiledHeaderOverride = MinFilesUsingPrecompiledHeaderOverride;
@@ -2138,7 +2144,8 @@ namespace UnrealBuildTool
 			bool InEnableExceptions,
 			bool InEnableShadowVariableWarnings,
 			bool bInBuildSourceFiles,
-			string InBuildCsFilename
+			string InBuildCsFilename,
+			bool bInUseAVX
 			)
 			: base(InTarget,InName,InType,InModuleDirectory,InGeneratedCodeDirectory,InIsRedistributableOverride,InIntelliSenseGatherer,
 			InSourceFiles,InPublicIncludePaths,InPublicSystemIncludePaths,null,InDefinitions,
@@ -2146,7 +2153,7 @@ namespace UnrealBuildTool
 			InPrivateIncludePaths,InPrivateIncludePathModuleNames,InPrivateDependencyModuleNames,
             InCircularlyReferencedDependentModules, InDynamicallyLoadedModuleNames, InPlatformSpecificDynamicallyLoadedModuleNames, InRuntimeDependencies, InOptimizeCode,
 			InAllowSharedPCH, InSharedPCHHeaderFile, InUseRTTI, InEnableBufferSecurityChecks, InFasterWithoutUnity, InMinFilesUsingPrecompiledHeaderOverride,
-			InEnableExceptions, InEnableShadowVariableWarnings, bInBuildSourceFiles, InBuildCsFilename)
+			InEnableExceptions, InEnableShadowVariableWarnings, bInBuildSourceFiles, InBuildCsFilename, bInUseAVX)
 		{
 			PrivateAssemblyReferences = HashSetFromOptionalEnumerableStringParameter(InPrivateAssemblyReferences);
 		}
diff --git a/Engine/Source/Programs/UnrealBuildTool/Configuration/UEBuildTarget.cs b/Engine/Source/Programs/UnrealBuildTool/Configuration/UEBuildTarget.cs
--- a/Engine/Source/Programs/UnrealBuildTool/Configuration/UEBuildTarget.cs
+++ b/Engine/Source/Programs/UnrealBuildTool/Configuration/UEBuildTarget.cs
@@ -2171,7 +2171,8 @@ namespace UnrealBuildTool
 				InEnableExceptions: false,
 				InEnableShadowVariableWarnings: true,
 				bInBuildSourceFiles: true,
-				InBuildCsFilename: null);
+				InBuildCsFilename: null,
+				bInUseAVX: false);
 		}
 
 
@@ -3213,7 +3214,8 @@ namespace UnrealBuildTool
 							InEnableExceptions: RulesObject.bEnableExceptions,
 							InEnableShadowVariableWarnings: RulesObject.bEnableShadowVariableWarnings,
 							bInBuildSourceFiles: bBuildSourceFiles,
-							InBuildCsFilename: InBuildCsFile
+							InBuildCsFilename: InBuildCsFile,
+							bInUseAVX: RulesObject.bUseAVX
 						);
 
 				case ModuleRules.ModuleType.CPlusPlusCLR:
@@ -3256,7 +3258,8 @@ namespace UnrealBuildTool
 							InEnableExceptions: RulesObject.bEnableExceptions,
 							InEnableShadowVariableWarnings: RulesObject.bEnableShadowVariableWarnings,
 							bInBuildSourceFiles : bBuildSourceFiles,
-							InBuildCsFilename: InBuildCsFile
+							InBuildCsFilename: InBuildCsFile,
+							bInUseAVX: RulesObject.bUseAVX
 						);
 
 				case ModuleRules.ModuleType.External:
diff --git a/Engine/Source/Programs/UnrealBuildTool/System/CPPEnvironment.cs b/Engine/Source/Programs/UnrealBuildTool/System/CPPEnvironment.cs
--- a/Engine/Source/Programs/UnrealBuildTool/System/CPPEnvironment.cs
+++ b/Engine/Source/Programs/UnrealBuildTool/System/CPPEnvironment.cs
@@ -200,6 +200,9 @@ namespace UnrealBuildTool
 		/** Use run time type information */
 		public bool bUseRTTI = false;
 
+		/** Use AVX instructions */
+		public bool bUseAVX = false;
+
 		/** Enable buffer security checks.   This should usually be enabled as it prevents severe security risks. */
 		public bool bEnableBufferSecurityChecks = true;
 
@@ -270,6 +273,7 @@ namespace UnrealBuildTool
 			PrecompiledHeaderAction                = InCopyEnvironment.PrecompiledHeaderAction;
 			bForceIncludePrecompiledHeader         = InCopyEnvironment.bForceIncludePrecompiledHeader;
 			bUseRTTI                               = InCopyEnvironment.bUseRTTI;
+			bUseAVX                                = InCopyEnvironment.bUseAVX;
 			bFasterWithoutUnity                    = InCopyEnvironment.bFasterWithoutUnity;
 			MinFilesUsingPrecompiledHeaderOverride = InCopyEnvironment.MinFilesUsingPrecompiledHeaderOverride;
 			bEnableExceptions                      = InCopyEnvironment.bEnableExceptions;
diff --git a/Engine/Source/Programs/UnrealBuildTool/System/RulesCompiler.cs b/Engine/Source/Programs/UnrealBuildTool/System/RulesCompiler.cs
--- a/Engine/Source/Programs/UnrealBuildTool/System/RulesCompiler.cs
+++ b/Engine/Source/Programs/UnrealBuildTool/System/RulesCompiler.cs
@@ -164,6 +164,9 @@ namespace UnrealBuildTool
 		/** Use run time type information */
 		public bool bUseRTTI = false;
 
+		/** Use AVX instructions */
+		public bool bUseAVX = false;
+
 		/** Enable buffer security checks.  This should usually be enabled as it prevents severe security risks. */
 		public bool bEnableBufferSecurityChecks = true;
 
diff --git a/Engine/Source/Programs/UnrealBuildTool/Windows/VCToolChain.cs b/Engine/Source/Programs/UnrealBuildTool/Windows/VCToolChain.cs
--- a/Engine/Source/Programs/UnrealBuildTool/Windows/VCToolChain.cs
+++ b/Engine/Source/Programs/UnrealBuildTool/Windows/VCToolChain.cs
@@ -372,13 +372,18 @@ namespace UnrealBuildTool
 			if ((CompileEnvironment.Config.Target.Platform == CPPTargetPlatform.Win32) || 
 				(CompileEnvironment.Config.Target.Platform == CPPTargetPlatform.Win64))
 			{
-				// SSE options are not allowed when using CLR compilation or the 64 bit toolchain
-				// (both enable SSE2 automatically)
-				if (CompileEnvironment.Config.CLRMode == CPPCLRMode.CLRDisabled &&
-					CompileEnvironment.Config.Target.Platform != CPPTargetPlatform.Win64)
+				// SSE options are not allowed when using CLR compilation
+				if (CompileEnvironment.Config.CLRMode == CPPCLRMode.CLRDisabled && WindowsPlatform.bUseVCCompilerArgs)
 				{
-					if( WindowsPlatform.bUseVCCompilerArgs )
-					{ 
+					if (CompileEnvironment.Config.bUseAVX)
+					{
+						// Allow the compiler to generate AVX instructions.
+						Arguments.Append(" /arch:AVX");
+					}
+					// SSE options are not allowed when using the 64 bit toolchain
+					// (enables SSE2 automatically)
+					else if (CompileEnvironment.Config.Target.Platform != CPPTargetPlatform.Win64)
+					{
 						// Allow the compiler to generate SSE2 instructions.
 						Arguments.Append(" /arch:SSE2");
 					}
