fix for extra output when building in VS
#ubt

diff --git a/Engine/Source/Programs/UnrealBuildTool/System/UnrealBuildTool.cs b/Engine/Source/Programs/UnrealBuildTool/System/UnrealBuildTool.cs
--- a/Engine/Source/Programs/UnrealBuildTool/System/UnrealBuildTool.cs
+++ b/Engine/Source/Programs/UnrealBuildTool/System/UnrealBuildTool.cs
@@ -604,7 +604,7 @@ namespace UnrealBuildTool
                 bLogSeverity: false,
                 bLogSources: false,
                 bColorConsoleOutput: true,
-                TraceListeners: new[] { new ConsoleTraceListener() });
+                TraceListeners: new[] { new FilteredConsoleTraceListener() });
         }
 
         /// <summary>
