Fix for UELinkerFixups not being built properly after module updates.

diff --git a/Engine/Source/Programs/UnrealBuildTool/Configuration/UEBuildTarget.cs b/Engine/Source/Programs/UnrealBuildTool/Configuration/UEBuildTarget.cs
--- a/Engine/Source/Programs/UnrealBuildTool/Configuration/UEBuildTarget.cs
+++ b/Engine/Source/Programs/UnrealBuildTool/Configuration/UEBuildTarget.cs
@@ -2022,7 +2022,8 @@ namespace UnrealBuildTool
 
 				// Create the source file list (just the one cpp file)
 				List<FileItem> SourceFiles = new List<FileItem>();
-				SourceFiles.Add(FileItem.GetItemByPath(LinkerFixupCPPFilename));
+				var LinkerFixupCPPFileItem = FileItem.GetItemByPath(LinkerFixupCPPFilename);
+				SourceFiles.Add(LinkerFixupCPPFileItem);
 
 				// Create the CPP module
 				var FakeModuleDirectory = Path.GetDirectoryName( LinkerFixupCPPFilename );
@@ -2052,6 +2053,10 @@ namespace UnrealBuildTool
 				{
 					ResponseFile.Create(LinkerFixupHeaderFilenameWithPath, new List<string>());
 					ResponseFile.Create(LinkerFixupCPPFilename, LinkerFixupsFileContents);
+
+					// Update the cached file states so that the linker fixups definitely get rebuilt
+					FileItem.GetItemByPath(LinkerFixupHeaderFilenameWithPath).ResetFileInfo();
+					LinkerFixupCPPFileItem.ResetFileInfo();
 				}
 			}
 		}
