Re-enable generating .dSYMs on Mac by default for [UEPLAT-308] in development & shipping builds now that the .dSYMs are flat and shouldn't case MAX_PATH errors.

diff --git a/Engine/Source/Programs/UnrealBuildTool/Mac/MacToolChain.cs b/Engine/Source/Programs/UnrealBuildTool/Mac/MacToolChain.cs
--- a/Engine/Source/Programs/UnrealBuildTool/Mac/MacToolChain.cs
+++ b/Engine/Source/Programs/UnrealBuildTool/Mac/MacToolChain.cs
@@ -1298,6 +1298,28 @@ namespace UnrealBuildTool
 
         public override void AddFilesToReceipt(BuildReceipt Receipt, UEBuildBinary Binary)
 		{
+			// The cross-platform code adds .dSYMs for static libraries, which is just wrong, so
+			// eliminate them here for now.
+			string DebugExtension = UEBuildPlatform.GetBuildPlatform(Binary.Target.Platform).GetDebugInfoExtension(Binary.Config.Type);
+			if(DebugExtension == ".dsym")
+			{
+				for (int i = 0; i < Receipt.BuildProducts.Count; i++)
+				{
+					if(Path.GetExtension(Receipt.BuildProducts[i].Path) == DebugExtension)
+					{
+						Receipt.BuildProducts.RemoveAt(i--);
+					}
+				}
+
+				for (int i = 0; i < Receipt.BuildProducts.Count; i++)
+				{
+					if(Receipt.BuildProducts[i].Type == BuildProductType.Executable || Receipt.BuildProducts[i].Type == BuildProductType.DynamicLibrary)
+					{
+						Receipt.AddBuildProduct(Path.ChangeExtension(Receipt.BuildProducts[i].Path, DebugExtension), BuildProductType.SymbolFile);
+					}
+				}
+			}
+
 			if (Binary.Target.GlobalLinkEnvironment.Config.bIsBuildingConsoleApplication)
 			{
 				return;
diff --git a/Engine/Source/Programs/UnrealBuildTool/Mac/UEBuildMac.cs b/Engine/Source/Programs/UnrealBuildTool/Mac/UEBuildMac.cs
--- a/Engine/Source/Programs/UnrealBuildTool/Mac/UEBuildMac.cs
+++ b/Engine/Source/Programs/UnrealBuildTool/Mac/UEBuildMac.cs
@@ -173,8 +173,7 @@ namespace UnrealBuildTool
 			BuildConfiguration.ProcessorCountMultiplier = MacToolChain.GetAdjustedProcessorCountMultiplier();
 			BuildConfiguration.bUseSharedPCHs = false;
 
-			// Disabled as we hit the Windows build machine's MAX_PATH limit for various projects
-			// BuildConfiguration.bUsePDBFiles = bCreateDebugInfo && Configuration != CPPTargetConfiguration.Debug && Platform == CPPTargetPlatform.Mac;
+			BuildConfiguration.bUsePDBFiles = bCreateDebugInfo && Configuration != CPPTargetConfiguration.Debug && Platform == CPPTargetPlatform.Mac;
 
 			// we always deploy - the build machines need to be able to copy the files back, which needs the full bundle
 			BuildConfiguration.bDeployAfterCompile = true;
