fix for build machines not always copying the latest mobile provision from Build/IOS
#ios

diff --git a/Engine/Source/Programs/UnrealBuildTool/IOS/UEDeployIOS.cs b/Engine/Source/Programs/UnrealBuildTool/IOS/UEDeployIOS.cs
--- a/Engine/Source/Programs/UnrealBuildTool/IOS/UEDeployIOS.cs
+++ b/Engine/Source/Programs/UnrealBuildTool/IOS/UEDeployIOS.cs
@@ -543,7 +543,7 @@ namespace UnrealBuildTool.IOS
 				DestFileInfo = new FileInfo (Environment.GetEnvironmentVariable ("HOME") + "/Library/MobileDevice/Provisioning Profiles/" + InProjectName + ".mobileprovision");
 				DestFileInfo.Attributes = DestFileInfo.Attributes & ~FileAttributes.ReadOnly;
 			}
-			else
+			if (!File.Exists(ProvisionWithPrefix) || Environment.GetEnvironmentVariable("IsBuildMachine") == "1")
 			{
 				// copy all provisions from the game directory, the engine directory, and the notforlicensees directory
 				// copy all of the provisions from the game directory to the library
