SN-DBS support for distributed compilation of Win64 builds
* Modules that use #import (DatabaseSupport & VisualStudioSourceCodeAccessor) are forced to be built locally due to SN-DBS limitation . These are very small modules so have a minimal impact on build time.

diff --git a/Engine/Plugins/Developer/VisualStudioSourceCodeAccess/Source/VisualStudioSourceCodeAccess/VisualStudioSourceCodeAccess.Build.cs b/Engine/Plugins/Developer/VisualStudioSourceCodeAccess/Source/VisualStudioSourceCodeAccess/VisualStudioSourceCodeAccess.Build.cs
--- a/Engine/Plugins/Developer/VisualStudioSourceCodeAccess/Source/VisualStudioSourceCodeAccess/VisualStudioSourceCodeAccess.Build.cs
+++ b/Engine/Plugins/Developer/VisualStudioSourceCodeAccess/Source/VisualStudioSourceCodeAccess/VisualStudioSourceCodeAccess.Build.cs
@@ -31,6 +31,8 @@ namespace UnrealBuildTool.Rules
 			{
 				Definitions.Add("VSACCESSOR_HAS_DTE=0");
 			}
+
+			bBuildLocallyWithSNDBS = true;
 		}
 	}
 }
\ No newline at end of file
diff --git a/Engine/Source/Programs/UnrealBuildTool/Configuration/UEBuildModule.cs b/Engine/Source/Programs/UnrealBuildTool/Configuration/UEBuildModule.cs
--- a/Engine/Source/Programs/UnrealBuildTool/Configuration/UEBuildModule.cs
+++ b/Engine/Source/Programs/UnrealBuildTool/Configuration/UEBuildModule.cs
@@ -1083,6 +1083,9 @@ namespace UnrealBuildTool
 		/** Overrides BuildConfiguration.MinFilesUsingPrecompiledHeader if non-zero. */
 		public int MinFilesUsingPrecompiledHeaderOverride = 0;
 
+		/**  Module uses a #import so must be built locally when compiling with SN-DBS */
+		public bool bBuildLocallyWithSNDBS = false;
+		
 		/** Enable exception handling */
 		public bool bEnableExceptions = false;
 
@@ -1176,6 +1179,7 @@ namespace UnrealBuildTool
 			bool InEnableBufferSecurityChecks,
 			bool InFasterWithoutUnity,
 			int InMinFilesUsingPrecompiledHeaderOverride,
+			bool InBuildLocallyWithSNDBS,
 			bool InEnableExceptions,
 			bool InEnableShadowVariableWarnings,
 			bool bInBuildSourceFiles,
@@ -1232,6 +1236,7 @@ namespace UnrealBuildTool
 			bEnableBufferSecurityChecks 		   = InEnableBufferSecurityChecks;
 			bFasterWithoutUnity                    = InFasterWithoutUnity;
 			MinFilesUsingPrecompiledHeaderOverride = InMinFilesUsingPrecompiledHeaderOverride;
+			bBuildLocallyWithSNDBS				   = InBuildLocallyWithSNDBS;
 			bEnableExceptions                      = InEnableExceptions;
 			bEnableShadowVariableWarnings          = InEnableShadowVariableWarnings;
 		}
@@ -1892,6 +1897,7 @@ namespace UnrealBuildTool
 			Result.Config.bEnableBufferSecurityChecks            = bEnableBufferSecurityChecks;
 			Result.Config.bFasterWithoutUnity                    = bFasterWithoutUnity;
 			Result.Config.MinFilesUsingPrecompiledHeaderOverride = MinFilesUsingPrecompiledHeaderOverride;
+			Result.Config.bBuildLocallyWithSNDBS				 = bBuildLocallyWithSNDBS;
 			Result.Config.bEnableExceptions                      = bEnableExceptions;
 			Result.Config.bEnableShadowVariableWarning          = bEnableShadowVariableWarnings;
 			Result.Config.bUseStaticCRT							 = (Target.Rules != null && Target.Rules.bUseStaticCRT);
@@ -2152,7 +2158,7 @@ namespace UnrealBuildTool
 			InPublicIncludePathModuleNames,InPublicDependencyModuleNames,InPublicDelayLoadDLLs,InPublicAdditionalLibraries,InPublicFrameworks,InPublicWeakFrameworks,InPublicAdditionalFrameworks,InPublicAdditionalShadowFiles,InPublicAdditionalBundleResources,
 			InPrivateIncludePaths,InPrivateIncludePathModuleNames,InPrivateDependencyModuleNames,
             InCircularlyReferencedDependentModules, InDynamicallyLoadedModuleNames, InPlatformSpecificDynamicallyLoadedModuleNames, InRuntimeDependencies, InOptimizeCode,
-			InAllowSharedPCH, InSharedPCHHeaderFile, InUseRTTI, InEnableBufferSecurityChecks, InFasterWithoutUnity, InMinFilesUsingPrecompiledHeaderOverride,
+			InAllowSharedPCH, InSharedPCHHeaderFile, InUseRTTI, InEnableBufferSecurityChecks, InFasterWithoutUnity, InMinFilesUsingPrecompiledHeaderOverride, true,
 			InEnableExceptions, InEnableShadowVariableWarnings, bInBuildSourceFiles, InBuildCsFilename, bInUseAVX)
 		{
 			PrivateAssemblyReferences = HashSetFromOptionalEnumerableStringParameter(InPrivateAssemblyReferences);
diff --git a/Engine/Source/Programs/UnrealBuildTool/Configuration/UEBuildTarget.cs b/Engine/Source/Programs/UnrealBuildTool/Configuration/UEBuildTarget.cs
--- a/Engine/Source/Programs/UnrealBuildTool/Configuration/UEBuildTarget.cs
+++ b/Engine/Source/Programs/UnrealBuildTool/Configuration/UEBuildTarget.cs
@@ -2173,6 +2173,7 @@ namespace UnrealBuildTool
 				InEnableBufferSecurityChecks: true,
 				InFasterWithoutUnity: true,
 				InMinFilesUsingPrecompiledHeaderOverride: 0,
+				InBuildLocallyWithSNDBS: false,
 				InEnableExceptions: false,
 				InEnableShadowVariableWarnings: true,
 				bInBuildSourceFiles: true,
@@ -3216,6 +3217,7 @@ namespace UnrealBuildTool
 							InEnableBufferSecurityChecks: RulesObject.bEnableBufferSecurityChecks,
 							InFasterWithoutUnity: RulesObject.bFasterWithoutUnity,
 							InMinFilesUsingPrecompiledHeaderOverride: RulesObject.MinFilesUsingPrecompiledHeaderOverride,
+							InBuildLocallyWithSNDBS: RulesObject.bBuildLocallyWithSNDBS,
 							InEnableExceptions: RulesObject.bEnableExceptions,
 							InEnableShadowVariableWarnings: RulesObject.bEnableShadowVariableWarnings,
 							bInBuildSourceFiles: bBuildSourceFiles,
diff --git a/Engine/Source/Programs/UnrealBuildTool/System/CPPEnvironment.cs b/Engine/Source/Programs/UnrealBuildTool/System/CPPEnvironment.cs
--- a/Engine/Source/Programs/UnrealBuildTool/System/CPPEnvironment.cs
+++ b/Engine/Source/Programs/UnrealBuildTool/System/CPPEnvironment.cs
@@ -212,6 +212,9 @@ namespace UnrealBuildTool
         /** Overrides BuildConfiguration.MinFilesUsingPrecompiledHeader if non-zero. */
         public int MinFilesUsingPrecompiledHeaderOverride = 0;
 
+		/**  Module uses a #import so must be built locally when compiling with SN-DBS */
+		public bool bBuildLocallyWithSNDBS = false;
+
 		/** Enable exception handling */
 		public bool bEnableExceptions = false;
 
@@ -276,6 +279,7 @@ namespace UnrealBuildTool
 			bUseAVX                                = InCopyEnvironment.bUseAVX;
 			bFasterWithoutUnity                    = InCopyEnvironment.bFasterWithoutUnity;
 			MinFilesUsingPrecompiledHeaderOverride = InCopyEnvironment.MinFilesUsingPrecompiledHeaderOverride;
+			bBuildLocallyWithSNDBS				   = InCopyEnvironment.bBuildLocallyWithSNDBS;
 			bEnableExceptions                      = InCopyEnvironment.bEnableExceptions;
 			bEnableShadowVariableWarning          = InCopyEnvironment.bEnableShadowVariableWarning;
 			OptimizeCode                           = InCopyEnvironment.OptimizeCode;
diff --git a/Engine/Source/Programs/UnrealBuildTool/System/RulesCompiler.cs b/Engine/Source/Programs/UnrealBuildTool/System/RulesCompiler.cs
--- a/Engine/Source/Programs/UnrealBuildTool/System/RulesCompiler.cs
+++ b/Engine/Source/Programs/UnrealBuildTool/System/RulesCompiler.cs
@@ -182,6 +182,9 @@ namespace UnrealBuildTool
 		/** Overrides BuildConfiguration.MinFilesUsingPrecompiledHeader if non-zero. */
 		public int MinFilesUsingPrecompiledHeaderOverride = 0;
 
+		/**  Module uses a #import so must be built locally when compiling with SN-DBS */
+		public bool bBuildLocallyWithSNDBS = false;
+
 		/// List of modules with header files that our module's public headers needs access to, but we don't need to "import" or link against.
 		public List<string> PublicIncludePathModuleNames = new List<string>();
 
diff --git a/Engine/Source/Programs/UnrealBuildTool/System/SNDBS.cs b/Engine/Source/Programs/UnrealBuildTool/System/SNDBS.cs
--- a/Engine/Source/Programs/UnrealBuildTool/System/SNDBS.cs
+++ b/Engine/Source/Programs/UnrealBuildTool/System/SNDBS.cs
@@ -234,7 +234,7 @@ namespace UnrealBuildTool
             if( NumScriptedActions > 0 )
             {
                 // Create the process
-                ProcessStartInfo PSI = new ProcessStartInfo("dbsbuild", "-p UE4 -s " + BuildConfiguration.BaseIntermediatePath + "/sndbs.bat");
+                ProcessStartInfo PSI = new ProcessStartInfo("dbsbuild", "-q -p UE4 -s " + BuildConfiguration.BaseIntermediatePath + "/sndbs.bat");
                 PSI.RedirectStandardOutput = true;
                 PSI.RedirectStandardError = true;
                 PSI.UseShellExecute = false;
@@ -290,14 +290,14 @@ namespace UnrealBuildTool
             {
                 string SCERoot = Environment.GetEnvironmentVariable("SCE_ROOT_DIR");
 
-		bool bSNDBSExists = false;
-		if( SCERoot != null )
-		{
-			string SNDBSExecutable = Path.Combine(SCERoot, "Common/SN-DBS/bin/dbsbuild.exe");
+				bool bSNDBSExists = false;
+				if( SCERoot != null )
+				{
+					string SNDBSExecutable = Path.Combine(SCERoot, "Common/SN-DBS/bin/dbsbuild.exe");
 
-			// Check that SN-DBS is available
-			bSNDBSExists = File.Exists(SNDBSExecutable);
-		}
+					// Check that SN-DBS is available
+					bSNDBSExists = File.Exists(SNDBSExecutable);
+				}
 
                 if( bSNDBSExists == false )
                 {
diff --git a/Engine/Source/Programs/UnrealBuildTool/Windows/UEBuildWindows.cs b/Engine/Source/Programs/UnrealBuildTool/Windows/UEBuildWindows.cs
--- a/Engine/Source/Programs/UnrealBuildTool/Windows/UEBuildWindows.cs
+++ b/Engine/Source/Programs/UnrealBuildTool/Windows/UEBuildWindows.cs
@@ -218,6 +218,17 @@ namespace UnrealBuildTool
             return new DirectoryInfo(Path.Combine(VSPath, "..", "Tools")).FullName;
         }
 
+		/**
+         *	If this platform can be compiled with SN-DBS
+         */
+		public override bool CanUseSNDBS()
+		{
+			// Check that SN-DBS is available
+			string SCERootPath = Environment.GetEnvironmentVariable("SCE_ROOT_DIR");
+			string SNDBSPath = Path.Combine(SCERootPath, "common", "sn-dbs", "bin", "dbsbuild.exe");
+			bool bIsSNDBSAvailable = File.Exists(SNDBSPath);
+			return bIsSNDBSAvailable;
+		}
 
         /**
          *	Register the platform with the UEBuildPlatform class
diff --git a/Engine/Source/Programs/UnrealBuildTool/Windows/VCToolChain.cs b/Engine/Source/Programs/UnrealBuildTool/Windows/VCToolChain.cs
--- a/Engine/Source/Programs/UnrealBuildTool/Windows/VCToolChain.cs
+++ b/Engine/Source/Programs/UnrealBuildTool/Windows/VCToolChain.cs
@@ -1350,6 +1350,12 @@ namespace UnrealBuildTool
 				{
 					CompileAction.bCanExecuteRemotely = false;
 				}
+
+				// When compiling with SN-DBS, modules that contain a #import must be built locally
+				if( CompileEnvironment.Config.bBuildLocallyWithSNDBS == true && BuildConfiguration.bAllowSNDBS)
+				{
+					CompileAction.bCanExecuteRemotely = false;
+				}
 			}
 			return Result;
 		}
diff --git a/Engine/Source/Runtime/DatabaseSupport/DatabaseSupport.Build.cs b/Engine/Source/Runtime/DatabaseSupport/DatabaseSupport.Build.cs
--- a/Engine/Source/Runtime/DatabaseSupport/DatabaseSupport.Build.cs
+++ b/Engine/Source/Runtime/DatabaseSupport/DatabaseSupport.Build.cs
@@ -9,6 +9,7 @@ namespace UnrealBuildTool.Rules
 			if ((UEBuildConfiguration.bIncludeADO == true) || (UEBuildConfiguration.bCompileAgainstEngine == true))
 			{
 				AddThirdPartyPrivateStaticDependencies(Target, "ADO");
+				bBuildLocallyWithSNDBS = true;
 			}
 			
 			
