Fix failure when SCE env is missing.

diff --git a/Engine/Source/Programs/UnrealBuildTool/Windows/UEBuildWindows.cs b/Engine/Source/Programs/UnrealBuildTool/Windows/UEBuildWindows.cs
--- a/Engine/Source/Programs/UnrealBuildTool/Windows/UEBuildWindows.cs
+++ b/Engine/Source/Programs/UnrealBuildTool/Windows/UEBuildWindows.cs
@@ -225,9 +225,16 @@ namespace UnrealBuildTool
 		{
 			// Check that SN-DBS is available
 			string SCERootPath = Environment.GetEnvironmentVariable("SCE_ROOT_DIR");
-			string SNDBSPath = Path.Combine(SCERootPath, "common", "sn-dbs", "bin", "dbsbuild.exe");
-			bool bIsSNDBSAvailable = File.Exists(SNDBSPath);
-			return bIsSNDBSAvailable;
+			if (!String.IsNullOrEmpty(SCERootPath))
+			{
+				string SNDBSPath = Path.Combine(SCERootPath, "common", "sn-dbs", "bin", "dbsbuild.exe");
+				bool bIsSNDBSAvailable = File.Exists(SNDBSPath);
+				return bIsSNDBSAvailable;
+			}
+			else
+			{
+				return false;
+			}
 		}
 
         /**
