Merging CL 2541630 - enabled syntax highlighting for MetalRHI module in Xcode project

diff --git a/Engine/Source/Programs/UnrealBuildTool/System/XcodeProject.cs b/Engine/Source/Programs/UnrealBuildTool/System/XcodeProject.cs
--- a/Engine/Source/Programs/UnrealBuildTool/System/XcodeProject.cs
+++ b/Engine/Source/Programs/UnrealBuildTool/System/XcodeProject.cs
@@ -183,7 +183,8 @@ namespace UnrealBuildTool
 				foreach (string PlatformName in Enum.GetNames(typeof(UnrealTargetPlatform)))
 				{
 					string AltName = PlatformName == "Win32" || PlatformName == "Win64" ? "windows" : PlatformName.ToLower();
-					if ((SourceFile.FilePath.ToLower().Contains("/" + PlatformName.ToLower() + "/") || SourceFile.FilePath.ToLower().Contains("/" + AltName + "/")) && PlatformName != "Mac")
+					if ((SourceFile.FilePath.ToLower().Contains("/" + PlatformName.ToLower() + "/") || SourceFile.FilePath.ToLower().Contains("/" + AltName + "/"))
+						&& PlatformName != "Mac" && !SourceFile.FilePath.Contains("MetalRHI"))
 					{
 						// UE4XcodeHelper is Mac only target, so skip other platforms files
 						return false;
