changes to allow two projects to have the same plugin name
UEB-258

diff --git a/Engine/Source/Programs/UnrealBuildTool/Configuration/UEBuildTarget.cs b/Engine/Source/Programs/UnrealBuildTool/Configuration/UEBuildTarget.cs
--- a/Engine/Source/Programs/UnrealBuildTool/Configuration/UEBuildTarget.cs
+++ b/Engine/Source/Programs/UnrealBuildTool/Configuration/UEBuildTarget.cs
@@ -173,6 +173,8 @@ namespace UnrealBuildTool
 		public bool bUsePrecompiled;
 		public List<string> ForeignPlugins;
 		public string ForceReceiptFileName;
+		public string AssemblyName;
+		public List<string> GameFolders;
 	}
 
 
@@ -476,6 +478,8 @@ namespace UnrealBuildTool
 					TargetName = PossibleTargetName;
 
 					// If a project file was not specified see if we can find one
+					List<string> GameFolders = new List<string>();
+					string AssemblyName = "UE4";
 					string CheckProjectFile = UProjectInfo.GetProjectForTarget(TargetName);
 					if (string.IsNullOrEmpty(CheckProjectFile) == false)
 					{
@@ -500,6 +504,8 @@ namespace UnrealBuildTool
 							// Parse off the path
 							TargetName = Path.GetFileNameWithoutExtension( TargetName );
 						}
+						GameFolders.Add(UnrealBuildTool.GetUProjectPath());
+						AssemblyName = Path.GetFileNameWithoutExtension(UnrealBuildTool.GetUProjectFile());
 					}
 
 					Targets.Add( new TargetDescriptor()
@@ -514,7 +520,9 @@ namespace UnrealBuildTool
 							bPrecompile = bPrecompile,
 							bUsePrecompiled = bUsePrecompiled,
 							ForeignPlugins = ForeignPlugins,
-							ForceReceiptFileName = ForceReceiptFileName
+							ForceReceiptFileName = ForceReceiptFileName,
+							AssemblyName = AssemblyName,
+							GameFolders = GameFolders
 						} );
 					break;
 				}
@@ -529,19 +537,6 @@ namespace UnrealBuildTool
 		public static UEBuildTarget CreateTarget( TargetDescriptor Desc )
 		{
 			UEBuildTarget BuildTarget = null;
-			if( !ProjectFileGenerator.bGenerateProjectFiles )
-			{
-				// Configure the rules compiler
-				string PossibleAssemblyName = Desc.TargetName;
-				if (Desc.bIsEditorRecompile == true)
-				{
-					PossibleAssemblyName += "_EditorRecompile";
-				}
-
-				// Scan the disk to find source files for all known targets and modules, and generate "in memory" project
-				// file data that will be used to determine what to build
-				RulesCompiler.SetAssemblyNameAndGameFolders( PossibleAssemblyName, UEBuildTarget.DiscoverAllGameFolders(), Desc.ForeignPlugins );
-			}
 
 			// Try getting it from the RulesCompiler
 			UEBuildTarget Target = RulesCompiler.CreateTarget(Desc);
diff --git a/Engine/Source/Programs/UnrealBuildTool/System/ProjectFileGenerator.cs b/Engine/Source/Programs/UnrealBuildTool/System/ProjectFileGenerator.cs
--- a/Engine/Source/Programs/UnrealBuildTool/System/ProjectFileGenerator.cs
+++ b/Engine/Source/Programs/UnrealBuildTool/System/ProjectFileGenerator.cs
@@ -375,23 +375,22 @@ namespace UnrealBuildTool
 
 			// Build the list of games to generate projects for
 			var AllGameProjects = UProjectInfo.FilterGameProjects(true, bGeneratingGameProjectFiles ? GameProjectName : null);
-				
+
+			List<string> AssemblyGameFolders = new List<string>();
 			var AssemblyName = "ProjectFileGenerator";
-			if( bGeneratingGameProjectFiles )
+			if (bGeneratingRocketProjectFiles)
 			{
-				AssemblyName = GameProjectName + "ProjectFileGenerator";
+				AssemblyName = "RocketProjectFileGenerator";
 			}
-			else if( bGeneratingRocketProjectFiles )
+			else if (bGeneratingGameProjectFiles)
 			{
-				AssemblyName = "RocketProjectFileGenerator";
+				AssemblyName = GameProjectName + "ProjectFileGenerator";
 			}
-
-			List<string> AssemblyGameFolders = new List<string>();
 			foreach (UProjectInfo Project in AllGameProjects)
 			{
 				AssemblyGameFolders.Add(Project.Folder);
 			}
-			RulesCompiler.SetAssemblyNameAndGameFolders( AssemblyName, AssemblyGameFolders );
+			RulesCompiler.SetAssemblyNameAndGameFolders(AssemblyName, AssemblyGameFolders);
 
 			ProjectFile EngineProject = null;
 			Dictionary<string, ProjectFile> GameProjects = null;
@@ -1725,6 +1724,16 @@ namespace UnrealBuildTool
 
 				if (WantProjectFileForTarget)
 				{
+					var AssemblyName = bGeneratingRocketProjectFiles ? "RocketUE4" : "UE4";
+					List<string> AssemblyGameFolders = new List<string>();
+					string CheckProjectFile = UProjectInfo.GetProjectForTarget(TargetName);
+					if (string.IsNullOrEmpty(CheckProjectFile) == false)
+					{
+						AssemblyGameFolders.Add(Path.GetDirectoryName(Path.GetFullPath(CheckProjectFile)));
+						AssemblyName = Path.GetFileNameWithoutExtension(CheckProjectFile);
+					}
+					RulesCompiler.SetAssemblyNameAndGameFolders(AssemblyName, AssemblyGameFolders);
+
 					// Create target rules for all of the platforms and configuration combinations that we want to enable support for.
 					// Just use the current platform as we only need to recover the target type and both should be supported for all targets...
 					string UnusedTargetFilePath;
diff --git a/Engine/Source/Programs/UnrealBuildTool/System/RulesCompiler.cs b/Engine/Source/Programs/UnrealBuildTool/System/RulesCompiler.cs
--- a/Engine/Source/Programs/UnrealBuildTool/System/RulesCompiler.cs
+++ b/Engine/Source/Programs/UnrealBuildTool/System/RulesCompiler.cs
@@ -1288,25 +1288,324 @@ namespace UnrealBuildTool
 			return SourceFiles;
 		}
 
-		/// Assembly that contains object types for module rules definitions, loaded (or compiled) on demand */
-		private static Assembly RulesAssembly = null;
+		class RulesAssembly
+		{
+			public Assembly Rules;
+
+			/// Maps module names to their actual xxx.Module.cs file on disk
+			public Dictionary<string, string> ModuleNameToModuleFileMap = new Dictionary<string, string>(StringComparer.InvariantCultureIgnoreCase);
+
+			/// Maps target names to their actual xxx.Target.cs file on disk
+			public Dictionary<string, string> TargetNameToTargetFileMap = new Dictionary<string, string>(StringComparer.InvariantCultureIgnoreCase);
+
+			public RulesAssembly(Assembly InRules)
+			{
+				Rules = InRules;
+			}
+
+			/// <summary>
+			/// 
+			/// </summary>
+			/// <param name="InModuleName"></param>
+			/// <returns></returns>
+			public string GetModuleFilename(string InModuleName)
+			{
+				// Make sure the module file is known to us
+				if (!ModuleNameToModuleFileMap.ContainsKey(InModuleName))
+				{
+					return "";
+				}
+
+				// Return the module file name to the caller
+				return ModuleNameToModuleFileMap[InModuleName];
+			}
+
+			public string GetTargetFilename(string InTargetName)
+			{
+				// Make sure the target file is known to us
+				if (!TargetNameToTargetFileMap.ContainsKey(InTargetName))
+				{
+					return "";
+				}
+
+				// Return the target file name to the caller
+				return TargetNameToTargetFileMap[InTargetName];
+			}
+
+			/// <summary>
+			/// Creates an instance of a module rules descriptor object for the specified module name
+			/// </summary>
+			/// <param name="ModuleName">Name of the module</param>
+			/// <param name="Target">Information about the target associated with this module</param>
+			/// <param name="ModuleFileName">The original source file name for the Module.cs file for this module</param>
+			/// <returns>Compiled module rule info</returns>
+			public ModuleRules CreateModuleRules(string ModuleName, TargetInfo Target, out string ModuleFileName)
+			{
+				var AssemblyFileName = Path.GetFileNameWithoutExtension(Rules.Location);
+
+				// Currently, we expect the user's rules object type name to be the same as the module name
+				var ModuleTypeName = ModuleName;
+
+				// Make sure the module file is known to us
+				if (!ModuleNameToModuleFileMap.ContainsKey(ModuleName))
+				{
+					throw new MissingModuleException(ModuleName);
+				}
+
+				// Return the module file name to the caller
+				ModuleFileName = ModuleNameToModuleFileMap[ModuleName];
+
+				UnrealTargetPlatform LocalPlatform = Target.Platform;
+				UnrealTargetConfiguration LocalConfiguration = Target.Configuration;
+				TargetInfo LocalTarget = new TargetInfo(LocalPlatform, LocalConfiguration, Target.Type.Value, Target.bIsMonolithic.Value);
+
+				// The build module must define a type named 'Rules' that derives from our 'ModuleRules' type.  
+				var RulesObjectType = Rules.GetType(ModuleName);
+
+				if (RulesObjectType == null)
+				{
+					// Temporary hack to avoid System namespace collisions
+					// @todo projectfiles: Make rules assemblies require namespaces.
+					RulesObjectType = Rules.GetType("UnrealBuildTool.Rules." + ModuleName);
+				}
+
+				if (RulesObjectType == null)
+				{
+					throw new BuildException("Expecting to find a type to be declared in a module rules named '{0}' in {1}.  This type must derive from the 'ModuleRules' type defined by Unreal Build Tool.", ModuleTypeName, Rules.FullName);
+				}
+
+				// Create an instance of the module's rules object
+				ModuleRules RulesObject;
+				try
+				{
+					RulesObject = (ModuleRules)Activator.CreateInstance(RulesObjectType, LocalTarget);
+				}
+				catch (Exception Ex)
+				{
+					throw new BuildException(Ex, "Unable to instantiate instance of '{0}' object type from compiled assembly '{1}'.  Unreal Build Tool creates an instance of your module's 'Rules' object in order to find out about your module's requirements.  The CLR exception details may provide more information:  {2}", ModuleTypeName, AssemblyFileName, Ex.ToString());
+				}
+
+				// Have to do absolute here as this could be a project that is under the root
+				var FullUProjectPath = string.IsNullOrWhiteSpace(UnrealBuildTool.GetUProjectPath())
+					? ""
+					: Path.GetFullPath(UnrealBuildTool.GetUProjectPath());
+				var bProjectModule = string.IsNullOrWhiteSpace(FullUProjectPath)
+					? false
+					: Utils.IsFileUnderDirectory(ModuleFileName, FullUProjectPath);
+
+				if (bProjectModule)
+				{
+					RulesObject.ReadAdditionalDependencies(UnrealBuildTool.GetUProjectFile(), ModuleName);
+				}
+
+				// Validate rules object
+				{
+					if (RulesObject.Type == ModuleRules.ModuleType.CPlusPlus)
+					{
+						if (RulesObject.PrivateAssemblyReferences.Count > 0)
+						{
+							throw new BuildException("Module rules for '{0}' may not specify PrivateAssemblyReferences unless it is a CPlusPlusCLR module type.", AssemblyFileName);
+						}
+
+						var InvalidDependencies = RulesObject.DynamicallyLoadedModuleNames.Intersect(RulesObject.PublicDependencyModuleNames.Concat(RulesObject.PrivateDependencyModuleNames)).ToList();
+						if (InvalidDependencies.Count != 0)
+						{
+							throw new BuildException("Module rules for '{0}' should not be dependent on modules which are also dynamically loaded: {1}", ModuleName, String.Join(", ", InvalidDependencies));
+						}
+
+						// Choose code optimization options based on module type (game/engine) if
+						// default optimization method is selected.
+						bool bIsEngineModule = Utils.IsFileUnderDirectory(ModuleFileName, ProjectFileGenerator.EngineRelativePath);
+						if (RulesObject.OptimizeCode == ModuleRules.CodeOptimization.Default)
+						{
+							// Engine/Source and Engine/Plugins are considered 'Engine' code...
+							if (bIsEngineModule)
+							{
+								// Engine module - always optimize (except Debug).
+								RulesObject.OptimizeCode = ModuleRules.CodeOptimization.Always;
+							}
+							else
+							{
+								// Game module - do not optimize in Debug and DebugGame builds.
+								RulesObject.OptimizeCode = ModuleRules.CodeOptimization.InNonDebugBuilds;
+							}
+						}
+
+						// Disable shared PCHs for game modules by default
+						if (RulesObject.PCHUsage == ModuleRules.PCHUsageMode.Default)
+						{
+							// Note that bIsEngineModule includes Engine/Plugins, so Engine/Plugins will use shared PCHs.
+							var IsProgramTarget = Target.Type != null && Target.Type == TargetRules.TargetType.Program;
+							if (bIsEngineModule || IsProgramTarget)
+							{
+								// Engine module or plugin module -- allow shared PCHs
+								RulesObject.PCHUsage = ModuleRules.PCHUsageMode.UseSharedPCHs;
+							}
+							else
+							{
+								// Game module.  Do not enable shared PCHs by default, because games usually have a large precompiled header of their own and compile times would suffer.
+								RulesObject.PCHUsage = ModuleRules.PCHUsageMode.NoSharedPCHs;
+							}
+						}
+					}
+				}
+
+				return RulesObject;
+			}
+
+			protected bool GetTargetTypeAndRulesInstance(string InTargetName, TargetInfo InTarget, out System.Type OutRulesObjectType, out TargetRules OutRulesObject)
+			{
+				// The build module must define a type named '<TargetName>Target' that derives from our 'TargetRules' type.  
+				OutRulesObjectType = Rules.GetType(InTargetName);
+				if (OutRulesObjectType == null)
+				{
+					throw new BuildException(
+						"Expecting to find a type to be declared in a target rules named '{0}'.  This type must derive from the 'TargetRules' type defined by Unreal Build Tool.",
+						InTargetName);
+				}
+
+				// Create an instance of the module's rules object
+				try
+				{
+					OutRulesObject = (TargetRules)Activator.CreateInstance(OutRulesObjectType, InTarget);
+				}
+				catch (Exception Ex)
+				{
+					var AssemblyFileName = Path.GetFileNameWithoutExtension(Rules.Location);
+					throw new BuildException(Ex,
+						"Unable to instantiate instance of '{0}' object type from compiled assembly '{1}'.  Unreal Build Tool creates an instance of your module's 'Rules' object in order to find out about your module's requirements.  The CLR exception details may provide more information:  {2}",
+						InTargetName, AssemblyFileName, Ex.ToString());
+				}
+
+				OutRulesObject.TargetName = InTargetName;
+
+				return true;
+			}
+
+			/// <summary>
+			/// Creates a target rules object for the specified target name.
+			/// </summary>
+			/// <param name="TargetName">Name of the target</param>
+			/// <param name="Target">Information about the target associated with this target</param>
+			/// <param name="TargetFileName">The original source file name of the Target.cs file for this target</param>
+			/// <returns>The build target rules for the specified target</returns>
+			public TargetRules CreateTargetRules(string TargetName, TargetInfo Target, bool bInEditorRecompile, out string TargetFileName)
+			{
+				// Make sure the target file is known to us
+				bool bFoundTargetName = TargetNameToTargetFileMap.ContainsKey(TargetName);
+				if (bFoundTargetName == false)
+				{
+					if (UnrealBuildTool.RunningRocket())
+					{
+						//@todo Rocket: Remove this when full game support is implemented
+						// If we are Rocket, they will currently only have an editor target.
+						// See if that exists
+						bFoundTargetName = TargetNameToTargetFileMap.ContainsKey(TargetName + "Editor");
+						if (bFoundTargetName)
+						{
+							TargetName += "Editor";
+						}
+					}
+				}
+
+				if (bFoundTargetName == false)
+				{
+					//				throw new BuildException("Couldn't find target rules file for target '{0}' in rules assembly '{1}'.", TargetName, RulesAssembly.FullName);
+					string ExceptionMessage = "Couldn't find target rules file for target '";
+					ExceptionMessage += TargetName;
+					ExceptionMessage += "' in rules assembly '";
+					ExceptionMessage += Rules.FullName;
+					ExceptionMessage += "'." + Environment.NewLine;
+
+					ExceptionMessage += "Location: " + Rules.Location + Environment.NewLine;
+
+					ExceptionMessage += "Target rules found:" + Environment.NewLine;
+					foreach (KeyValuePair<string, string> entry in TargetNameToTargetFileMap)
+					{
+						ExceptionMessage += "\t" + entry.Key + " - " + entry.Value + Environment.NewLine;
+					}
+
+					throw new BuildException(ExceptionMessage);
+				}
+
+				// Return the target file name to the caller
+				TargetFileName = TargetNameToTargetFileMap[TargetName];
+
+				// Currently, we expect the user's rules object type name to be the same as the module name + 'Target'
+				string TargetTypeName = TargetName + "Target";
+
+				// The build module must define a type named '<TargetName>Target' that derives from our 'TargetRules' type.  
+				System.Type RulesObjectType;
+				TargetRules RulesObject;
+				GetTargetTypeAndRulesInstance(TargetTypeName, Target, out RulesObjectType, out RulesObject);
+				if (bInEditorRecompile)
+				{
+					// Make sure this is an editor module.
+					if (RulesObject != null)
+					{
+						if (RulesObject.Type != TargetRules.TargetType.Editor)
+						{
+							// Not the editor... determine the editor project
+							string TargetSourceFolder = TargetFileName;
+							Int32 SourceFolderIndex = -1;
+							if (Utils.IsRunningOnMono)
+							{
+								TargetSourceFolder = TargetSourceFolder.Replace("\\", "/");
+								SourceFolderIndex = TargetSourceFolder.LastIndexOf("/Source/", StringComparison.InvariantCultureIgnoreCase);
+							}
+							else
+							{
+								TargetSourceFolder = TargetSourceFolder.Replace("/", "\\");
+								SourceFolderIndex = TargetSourceFolder.LastIndexOf("\\Source\\", StringComparison.InvariantCultureIgnoreCase);
+							}
+							if (SourceFolderIndex != -1)
+							{
+								TargetSourceFolder = TargetSourceFolder.Substring(0, SourceFolderIndex + 8);
+								foreach (KeyValuePair<string, string> CheckEntry in TargetNameToTargetFileMap)
+								{
+									if (CheckEntry.Value.StartsWith(TargetSourceFolder, StringComparison.InvariantCultureIgnoreCase))
+									{
+										if (CheckEntry.Key.Equals(TargetName, StringComparison.InvariantCultureIgnoreCase) == false)
+										{
+											// We have found a target in the same source folder that is not the original target found.
+											// See if it is the editor project
+											string CheckTargetTypeName = CheckEntry.Key + "Target";
+											System.Type CheckRulesObjectType;
+											TargetRules CheckRulesObject;
+											GetTargetTypeAndRulesInstance(CheckTargetTypeName, Target, out CheckRulesObjectType, out CheckRulesObject);
+											if (CheckRulesObject != null)
+											{
+												if (CheckRulesObject.Type == TargetRules.TargetType.Editor)
+												{
+													// Found it
+													// NOTE: This prevents multiple Editor targets from co-existing...
+													RulesObject = CheckRulesObject;
+													break;
+												}
+											}
+										}
+									}
+								}
+							}
+						}
+					}
+				}
 
-		/// Maps module names to their actual xxx.Module.cs file on disk
-		private static Dictionary<string, string> ModuleNameToModuleFileMap = new Dictionary<string, string>( StringComparer.InvariantCultureIgnoreCase );
+				return RulesObject;
+			}
 
-		/// Maps target names to their actual xxx.Target.cs file on disk
-		private static Dictionary<string, string> TargetNameToTargetFileMap = new Dictionary<string, string>( StringComparer.InvariantCultureIgnoreCase );
+		}
 
 		private class LoadedAssemblyData
 		{
-			public LoadedAssemblyData(Assembly InAssembly, List<string> InGameFolders)
+			public LoadedAssemblyData(RulesAssembly InAssembly, List<string> InGameFolders)
 			{
 				ExistingAssembly    = InAssembly;
 				ExistingGameFolders = InGameFolders;
 			}
 
-			public Assembly     ExistingAssembly     { get; private set; }
-			public List<string> ExistingGameFolders  { get; private set; }
+			public RulesAssembly     ExistingAssembly     { get; private set; }
+			public List<string>		ExistingGameFolders  { get; private set; }
 		}
 
 		/// Map of assembly names we've already compiled and loaded to their Assembly and list of game folders.  This is used to prevent
@@ -1327,10 +1626,8 @@ namespace UnrealBuildTool
 			LoadedAssemblyData LoadedAssembly;
 			if (LoadedAssemblyMap.TryGetValue(AssemblyName, out LoadedAssembly))
 			{
-				Assembly     ExistingAssembly    = LoadedAssembly.ExistingAssembly;
-				List<string> ExistingGameFolders = LoadedAssembly.ExistingGameFolders;
-
-				RulesAssembly = ExistingAssembly;
+				RulesAssembly	ExistingAssembly    = LoadedAssembly.ExistingAssembly;
+				List<string>	ExistingGameFolders = LoadedAssembly.ExistingGameFolders;
 
 				// Make sure the game folder list wasn't changed since we last compiled this assembly
 				if( ExistingGameFolders != AllGameFolders )	// Quick-check pointers first to avoid iterating
@@ -1369,8 +1666,6 @@ namespace UnrealBuildTool
 				return;
 			}
 
-			RulesAssembly = null;
-
 			var AdditionalSearchPaths = new List<string>();
 
 			if (UnrealBuildTool.HasUProjectFile())
@@ -1405,16 +1700,16 @@ namespace UnrealBuildTool
 
 			string OutputAssemblyPath = Path.GetFullPath(Path.Combine(BaseIntermediatePath, "BuildRules", AssemblyName + ".dll"));
 
-			RulesAssembly = DynamicCompilation.CompileAndLoadAssembly( OutputAssemblyPath, AssemblySourceFiles );
+			RulesAssembly Rules = new RulesAssembly(DynamicCompilation.CompileAndLoadAssembly( OutputAssemblyPath, AssemblySourceFiles ));
 
 			// Setup the module map
 			foreach( var CurModuleFileName in ModuleFileNames )
 			{
 				var CleanFileName = Utils.CleanDirectorySeparators( CurModuleFileName );
 				var ModuleName = Path.GetFileNameWithoutExtension( Path.GetFileNameWithoutExtension( CleanFileName ) );	// Strip both extensions
-				if( !ModuleNameToModuleFileMap.ContainsKey( ModuleName ) )
+				if( !Rules.ModuleNameToModuleFileMap.ContainsKey( ModuleName ) )
 				{
-					ModuleNameToModuleFileMap.Add( ModuleName, CurModuleFileName );
+					Rules.ModuleNameToModuleFileMap.Add( ModuleName, CurModuleFileName );
 				}
 			}
 
@@ -1423,15 +1718,15 @@ namespace UnrealBuildTool
 			{
 				var CleanFileName = Utils.CleanDirectorySeparators( CurTargetFileName );
 				var TargetName = Path.GetFileNameWithoutExtension( Path.GetFileNameWithoutExtension( CleanFileName ) );	// Strip both extensions
-				if( !TargetNameToTargetFileMap.ContainsKey( TargetName ) )
+				if( !Rules.TargetNameToTargetFileMap.ContainsKey( TargetName ) )
 				{
-					TargetNameToTargetFileMap.Add( TargetName, CurTargetFileName );
+					Rules.TargetNameToTargetFileMap.Add( TargetName, CurTargetFileName );
 				}
 			}
 
 			// Remember that we loaded this assembly
-			var RulesAssemblyName = Path.GetFileNameWithoutExtension( RulesAssembly.Location );
-			LoadedAssemblyMap[RulesAssemblyName] = new LoadedAssemblyData(RulesAssembly, AllGameFolders);
+			var RulesAssemblyName = Path.GetFileNameWithoutExtension( Rules.Rules.Location );
+			LoadedAssemblyMap[RulesAssemblyName] = new LoadedAssemblyData(Rules, AllGameFolders);
 		}
 
 		/// <summary>
@@ -1441,26 +1736,12 @@ namespace UnrealBuildTool
 		/// <returns></returns>
 		public static string GetModuleFilename(string InModuleName)
 		{
-			// Make sure the module file is known to us
-			if (!ModuleNameToModuleFileMap.ContainsKey(InModuleName))
-			{
-				return "";
-			}
-
-			// Return the module file name to the caller
-			return ModuleNameToModuleFileMap[InModuleName];
+			return LoadedAssemblyMap[AssemblyName].ExistingAssembly.GetModuleFilename(InModuleName);
 		}
 
 		public static string GetTargetFilename(string InTargetName)
 		{
-			// Make sure the target file is known to us
-			if (!TargetNameToTargetFileMap.ContainsKey(InTargetName))
-			{
-				return "";
-			}
-
-			// Return the target file name to the caller
-			return TargetNameToTargetFileMap[InTargetName];
+			return LoadedAssemblyMap[AssemblyName].ExistingAssembly.GetTargetFilename(InTargetName);
 		}
 
 		/// <summary>
@@ -1527,116 +1808,7 @@ namespace UnrealBuildTool
 		public static ModuleRules CreateModuleRules( string ModuleName, TargetInfo Target, out string ModuleFileName )
 		{
 			ConditionallyCompileAndLoadRulesAssembly();
-			var AssemblyFileName = Path.GetFileNameWithoutExtension( RulesAssembly.Location );
-
-			// Currently, we expect the user's rules object type name to be the same as the module name
-			var ModuleTypeName = ModuleName;
-
-			// Make sure the module file is known to us
-			if( !ModuleNameToModuleFileMap.ContainsKey( ModuleName ) )
-			{
-				throw new MissingModuleException( ModuleName );
-			}
-
-			// Return the module file name to the caller
-			ModuleFileName = ModuleNameToModuleFileMap[ ModuleName ];
-
-			UnrealTargetPlatform LocalPlatform = Target.Platform;
-			UnrealTargetConfiguration LocalConfiguration = Target.Configuration;
-			TargetInfo LocalTarget = new TargetInfo(LocalPlatform, LocalConfiguration, Target.Type.Value, Target.bIsMonolithic.Value);
-
-			// The build module must define a type named 'Rules' that derives from our 'ModuleRules' type.  
-			var RulesObjectType = RulesAssembly.GetType( ModuleName );
-
-			if (RulesObjectType == null)
-			{
-				// Temporary hack to avoid System namespace collisions
-				// @todo projectfiles: Make rules assemblies require namespaces.
-				RulesObjectType = RulesAssembly.GetType("UnrealBuildTool.Rules." + ModuleName);
-			}
-
-			if( RulesObjectType == null )
-			{
-				throw new BuildException( "Expecting to find a type to be declared in a module rules named '{0}' in {1}.  This type must derive from the 'ModuleRules' type defined by Unreal Build Tool.", ModuleTypeName, RulesAssembly.FullName );
-			}
-
-			// Create an instance of the module's rules object
-			ModuleRules RulesObject;
-			try
-			{
-				RulesObject = (ModuleRules)Activator.CreateInstance(RulesObjectType, LocalTarget);
-			}
-			catch( Exception Ex )
-			{
-				throw new BuildException( Ex, "Unable to instantiate instance of '{0}' object type from compiled assembly '{1}'.  Unreal Build Tool creates an instance of your module's 'Rules' object in order to find out about your module's requirements.  The CLR exception details may provide more information:  {2}", ModuleTypeName, AssemblyFileName, Ex.ToString() );
-			}
-
-			// Have to do absolute here as this could be a project that is under the root
-			var FullUProjectPath = string.IsNullOrWhiteSpace(UnrealBuildTool.GetUProjectPath())
-				? ""
-				: Path.GetFullPath(UnrealBuildTool.GetUProjectPath());
-			var bProjectModule = string.IsNullOrWhiteSpace(FullUProjectPath)
-				? false
-				: Utils.IsFileUnderDirectory(ModuleFileName, FullUProjectPath);
-
-			if (bProjectModule)
-			{
-				RulesObject.ReadAdditionalDependencies(UnrealBuildTool.GetUProjectFile(), ModuleName);
-			}
-
-			// Validate rules object
-			{
-				if( RulesObject.Type == ModuleRules.ModuleType.CPlusPlus )
-				{
-					if (RulesObject.PrivateAssemblyReferences.Count > 0)
-					{
-						throw new BuildException("Module rules for '{0}' may not specify PrivateAssemblyReferences unless it is a CPlusPlusCLR module type.", ModuleName);
-					}
-
-					var InvalidDependencies = RulesObject.DynamicallyLoadedModuleNames.Intersect(RulesObject.PublicDependencyModuleNames.Concat(RulesObject.PrivateDependencyModuleNames)).ToList();
-					if (InvalidDependencies.Count != 0)
-					{
-						throw new BuildException("Module rules for '{0}' should not be dependent on modules which are also dynamically loaded: {1}", ModuleName, String.Join(", ", InvalidDependencies));
-					}
-
-					// Choose code optimization options based on module type (game/engine) if
-					// default optimization method is selected.
-					bool bIsEngineModule = Utils.IsFileUnderDirectory( ModuleFileName, ProjectFileGenerator.EngineRelativePath );
-					if (RulesObject.OptimizeCode == ModuleRules.CodeOptimization.Default)
-					{
-						// Engine/Source and Engine/Plugins are considered 'Engine' code...
-						if (bIsEngineModule)
-						{
-							// Engine module - always optimize (except Debug).
-							RulesObject.OptimizeCode = ModuleRules.CodeOptimization.Always;
-						}
-						else
-						{
-							// Game module - do not optimize in Debug and DebugGame builds.
-							RulesObject.OptimizeCode = ModuleRules.CodeOptimization.InNonDebugBuilds;
-						}
-					}
-
-					// Disable shared PCHs for game modules by default
-					if (RulesObject.PCHUsage == ModuleRules.PCHUsageMode.Default)
-					{
-						// Note that bIsEngineModule includes Engine/Plugins, so Engine/Plugins will use shared PCHs.
-						var IsProgramTarget = Target.Type != null && Target.Type == TargetRules.TargetType.Program;
-						if (bIsEngineModule || IsProgramTarget)
-						{
-							// Engine module or plugin module -- allow shared PCHs
-							RulesObject.PCHUsage = ModuleRules.PCHUsageMode.UseSharedPCHs;
-						}
-						else
-						{
-							// Game module.  Do not enable shared PCHs by default, because games usually have a large precompiled header of their own and compile times would suffer.
-							RulesObject.PCHUsage = ModuleRules.PCHUsageMode.NoSharedPCHs;
-						}
-					}
-				}
-			}
-
-			return RulesObject;
+			return LoadedAssemblyMap[AssemblyName].ExistingAssembly.CreateModuleRules(ModuleName, Target, out ModuleFileName);
 		}
 
 		/// <summary>
@@ -1791,35 +1963,6 @@ namespace UnrealBuildTool
 			}
 		}
 
-		protected static bool GetTargetTypeAndRulesInstance(string InTargetName, TargetInfo InTarget, out System.Type OutRulesObjectType, out TargetRules OutRulesObject)
-		{
-			// The build module must define a type named '<TargetName>Target' that derives from our 'TargetRules' type.  
-			OutRulesObjectType = RulesAssembly.GetType(InTargetName);
-			if (OutRulesObjectType == null)
-			{
-				throw new BuildException(
-					"Expecting to find a type to be declared in a target rules named '{0}'.  This type must derive from the 'TargetRules' type defined by Unreal Build Tool.", 
-					InTargetName);
-			}
-
-			// Create an instance of the module's rules object
-			try
-			{
-				OutRulesObject = (TargetRules)Activator.CreateInstance(OutRulesObjectType, InTarget);
-			}
-			catch (Exception Ex)
-			{
-				var AssemblyFileName = Path.GetFileNameWithoutExtension(RulesAssembly.Location);
-				throw new BuildException(Ex,
-					"Unable to instantiate instance of '{0}' object type from compiled assembly '{1}'.  Unreal Build Tool creates an instance of your module's 'Rules' object in order to find out about your module's requirements.  The CLR exception details may provide more information:  {2}", 
-					InTargetName, AssemblyFileName, Ex.ToString());
-			}
-
-            OutRulesObject.TargetName = InTargetName;
-
-			return true;
-		}
-
 		/// <summary>
 		/// Creates a target rules object for the specified target name.
 		/// </summary>
@@ -1830,108 +1973,7 @@ namespace UnrealBuildTool
 		public static TargetRules CreateTargetRules(string TargetName, TargetInfo Target, bool bInEditorRecompile, out string TargetFileName)
 		{
 			ConditionallyCompileAndLoadRulesAssembly();
-
-			// Make sure the target file is known to us
-			bool bFoundTargetName = TargetNameToTargetFileMap.ContainsKey(TargetName);
-			if (bFoundTargetName == false)
-			{
-				if (UnrealBuildTool.RunningRocket())
-				{
-					//@todo Rocket: Remove this when full game support is implemented
-					// If we are Rocket, they will currently only have an editor target.
-					// See if that exists
-					bFoundTargetName = TargetNameToTargetFileMap.ContainsKey(TargetName + "Editor");
-					if (bFoundTargetName)
-					{
-						TargetName += "Editor";
-					}
-				}
-			}
-
-			if (bFoundTargetName == false)
-			{
-//				throw new BuildException("Couldn't find target rules file for target '{0}' in rules assembly '{1}'.", TargetName, RulesAssembly.FullName);
-				string ExceptionMessage = "Couldn't find target rules file for target '";
-				ExceptionMessage += TargetName;
-				ExceptionMessage += "' in rules assembly '";
-				ExceptionMessage += RulesAssembly.FullName;
-                ExceptionMessage += "'." + Environment.NewLine;
-
-                ExceptionMessage += "Location: " + RulesAssembly.Location + Environment.NewLine;
-
-                ExceptionMessage += "Target rules found:" + Environment.NewLine;
-				foreach (KeyValuePair<string, string> entry in TargetNameToTargetFileMap)
-				{
-                    ExceptionMessage += "\t" + entry.Key + " - " + entry.Value + Environment.NewLine;
-				}
-
-				throw new BuildException(ExceptionMessage);
-			}
-
-			// Return the target file name to the caller
-			TargetFileName = TargetNameToTargetFileMap[ TargetName ];
-	
-			// Currently, we expect the user's rules object type name to be the same as the module name + 'Target'
-			string TargetTypeName = TargetName + "Target";
-
- 			// The build module must define a type named '<TargetName>Target' that derives from our 'TargetRules' type.  
-			System.Type RulesObjectType;
-			TargetRules RulesObject;
-			GetTargetTypeAndRulesInstance(TargetTypeName, Target, out RulesObjectType, out RulesObject);
-			if (bInEditorRecompile)
-			{
-				// Make sure this is an editor module.
-				if (RulesObject != null)
-				{
-					if (RulesObject.Type != TargetRules.TargetType.Editor)
-					{
-						// Not the editor... determine the editor project
-						string TargetSourceFolder = TargetFileName;
-						Int32 SourceFolderIndex = -1;
-						if (Utils.IsRunningOnMono)
-						{
-							TargetSourceFolder = TargetSourceFolder.Replace("\\", "/");
-							SourceFolderIndex = TargetSourceFolder.LastIndexOf("/Source/", StringComparison.InvariantCultureIgnoreCase);
-						}
-						else
-						{
-							TargetSourceFolder = TargetSourceFolder.Replace("/", "\\");
-							SourceFolderIndex = TargetSourceFolder.LastIndexOf("\\Source\\", StringComparison.InvariantCultureIgnoreCase);
-						}
-						if (SourceFolderIndex != -1)
-						{
-							TargetSourceFolder = TargetSourceFolder.Substring(0, SourceFolderIndex + 8);
-							foreach (KeyValuePair<string, string> CheckEntry in TargetNameToTargetFileMap)
-							{
-								if (CheckEntry.Value.StartsWith(TargetSourceFolder, StringComparison.InvariantCultureIgnoreCase))
-								{
-									if (CheckEntry.Key.Equals(TargetName, StringComparison.InvariantCultureIgnoreCase) == false)
-									{
-										// We have found a target in the same source folder that is not the original target found.
-										// See if it is the editor project
-										string CheckTargetTypeName = CheckEntry.Key + "Target";
-										System.Type CheckRulesObjectType;
-										TargetRules CheckRulesObject;
-										GetTargetTypeAndRulesInstance(CheckTargetTypeName, Target, out CheckRulesObjectType, out CheckRulesObject);
-										if (CheckRulesObject != null)
-										{
-											if (CheckRulesObject.Type == TargetRules.TargetType.Editor)
-											{
-												// Found it
-												// NOTE: This prevents multiple Editor targets from co-existing...
-												RulesObject = CheckRulesObject;
-												break;
-											}
-										}
-									}
-								}
-							}
-						}
-					}
-				}
-			}
-
-			return RulesObject;
+			return LoadedAssemblyMap[AssemblyName].ExistingAssembly.CreateTargetRules(TargetName, Target, bInEditorRecompile, out TargetFileName);
 		}
 
 		/// <summary>
@@ -1945,6 +1987,8 @@ namespace UnrealBuildTool
 		{
 			var CreateTargetStartTime = DateTime.UtcNow;
 
+			RulesCompiler.SetAssemblyNameAndGameFolders(Desc.AssemblyName, Desc.GameFolders);
+
 			string TargetFileName;
 			TargetRules RulesObject = CreateTargetRules(Desc.TargetName, new TargetInfo(Desc.Platform, Desc.Configuration), Desc.bIsEditorRecompile, out TargetFileName);
 			if (Desc.bIsEditorRecompile)
diff --git a/Engine/Source/Programs/UnrealBuildTool/System/XcodeProjectFileGenerator.cs b/Engine/Source/Programs/UnrealBuildTool/System/XcodeProjectFileGenerator.cs
--- a/Engine/Source/Programs/UnrealBuildTool/System/XcodeProjectFileGenerator.cs
+++ b/Engine/Source/Programs/UnrealBuildTool/System/XcodeProjectFileGenerator.cs
@@ -1071,109 +1071,113 @@ namespace UnrealBuildTool
 			}
 		}
 
-		private void PopulateTargets(List<string> GamePaths, List<XcodeProjectTarget> ProjectTargets, List<XcodeContainerItemProxy> ContainerItemProxies, List<XcodeTargetDependency> TargetDependencies, XcodeProjectTarget ProjectTarget, List<XcodeFramework> Frameworks)
+		private void PopulateTargets(List<XcodeProjectTarget> ProjectTargets, List<XcodeContainerItemProxy> ContainerItemProxies, List<XcodeTargetDependency> TargetDependencies, XcodeProjectTarget ProjectTarget, List<XcodeFramework> Frameworks)
 		{
-			foreach (string TargetPath in GamePaths)
+			foreach (var Project in GeneratedProjectFiles)
 			{
-				string TargetName = Utils.GetFilenameWithoutAnyExtensions(Path.GetFileName(TargetPath));
-				bool WantProjectFileForTarget = true;
-				bool IsEngineTarget = false;
-				if (bGeneratingGameProjectFiles || bGeneratingRocketProjectFiles)
+				foreach (var TargetProject in Project.ProjectTargets)
 				{
-					// Check to see if this is an Engine target.  That is, the target is located under the "Engine" folder
-					string TargetFileRelativeToEngineDirectory = Utils.MakePathRelativeTo(TargetPath, Path.Combine(EngineRelativePath), AlwaysTreatSourceAsDirectory: false);
-					if (!TargetFileRelativeToEngineDirectory.StartsWith("..") && !Path.IsPathRooted(TargetFileRelativeToEngineDirectory))
+					string TargetPath = TargetProject.TargetFilePath;
+					string TargetName = Utils.GetFilenameWithoutAnyExtensions(Path.GetFileName(TargetPath));
+					bool WantProjectFileForTarget = true;
+					bool IsEngineTarget = false;
+					if (bGeneratingGameProjectFiles || bGeneratingRocketProjectFiles)
 					{
-						// This is an engine target
-						IsEngineTarget = true;
-					}
-
-					if (IsEngineTarget)
-					{
-						if (!IncludeEngineSource)
+						// Check to see if this is an Engine target.  That is, the target is located under the "Engine" folder
+						string TargetFileRelativeToEngineDirectory = Utils.MakePathRelativeTo(TargetPath, Path.Combine(EngineRelativePath), AlwaysTreatSourceAsDirectory: false);
+						if (!TargetFileRelativeToEngineDirectory.StartsWith("..") && !Path.IsPathRooted(TargetFileRelativeToEngineDirectory))
 						{
-							// We were asked to exclude engine modules from the generated projects
-							WantProjectFileForTarget = false;
+							// This is an engine target
+							IsEngineTarget = true;
 						}
-						if (bGeneratingGameProjectFiles && this.GameProjectName == TargetName)
-						{
-							WantProjectFileForTarget = true;
-						}
-					}
 
-					if (bGeneratingRocketProjectFiles && TargetName.EndsWith("Server"))
-					{
-						WantProjectFileForTarget = false;
-					}
-				}
-
-				if (WantProjectFileForTarget)
-				{
-					string TargetFilePath;
-					var Target = new TargetInfo(UnrealTargetPlatform.Mac, UnrealTargetConfiguration.Development);
-					var TargetRulesObject = RulesCompiler.CreateTargetRules(TargetName, Target, false, out TargetFilePath);
-					List<UnrealTargetPlatform> SupportedPlatforms = new List<UnrealTargetPlatform>();
-					TargetRulesObject.GetSupportedPlatforms(ref SupportedPlatforms);
-					LinkEnvironmentConfiguration LinkConfiguration = new LinkEnvironmentConfiguration();
-					CPPEnvironmentConfiguration CPPConfiguration = new CPPEnvironmentConfiguration();
-					TargetRulesObject.SetupGlobalEnvironment(Target, ref LinkConfiguration, ref CPPConfiguration);
-
-					if (!LinkConfiguration.bIsBuildingConsoleApplication)
-					{
-						TargetsThatNeedApp.Add(TargetName);
-					}
-
-					// if the project is not an engine project check to make sure we have the correct name
-					string DisplayName = TargetName;
-					if (!IsEngineTarget && TargetRulesObject.Type != TargetRules.TargetType.Program && TargetRulesObject.Type != TargetRules.TargetType.Client)
-					{
-						List<UProjectInfo> AllGames = UProjectInfo.FilterGameProjects(true, bGeneratingGameProjectFiles ? GameProjectName : null);
-						UProjectInfo ProjectInfo = FindGameContainingFile (AllGames, TargetFilePath);
-						if (ProjectInfo != null)
+						if (IsEngineTarget)
 						{
-							DisplayName = ProjectInfo.GameName;
-							if (TargetName.Contains("Editor"))
+							if (!IncludeEngineSource)
 							{
-								DisplayName += "Editor";
+								// We were asked to exclude engine modules from the generated projects
+								WantProjectFileForTarget = false;
 							}
-							else if (TargetName.Contains("Server"))
+							if (bGeneratingGameProjectFiles && this.GameProjectName == TargetName)
 							{
-								DisplayName += "Server";
+								WantProjectFileForTarget = true;
 							}
 						}
-					}
 
-					// @todo: Remove target platform param and merge Mac and iOS targets. For now BuildTarget knows how to build iOS, but cannot run iOS apps, so we need separate DeployTarget.
-					bool bIsMacOnly = !SupportedPlatforms.Contains(UnrealTargetPlatform.IOS);
-
-					XcodeProjectTarget BuildTarget = new XcodeProjectTarget(DisplayName + " - Mac", TargetName, XcodeTargetType.Legacy, TargetFilePath, "", UnrealTargetPlatform.Mac, bIsMacOnly);
-					if (!bGeneratingRunIOSProject)
-					{
-						ProjectTargets.Add(BuildTarget);
+						if (bGeneratingRocketProjectFiles && TargetName.EndsWith("Server"))
+						{
+							WantProjectFileForTarget = false;
+						}
 					}
 
-					if (ProjectFilePlatform.HasFlag(XcodeProjectFilePlatform.iOS) && SupportedPlatforms.Contains(UnrealTargetPlatform.IOS))
+					if (WantProjectFileForTarget)
 					{
-						if ((bGeneratingRocketProjectFiles && TargetName == "UE4Game") || bGeneratingRunIOSProject)
+						string TargetFilePath = TargetProject.TargetFilePath;
+						var Target = new TargetInfo(UnrealTargetPlatform.Mac, UnrealTargetConfiguration.Development);
+						var TargetRulesObject = TargetProject.TargetRules; // RulesCompiler.CreateTargetRules(TargetName, Target, false, out TargetFilePath);
+						List<UnrealTargetPlatform> SupportedPlatforms = new List<UnrealTargetPlatform>();
+						TargetRulesObject.GetSupportedPlatforms(ref SupportedPlatforms);
+						LinkEnvironmentConfiguration LinkConfiguration = new LinkEnvironmentConfiguration();
+						CPPEnvironmentConfiguration CPPConfiguration = new CPPEnvironmentConfiguration();
+						TargetRulesObject.SetupGlobalEnvironment(Target, ref LinkConfiguration, ref CPPConfiguration);
+
+						if (!LinkConfiguration.bIsBuildingConsoleApplication)
 						{
-							// Generate Framework references.
-							List<XcodeFrameworkRef> FrameworkRefs = new List<XcodeFrameworkRef>();
-							foreach (XcodeFramework Framework in Frameworks)
+							TargetsThatNeedApp.Add(TargetName);
+						}
+
+						// if the project is not an engine project check to make sure we have the correct name
+						string DisplayName = TargetName;
+						if (!IsEngineTarget && TargetRulesObject.Type != TargetRules.TargetType.Program && TargetRulesObject.Type != TargetRules.TargetType.Client)
+						{
+							List<UProjectInfo> AllGames = UProjectInfo.FilterGameProjects(true, bGeneratingGameProjectFiles ? GameProjectName : null);
+							UProjectInfo ProjectInfo = FindGameContainingFile(AllGames, TargetFilePath);
+							if (ProjectInfo != null)
 							{
-								FrameworkRefs.Add(new XcodeFrameworkRef(Framework));
+								DisplayName = ProjectInfo.GameName;
+								if (TargetName.Contains("Editor"))
+								{
+									DisplayName += "Editor";
+								}
+								else if (TargetName.Contains("Server"))
+								{
+									DisplayName += "Server";
+								}
 							}
+						}
 
-							XcodeProjectTarget IOSDeployTarget = new XcodeProjectTarget(DisplayName + " - iOS", TargetName, XcodeTargetType.Native, TargetFilePath, TargetName + ".app", UnrealTargetPlatform.IOS, false, null, true, FrameworkRefs);
-							ProjectTargets.Add(IOSDeployTarget);
+						// @todo: Remove target platform param and merge Mac and iOS targets. For now BuildTarget knows how to build iOS, but cannot run iOS apps, so we need separate DeployTarget.
+						bool bIsMacOnly = !SupportedPlatforms.Contains(UnrealTargetPlatform.IOS);
+
+						XcodeProjectTarget BuildTarget = new XcodeProjectTarget(DisplayName + " - Mac", TargetName, XcodeTargetType.Legacy, TargetFilePath, "", UnrealTargetPlatform.Mac, bIsMacOnly);
+						if (!bGeneratingRunIOSProject)
+						{
+							ProjectTargets.Add(BuildTarget);
 						}
-						else
+
+						if (ProjectFilePlatform.HasFlag(XcodeProjectFilePlatform.iOS) && SupportedPlatforms.Contains(UnrealTargetPlatform.IOS))
 						{
-							XcodeContainerItemProxy ContainerProxy = new XcodeContainerItemProxy(ProjectTarget.Guid, BuildTarget.Guid, BuildTarget.DisplayName);
-							XcodeTargetDependency TargetDependency = new XcodeTargetDependency(BuildTarget.DisplayName, BuildTarget.Guid, ContainerProxy.Guid);
-							XcodeProjectTarget IOSDeployTarget = new XcodeProjectTarget(DisplayName + " - iOS", TargetName, XcodeTargetType.Native, TargetFilePath, TargetName + ".app", UnrealTargetPlatform.IOS, false, new List<XcodeTargetDependency>() { TargetDependency }, true);
-							ProjectTargets.Add(IOSDeployTarget);
-							ContainerItemProxies.Add(ContainerProxy);
-							TargetDependencies.Add(TargetDependency);
+							if ((bGeneratingRocketProjectFiles && TargetName == "UE4Game") || bGeneratingRunIOSProject)
+							{
+								// Generate Framework references.
+								List<XcodeFrameworkRef> FrameworkRefs = new List<XcodeFrameworkRef>();
+								foreach (XcodeFramework Framework in Frameworks)
+								{
+									FrameworkRefs.Add(new XcodeFrameworkRef(Framework));
+								}
+
+								XcodeProjectTarget IOSDeployTarget = new XcodeProjectTarget(DisplayName + " - iOS", TargetName, XcodeTargetType.Native, TargetFilePath, TargetName + ".app", UnrealTargetPlatform.IOS, false, null, true, FrameworkRefs);
+								ProjectTargets.Add(IOSDeployTarget);
+							}
+							else
+							{
+								XcodeContainerItemProxy ContainerProxy = new XcodeContainerItemProxy(ProjectTarget.Guid, BuildTarget.Guid, BuildTarget.DisplayName);
+								XcodeTargetDependency TargetDependency = new XcodeTargetDependency(BuildTarget.DisplayName, BuildTarget.Guid, ContainerProxy.Guid);
+								XcodeProjectTarget IOSDeployTarget = new XcodeProjectTarget(DisplayName + " - iOS", TargetName, XcodeTargetType.Native, TargetFilePath, TargetName + ".app", UnrealTargetPlatform.IOS, false, new List<XcodeTargetDependency>() { TargetDependency }, true);
+								ProjectTargets.Add(IOSDeployTarget);
+								ContainerItemProxies.Add(ContainerProxy);
+								TargetDependencies.Add(TargetDependency);
+							}
 						}
 					}
 				}
@@ -1536,9 +1540,9 @@ namespace UnrealBuildTool
 			XcodeFramework XCTestFramework = new XcodeFramework("XCTest.framework", "Library/Frameworks/XCTest.framework", "DEVELOPER_DIR");
             Frameworks.Add(XCTestFramework);
 
-			var AllTargets = DiscoverTargets();
+//			var AllTargets = DiscoverTargets();
 
-			PopulateTargets(AllTargets, ProjectTargets, ContainerItemProxies, TargetDependencies, UE4ProjectTarget, Frameworks);
+			PopulateTargets(ProjectTargets, ContainerItemProxies, TargetDependencies, UE4ProjectTarget, Frameworks);
 
 			Log.TraceInformation(string.Format ("Found {0} targets!", ProjectTargets.Count));
 
