Changes integrated from CL# 2417785, which was not merged back to Main from UE4-Releases/4.7.  Original changelist description:

> UnrealBuildTool: No longer needs to run UnrealHeaderTool unless you actually changed a source file with UObjects in it
> - This makes it much faster to iterate on source files that UCLASS's or USTRUCT's aren't declared in
> - UBT now keeps track of the set of UObject module files (inside it's per-module Timestamp file)
> - UHT only neesd to run if a UObject file was added/deleted or changed
> - Directory timestamp checking is no longer needed by UBT for generated code
> - Added new diagnostic logging for generated code outdatedness checking (with '-verbose' option)

diff --git a/Engine/Source/Programs/UnrealBuildTool/System/ExternalExecution.cs b/Engine/Source/Programs/UnrealBuildTool/System/ExternalExecution.cs
--- a/Engine/Source/Programs/UnrealBuildTool/System/ExternalExecution.cs
+++ b/Engine/Source/Programs/UnrealBuildTool/System/ExternalExecution.cs
@@ -411,15 +411,20 @@ namespace UnrealBuildTool
 				if( !TestDirectory.Exists )
 				{
 					// Generated code directory is missing entirely!
+					Log.TraceVerbose( "UnrealHeaderTool needs to run because no generated code directory was found for module {0}", Module.ModuleName );
 					return true;
 				}
 
-				// Grab our special "Timestamp" file that we saved after the last set of headers were generated
+				// Grab our special "Timestamp" file that we saved after the last set of headers were generated.  This file
+				// actually contains the list of source files which contained UObjects, so that we can compare to see if any
+				// UObject source files were deleted (or no longer contain UObjects), which means we need to run UHT even
+				// if no other source files were outdated
 				string TimestampFile = Path.Combine( GeneratedCodeDirectory, @"Timestamp" );
 				var SavedTimestampFileInfo = (FileSystemInfo)new FileInfo(TimestampFile);
 				if (!SavedTimestampFileInfo.Exists)
 				{
 					// Timestamp file was missing (possibly deleted/cleaned), so headers are out of date
+					Log.TraceVerbose( "UnrealHeaderTool needs to run because UHT Timestamp file did not exist for module {0}", Module.ModuleName );
 					return true;
 				}
 
@@ -427,7 +432,8 @@ namespace UnrealBuildTool
 				var SavedTimestamp = SavedTimestampFileInfo.LastWriteTime;
 				if( HeaderToolTimestamp > SavedTimestamp || CoreGeneratedTimestamp > SavedTimestamp )
 				{
-					// Generated code is older UnrealHeaderTool.exe or CoreUObject headers.  Out of date!
+					// Generated code is older than UnrealHeaderTool.exe or CoreUObject headers.  Out of date!
+					Log.TraceVerbose( "UnrealHeaderTool needs to run because UnrealHeaderTool.exe or CoreUObject headers are newer than SavedTimestamp for module {0}", Module.ModuleName );
 					return true;
 				}
 
@@ -436,6 +442,25 @@ namespace UnrealBuildTool
 				AllUObjectHeaders.AddRange( Module.PublicUObjectClassesHeaders );
 				AllUObjectHeaders.AddRange( Module.PublicUObjectHeaders );
 				AllUObjectHeaders.AddRange( Module.PrivateUObjectHeaders );
+
+				// Load up the old timestamp file and check to see if anything has changed
+				{
+					var UObjectFilesFromPreviousRun = File.ReadAllLines( TimestampFile );
+					if( AllUObjectHeaders.Count != UObjectFilesFromPreviousRun.Length )
+					{
+						Log.TraceVerbose( "UnrealHeaderTool needs to run because there are a different number of UObject source files in module {0}", Module.ModuleName );
+						return true;
+					}
+					for( int FileIndex = 0; FileIndex < AllUObjectHeaders.Count; ++FileIndex )
+					{
+						if( !UObjectFilesFromPreviousRun[ FileIndex ].Equals( AllUObjectHeaders[ FileIndex ].AbsolutePath, StringComparison.InvariantCultureIgnoreCase ) )
+						{
+							Log.TraceVerbose( "UnrealHeaderTool needs to run because the set of UObject source files in module {0} has changed", Module.ModuleName );
+							return true;
+						}
+					}
+				}
+
 				foreach( var HeaderFile in AllUObjectHeaders )
 				{
 					var HeaderFileTimestamp = HeaderFile.Info.LastWriteTime;
@@ -443,6 +468,7 @@ namespace UnrealBuildTool
 					// Has the source header changed since we last generated headers successfully?
 					if( HeaderFileTimestamp > SavedTimestamp )
 					{
+						Log.TraceVerbose( "UnrealHeaderTool needs to run because SavedTimestamp is older than HeaderFileTimestamp ({0}) for module {1}", HeaderFile.AbsolutePath, Module.ModuleName );
 						return true;
 					}
 
@@ -461,6 +487,7 @@ namespace UnrealBuildTool
 						var HeaderDirectoryTimestamp = new DirectoryInfo( Path.GetDirectoryName( HeaderFile.AbsolutePath ) ).LastWriteTime;
 						if( HeaderDirectoryTimestamp > SavedTimestamp )
 						{
+							Log.TraceVerbose( "UnrealHeaderTool needs to run because the directory containing an existing header ({0}) has changed, and headers may have been added to or deleted from module {1}", HeaderFile.AbsolutePath, Module.ModuleName );
 							return true;
 						}
 					}
@@ -499,12 +526,15 @@ namespace UnrealBuildTool
 							GeneratedCodeDirectoryInfo.Create();
 						}
 
-						if(File.Exists(TimestampFile))
-						{
-							File.Delete(TimestampFile);
-						}
-						using (File.Create(TimestampFile))
+						// Save all of the UObject files to a timestamp file.  We'll load these on the next run to see if any new
+						// files with UObject classes were deleted, so that we'll know to run UHT even if the timestamps of all
+						// of the other source files were unchanged
 						{
+							var AllUObjectFiles = new List<string>();
+							AllUObjectFiles.AddRange( Module.PublicUObjectClassesHeaders.ConvertAll( Item => Item.AbsolutePath ) );
+							AllUObjectFiles.AddRange( Module.PublicUObjectHeaders.ConvertAll( Item => Item.AbsolutePath ) );
+							AllUObjectFiles.AddRange( Module.PrivateUObjectHeaders.ConvertAll( Item => Item.AbsolutePath ) );
+							ResponseFile.Create( TimestampFile, AllUObjectFiles );
 						}
 					}
 				}
