Update to support Nsight 3.1
#ue4
#android

diff --git a/Engine/Source/Programs/UnrealBuildTool/Android/AndroidProjectGenerator.cs b/Engine/Source/Programs/UnrealBuildTool/Android/AndroidProjectGenerator.cs
--- a/Engine/Source/Programs/UnrealBuildTool/Android/AndroidProjectGenerator.cs
+++ b/Engine/Source/Programs/UnrealBuildTool/Android/AndroidProjectGenerator.cs
@@ -54,11 +54,23 @@ namespace UnrealBuildTool
 			// grab the version info from the DLL
 			FileVersionInfo NsightVersion = FileVersionInfo.GetVersionInfo(NsightDllPath);
 
-			if (NsightVersion.ProductMajorPart >= 3)
+			if (NsightVersion.ProductMajorPart > 3)
 			{
-				// Nsight 3.0+ should be valid
+				// Mark as Nsight 3.1 (project will be updated)
+				NsightVersionCode = 10;
+				NsightInstalled = true;
+			}
+			else if (NsightVersion.ProductMajorPart == 3)
+			{
+				// Nsight 3.0 supported
 				NsightVersionCode = 9;
 				NsightInstalled = true;
+
+				if (NsightVersion.ProductMinorPart >= 1)
+				{
+					// Nsight 3.1+ should be valid (will update project if newer)
+					NsightVersionCode = 10;
+				}
 			}
 			else if (NsightVersion.ProductMajorPart == 2)
 			{
