#html5 correct sanity check for the emscripten config file, always recreate if not created by UE4 build system.

diff --git a/Engine/Source/Programs/UnrealBuildTool/HTML5/HTML5SDKInfo.cs b/Engine/Source/Programs/UnrealBuildTool/HTML5/HTML5SDKInfo.cs
--- a/Engine/Source/Programs/UnrealBuildTool/HTML5/HTML5SDKInfo.cs
+++ b/Engine/Source/Programs/UnrealBuildTool/HTML5/HTML5SDKInfo.cs
@@ -89,7 +89,7 @@ namespace UnrealBuildTool
 			string UserFolder = Environment.GetFolderPath(Environment.SpecialFolder.UserProfile); 
 			string ConfigFile = Path.Combine(UserFolder, ".emscripten");
 
-			if (!File.Exists(ConfigFile) || !File.ReadAllText(ConfigFile).Contains(SDKVersion))
+			if (!File.Exists(ConfigFile) || !File.ReadAllText(ConfigFile).Contains("GENERATEDBYUE4"))
 			{
 				var ConfigString = String.Join(
 								Environment.NewLine,
