Override '.emscripten' config file location for UBT builds.

diff --git a/Engine/Source/Programs/UnrealBuildTool/HTML5/HTML5SDKInfo.cs b/Engine/Source/Programs/UnrealBuildTool/HTML5/HTML5SDKInfo.cs
--- a/Engine/Source/Programs/UnrealBuildTool/HTML5/HTML5SDKInfo.cs
+++ b/Engine/Source/Programs/UnrealBuildTool/HTML5/HTML5SDKInfo.cs
@@ -73,6 +73,19 @@ namespace UnrealBuildTool
 			}
 		}
 
+		static public string DOT_EMSCRIPTEN
+		{
+			get 
+			{
+				var TempPath = Path.GetFullPath(Path.Combine(BuildConfiguration.RelativeEnginePath, BuildConfiguration.PlatformIntermediateFolder));
+				if (!Directory.Exists(TempPath))
+				{
+					Directory.CreateDirectory(TempPath);
+				}
+				return Path.Combine(TempPath, ".emscripten");
+			}
+		}
+
 		public static string  SetupEmscriptenTemp()
 		{
 			string HTML5Intermediatory = Path.GetFullPath(Path.Combine(BuildConfiguration.RelativeEnginePath, BuildConfiguration.PlatformIntermediateFolder));
@@ -86,8 +99,8 @@ namespace UnrealBuildTool
 
 		public static string SetUpEmscriptenConfigFile()
 		{
-			string UserFolder = Environment.GetFolderPath(Environment.SpecialFolder.UserProfile); 
-			string ConfigFile = Path.Combine(UserFolder, ".emscripten");
+			string UserFolder = Environment.GetFolderPath(Environment.SpecialFolder.UserProfile);
+			string ConfigFile = DOT_EMSCRIPTEN;
 
 			if (!File.Exists(ConfigFile) || !File.ReadAllText(ConfigFile).Contains("GENERATEDBYUE4"))
 			{
diff --git a/Engine/Source/Programs/UnrealBuildTool/HTML5/HTML5ToolChain.cs b/Engine/Source/Programs/UnrealBuildTool/HTML5/HTML5ToolChain.cs
--- a/Engine/Source/Programs/UnrealBuildTool/HTML5/HTML5ToolChain.cs
+++ b/Engine/Source/Programs/UnrealBuildTool/HTML5/HTML5ToolChain.cs
@@ -19,19 +19,22 @@ namespace UnrealBuildTool
 		{
 			if (HTML5SDKInfo.IsSDKInstalled())
 			{
-					// set some environment variable we'll need
-					//Environment.SetEnvironmentVariable("EMCC_DEBUG", "cache");
-					Environment.SetEnvironmentVariable("EMCC_CORES", "8");
-					Environment.SetEnvironmentVariable("EMCC_OPTIMIZE_NORMALLY", "1");
-					// finally register the toolchain that is now ready to go
-					Log.TraceVerbose("        Registered for {0}", CPPTargetPlatform.HTML5.ToString());
-					UEToolChain.RegisterPlatformToolChain(CPPTargetPlatform.HTML5, this);
+				// set some environment variable we'll need
+				//Environment.SetEnvironmentVariable("EMCC_DEBUG", "cache");
+				Environment.SetEnvironmentVariable("EMCC_CORES", "8");
+				Environment.SetEnvironmentVariable("EMCC_OPTIMIZE_NORMALLY", "1");
+				// finally register the toolchain that is now ready to go
+				Log.TraceVerbose("        Registered for {0}", CPPTargetPlatform.HTML5.ToString());
+				UEToolChain.RegisterPlatformToolChain(CPPTargetPlatform.HTML5, this);
 			}
 		}
 		public override void PreBuildSync()
 		{
 			HTML5SDKInfo.SetupEmscriptenTemp();
 			HTML5SDKInfo.SetUpEmscriptenConfigFile();
+			// set some environment variable we'll need.
+			// Forces emcc to use our generated .emscripten config, not the one in the users home directory.
+			Environment.SetEnvironmentVariable("EM_CONFIG", HTML5SDKInfo.DOT_EMSCRIPTEN);
 		}
 	
 		static string GetSharedArguments_Global(CPPTargetConfiguration TargetConfiguration, string Architecture, bool bEnableShadowVariableWarning)
