Summary: running UAT from VS is simpler and faster.
UEB-261 - Ensure that compiling AutomationTool in VS will compile all other Automation Projects
* Just set AutomationTool as your startup project and pass the command to execute.
* VS will build the script modules at build time, instead of every time at runtime.
* To make this happen, "UBT.exe -ProjectFiles" now generates a companion AutomationTool.csproj.References that make AutomationTool depend on all Automation modules.
* AutomationTool.exe defaults to not building script modules at runtime. Pass -compile if you want to dynamically build them.
* Without the .references file, AutomationTool will only build itself and you will need to pass -compile.
* RunUAT.bat still works that same, defaulting to runtime compilation and supporting -nocompile flag. It then passes -compile (or nothing) to AutomationTool.

Other
* All Automation projects target .Net 4.5. Some already were and had hard dependencies on them (Rocket and SyncGithub -> Octokit). Now that AutomationTool directly depends on them, everything had to use .Net 4.5.
* Decoupled logic for -NoCompile and -NoCompileEditor. The flags are still confusing, but -NoCompile is no longer linked to -NoCompileEditor.
* Had to leave in stub support in UAT for -NoCompile else RunUAT.bat passes it along and UAT complains that it doesn't understand it.
* Added a CommandUtils.Run option to support run command, but still output the run duration.
* Reduced the verbosity when UAT.proj is run from dozens of lines per module to a single Module -> Output line. It was looking like there were problems, but it was just msbuild spew.

diff --git a/Engine/Build/BatchFiles/RunUAT.bat b/Engine/Build/BatchFiles/RunUAT.bat
--- a/Engine/Build/BatchFiles/RunUAT.bat
+++ b/Engine/Build/BatchFiles/RunUAT.bat
@@ -11,7 +11,7 @@ rem ## if you copy it to a different location and run it.
 
 set UATExecutable=AutomationToolLauncher.exe
 set UATDirectory=Binaries\DotNET\
-set UATNoCompileArg=
+set UATCompileArg=-compile
 
 rem ## Change the CWD to /Engine. 
 pushd %~dp0..\..\
@@ -59,7 +59,7 @@ rem ## ok, well it doesn't look like visual studio is installed, let's try runni
 
 if not exist Binaries\DotNET\AutomationTool.exe goto Error_NoFallbackExecutable
 
-set UATNoCompileArg=-NoCompile
+set UATCompileArg=
 
 if not exist Binaries\DotNET\AutomationToolLauncher.exe set UATExecutable=AutomationTool.exe
 goto DoRunUAT
@@ -75,7 +75,7 @@ if not %ERRORLEVEL% == 0 goto Error_UATCompileFailed
 rem ## Run AutomationTool
 :DoRunUAT
 pushd %UATDirectory%
-%UATExecutable% %* %UATNoCompileArg%
+%UATExecutable% %* %UATCompileArg%
 if not %ERRORLEVEL% == 0 goto Error_UATFailed
 popd
 
diff --git a/Engine/Build/BatchFiles/RunUAT.sh b/Engine/Build/BatchFiles/RunUAT.sh
--- a/Engine/Build/BatchFiles/RunUAT.sh
+++ b/Engine/Build/BatchFiles/RunUAT.sh
@@ -31,7 +31,7 @@ SCRIPT_DIR=$(cd "`dirname "$0"`" && pwd)
 cd "$SCRIPT_DIR/../.."
 
 UATDirectory=Binaries/DotNET/
-UATNoCompileArg=
+UATCompileArg=-compile
 
 if [ ! -f Build/BatchFiles/RunUAT.sh ]; then
 	echo "RunUAT ERROR: The script does not appear to be located in the "
@@ -41,9 +41,9 @@ fi
 
 # see if we have the no compile arg
 if echo "${Args[@]}" | grep -q -i "\-nocompile"; then
-	UATNoCompileArg=-NoCompile
+	UATCompileArg=
 else
-	UATNoCompileArg=
+	UATCompileArg=-compile
 fi
 
 if [ "$(uname)" = "Darwin" ]; then
@@ -51,11 +51,11 @@ if [ "$(uname)" = "Darwin" ]; then
 	source "`dirname "$0"`/Mac/SetupMono.sh" "`dirname "$0"`/Mac"
 fi
 
-if [ "$UATNoCompileArg" != "-NoCompile" ]; then
+if [ "$UATCompileArg" = "-compile" ]; then
   # see if the .csproj exists to be compiled
 	if [ ! -f Source/Programs/AutomationTool/AutomationTool.csproj ]; then
 		echo No project to compile, attempting to use precompiled AutomationTool
-		UATNoCompileArg=-NoCompile
+		UATCompileArg=
 	else
 		if [ -f "$UATDirectory/AutomationTool.exe" ]; then
 			echo AutomationTool exists: Deleting
@@ -94,7 +94,7 @@ fi
 # you can't set a dotted env var nicely in sh, but env will run a command with
 # a list of env vars set, including dotted ones
 echo Start UAT: mono AutomationTool.exe "${Args[@]}"
-env uebp_LogFolder="$LogDir" mono AutomationTool.exe "${Args[@]}" $UATNoCompileArg
+env uebp_LogFolder="$LogDir" mono AutomationTool.exe "${Args[@]}" $UATCompileArg
 UATReturn=$?
 
 # @todo: Copy log files to somewhere useful
diff --git a/Engine/Source/Editor/MainFrame/Private/Frame/MainFrameActions.cpp b/Engine/Source/Editor/MainFrame/Private/Frame/MainFrameActions.cpp
--- a/Engine/Source/Editor/MainFrame/Private/Frame/MainFrameActions.cpp
+++ b/Engine/Source/Editor/MainFrame/Private/Frame/MainFrameActions.cpp
@@ -286,8 +286,10 @@ void FMainFrameActionCallbacks::AddCodeToProject()
  */
 const TCHAR* GetUATCompilationFlags()
 {
+	// We never want to compile editor targets when invoking UAT in this context.
+	// If we are rocket or don't have a compiler, we must assume we have a precompiled UAT.
 	return FRocketSupport::IsRocket() || !FSourceCodeNavigation::IsCompilerAvailable()
-		? TEXT("-nocompile")
+		? TEXT("-nocompile -nocompileeditor")
 		: TEXT("-nocompileeditor");
 }
 
diff --git a/Engine/Source/Programs/AutomationTool/AllDesktop/AllDesktop.Automation.csproj b/Engine/Source/Programs/AutomationTool/AllDesktop/AllDesktop.Automation.csproj
--- a/Engine/Source/Programs/AutomationTool/AllDesktop/AllDesktop.Automation.csproj
+++ b/Engine/Source/Programs/AutomationTool/AllDesktop/AllDesktop.Automation.csproj
@@ -9,7 +9,7 @@
     <AppDesignerFolder>Properties</AppDesignerFolder>
     <RootNamespace>AllDesktop.Automation</RootNamespace>
     <AssemblyName>AllDesktop.Automation</AssemblyName>
-    <TargetFrameworkVersion>v4.0</TargetFrameworkVersion>
+    <TargetFrameworkVersion>v4.5</TargetFrameworkVersion>
     <FileAlignment>512</FileAlignment>
   </PropertyGroup>
   <PropertyGroup Condition=" '$(Configuration)|$(Platform)' == 'Debug|AnyCPU' ">
diff --git a/Engine/Source/Programs/AutomationTool/Android/Android.Automation.csproj b/Engine/Source/Programs/AutomationTool/Android/Android.Automation.csproj
--- a/Engine/Source/Programs/AutomationTool/Android/Android.Automation.csproj
+++ b/Engine/Source/Programs/AutomationTool/Android/Android.Automation.csproj
@@ -9,7 +9,7 @@
     <AppDesignerFolder>Properties</AppDesignerFolder>
     <RootNamespace>Android.Automation</RootNamespace>
     <AssemblyName>Android.Automation</AssemblyName>
-    <TargetFrameworkVersion>v4.0</TargetFrameworkVersion>
+    <TargetFrameworkVersion>v4.5</TargetFrameworkVersion>
     <FileAlignment>512</FileAlignment>
   </PropertyGroup>
   <PropertyGroup Condition=" '$(Configuration)|$(Platform)' == 'Debug|AnyCPU' ">
diff --git a/Engine/Source/Programs/AutomationTool/AutomationTool.csproj b/Engine/Source/Programs/AutomationTool/AutomationTool.csproj
--- a/Engine/Source/Programs/AutomationTool/AutomationTool.csproj
+++ b/Engine/Source/Programs/AutomationTool/AutomationTool.csproj
@@ -10,7 +10,7 @@
     <AppDesignerFolder>Properties</AppDesignerFolder>
     <RootNamespace>AutomationTool</RootNamespace>
     <AssemblyName>AutomationTool</AssemblyName>
-    <TargetFrameworkVersion>v4.0</TargetFrameworkVersion>
+    <TargetFrameworkVersion>v4.5</TargetFrameworkVersion>
     <TargetFrameworkProfile>
     </TargetFrameworkProfile>
     <FileAlignment>512</FileAlignment>
@@ -78,9 +78,11 @@
   <ItemGroup>
     <Reference Include="Ionic.Zip.Reduced">
       <HintPath>..\..\..\Binaries\DotNET\Ionic.Zip.Reduced.dll</HintPath>
+      <Private>False</Private>
     </Reference>
     <Reference Include="RPCUtility">
       <HintPath>..\..\..\Binaries\DotNET\RPCUtility.exe</HintPath>
+      <Private>False</Private>
     </Reference>
     <Reference Include="System" />
     <Reference Include="System.Core" />
@@ -104,6 +106,7 @@
     <ProjectReference Include="..\DotNETCommon\DotNETUtilities\DotNETUtilities.csproj">
       <Project>{5D7D66E8-8C76-4AF9-B3EC-2EF03421D730}</Project>
       <Name>DotNETUtilities</Name>
+      <Private>False</Private>
     </ProjectReference>
     <ProjectReference Include="..\UnrealBuildTool\UnrealBuildTool.csproj">
       <Project>{fd7c5e1a-cfe4-4fd5-a525-1eb1599a39ac}</Project>
@@ -123,6 +126,7 @@
     <WCFMetadata Include="Service References\" />
   </ItemGroup>
   <Import Project="$(MSBuildToolsPath)\Microsoft.CSharp.targets" />
+  <Import Project="AutomationTool.csproj.References" Condition="Exists('AutomationTool.csproj.References')" />
   <!-- To modify your build process, add your task inside one of the targets below and uncomment it. 
        Other similar extension points exist, see Microsoft.Common.targets.
   <Target Name="BeforeBuild">
diff --git a/Engine/Source/Programs/AutomationTool/AutomationUtils/Automation.cs b/Engine/Source/Programs/AutomationTool/AutomationUtils/Automation.cs
--- a/Engine/Source/Programs/AutomationTool/AutomationUtils/Automation.cs
+++ b/Engine/Source/Programs/AutomationTool/AutomationUtils/Automation.cs
@@ -103,9 +103,16 @@ namespace AutomationTool
 		public static CommandLineArg ForceLocal = new CommandLineArg("-ForceLocal");
 		public static CommandLineArg NoP4 = new CommandLineArg("-NoP4");
 		public static CommandLineArg P4 = new CommandLineArg("-P4");
-		public static CommandLineArg Preprocess = new CommandLineArg("-Preprocess");
-		public static CommandLineArg NoCompile = new CommandLineArg("-NoCompile");
-		public static CommandLineArg NoCompileEditor = new CommandLineArg("-NoCompileEditor");
+        public static CommandLineArg Preprocess = new CommandLineArg("-Preprocess");
+        public static CommandLineArg Compile = new CommandLineArg("-Compile");
+        /// <summary>
+        /// This command is LEGACY because we used to run UAT.exe to compile scripts by default.
+        /// Now we only compile by default when run via RunUAT.bat, which still understands -nocompile.
+        /// However, the batch file simply passes on all arguments, so UAT will choke when encountering -nocompile.
+        /// Keep this CommandLineArg around so that doesn't happen.
+        /// </summary>
+        public static CommandLineArg NoCompileLegacyDontUse = new CommandLineArg("-NoCompile");
+        public static CommandLineArg NoCompileEditor = new CommandLineArg("-NoCompileEditor");
 		public static CommandLineArg Help = new CommandLineArg("-Help");
 		public static CommandLineArg List = new CommandLineArg("-List");
 		public static CommandLineArg Rocket = new CommandLineArg("-Rocket");
@@ -134,7 +141,8 @@ AutomationTool.exe [-verbose] [-compileonly] [-p4] Command0 [-Arg0 -Arg1 -Arg2 
 	[Help("nop4", "Disables Perforce functionality (default if not run on a build machine)")]
 	[Help("p4", "Enables Perforce functionality (default if run on a build machine)")]
 	[Help("compileonly", "Does not run any commands, only compiles them")]
-	[Help("forcelocal", "Forces local execution")]
+    [Help("compile", "Dynamically compiles all commands (otherwise assumes they are already built)")]
+    [Help("forcelocal", "Forces local execution")]
 	[Help("help", "Displays help")]
 	[Help("list", "Lists all available commands")]
 	[Help("submit", "Allows UAT command to submit changes")]
diff --git a/Engine/Source/Programs/AutomationTool/AutomationUtils/AutomationUtils.Automation.csproj b/Engine/Source/Programs/AutomationTool/AutomationUtils/AutomationUtils.Automation.csproj
--- a/Engine/Source/Programs/AutomationTool/AutomationUtils/AutomationUtils.Automation.csproj
+++ b/Engine/Source/Programs/AutomationTool/AutomationUtils/AutomationUtils.Automation.csproj
@@ -1,5 +1,5 @@
 ﻿<?xml version="1.0" encoding="utf-8"?>
-<Project ToolsVersion="11.0" DefaultTargets="Build" xmlns="http://schemas.microsoft.com/developer/msbuild/2003">
+<Project ToolsVersion="4.0" DefaultTargets="Build" xmlns="http://schemas.microsoft.com/developer/msbuild/2003">
   <Import Project="$(MSBuildExtensionsPath)\$(MSBuildToolsVersion)\Microsoft.Common.props" Condition="Exists('$(MSBuildExtensionsPath)\$(MSBuildToolsVersion)\Microsoft.Common.props')" />
   <PropertyGroup>
     <Configuration Condition=" '$(Configuration)' == '' ">Debug</Configuration>
@@ -9,7 +9,7 @@
     <AppDesignerFolder>Properties</AppDesignerFolder>
     <RootNamespace>AutomationUtils.Automation</RootNamespace>
     <AssemblyName>AutomationUtils.Automation</AssemblyName>
-    <TargetFrameworkVersion>v4.0</TargetFrameworkVersion>
+    <TargetFrameworkVersion>v4.5</TargetFrameworkVersion>
     <FileAlignment>512</FileAlignment>
   </PropertyGroup>
   <PropertyGroup Condition=" '$(Configuration)|$(Platform)' == 'Debug|AnyCPU' ">
diff --git a/Engine/Source/Programs/AutomationTool/AutomationUtils/ProcessUtils.cs b/Engine/Source/Programs/AutomationTool/AutomationUtils/ProcessUtils.cs
--- a/Engine/Source/Programs/AutomationTool/AutomationUtils/ProcessUtils.cs
+++ b/Engine/Source/Programs/AutomationTool/AutomationUtils/ProcessUtils.cs
@@ -650,6 +650,10 @@ namespace AutomationTool
 
 			/// When specified with AllowSpew, the output will be TraceEventType.Verbose instead of TraceEventType.Information
 			SpewIsVerbose = 1 << 6,
+            /// <summary>
+            /// If NoLoggingOfRunCommand is set, it normally suppresses the run duration output. This turns it back on.
+            /// </summary>
+            LoggingOfRunDuration = 1 << 7,
 
 			Default = AllowSpew | AppMustExist,
 		}
@@ -723,7 +727,7 @@ namespace AutomationTool
 				var BuildDuration = (DateTime.UtcNow - StartTime).TotalMilliseconds;
 				AddRunTime(App, (int)(BuildDuration));
 				Result.ExitCode = Proc.ExitCode;
-				if (!Options.HasFlag(ERunOptions.NoLoggingOfRunCommand))
+				if (!Options.HasFlag(ERunOptions.NoLoggingOfRunCommand) || Options.HasFlag(ERunOptions.LoggingOfRunDuration))
                 {
                     Log(SpewVerbosity,"Run: Took {0}s to run {1}, ExitCode={2}", BuildDuration / 1000, Path.GetFileName(App), Result.ExitCode);
                 }
@@ -884,8 +888,6 @@ namespace AutomationTool
 			Log("Recursive UAT Run, in log folder {0}, main log file {1}", LogSubdir, LogFile);
 			CreateDirectory(LogSubdir);
 
-			CommandLine = CommandLine + " -NoCompile";
-
 			string App = CmdEnv.UATExe;
 
 			Log("Running {0} {1}", App, CommandLine);
diff --git a/Engine/Source/Programs/AutomationTool/AutomationUtils/ProjectParams.cs b/Engine/Source/Programs/AutomationTool/AutomationUtils/ProjectParams.cs
--- a/Engine/Source/Programs/AutomationTool/AutomationUtils/ProjectParams.cs
+++ b/Engine/Source/Programs/AutomationTool/AutomationUtils/ProjectParams.cs
@@ -1577,7 +1577,7 @@ namespace AutomationTool
 
 			if (EditorTargetsList == null)
 			{
-				if (!GlobalCommandLine.NoCompile && !GlobalCommandLine.NoCompileEditor && (ProjectType != TargetRules.TargetType.Program) && !String.IsNullOrEmpty(EditorTarget))
+				if (!GlobalCommandLine.NoCompileEditor && (ProjectType != TargetRules.TargetType.Program) && !String.IsNullOrEmpty(EditorTarget))
 				{
 					EditorTargetsList = new ParamList<string>(EditorTarget);
 				}
diff --git a/Engine/Source/Programs/AutomationTool/AutomationUtils/ScriptCompiler.cs b/Engine/Source/Programs/AutomationTool/AutomationUtils/ScriptCompiler.cs
--- a/Engine/Source/Programs/AutomationTool/AutomationUtils/ScriptCompiler.cs
+++ b/Engine/Source/Programs/AutomationTool/AutomationUtils/ScriptCompiler.cs
@@ -56,10 +56,10 @@ namespace AutomationTool
 		/// <param name="AdditionalScriptsFolders">Additional script fodlers to look for source files in.</param>
 		public void FindAndCompileAllScripts(List<string> AdditionalScriptsFolders)
 		{
-			bool DoCompile = true;
-			if (GlobalCommandLine.NoCompile)
+			bool DoCompile = false;
+			if (GlobalCommandLine.Compile)
 			{
-				DoCompile = false;
+				DoCompile = true;
 			}
 
 			// Change to Engine\Source (if exists) to properly discover all UBT classes
@@ -160,8 +160,10 @@ namespace AutomationTool
 			{
 				string Modules = string.Join(";", ModulesToCompile.ToArray());
 				var UATProj = CommandUtils.CombinePaths(CommandUtils.CmdEnv.LocalRoot, @"Engine\Source\Programs\AutomationTool\Scripts\UAT.proj");
-				var CmdLine = String.Format("\"{0}\" /p:Modules=\"{1}\" /p:Configuration={2}", UATProj, Modules, BuildConfig);
-				var Result = CommandUtils.Run(CommandUtils.CmdEnv.MsBuildExe, CmdLine);
+				var CmdLine = String.Format("\"{0}\" /p:Modules=\"{1}\" /p:Configuration={2} /verbosity:minimal", UATProj, Modules, BuildConfig);
+                Log.TraceInformation("Building Automation projects in parallel...");
+                // supress the run command because it can be long and intimidating, making the logs around this code harder to read.
+                var Result = CommandUtils.Run(CommandUtils.CmdEnv.MsBuildExe, CmdLine, Options: CommandUtils.ERunOptions.Default | CommandUtils.ERunOptions.NoLoggingOfRunCommand | CommandUtils.ERunOptions.LoggingOfRunDuration);
 				if (Result.ExitCode != 0)
 				{
 					throw new AutomationException(String.Format("Failed to build \"{0}\":{1}{2}", UATProj, Environment.NewLine, Result.Output));
diff --git a/Engine/Source/Programs/AutomationTool/AutomationUtils/UE4Build.cs b/Engine/Source/Programs/AutomationTool/AutomationUtils/UE4Build.cs
--- a/Engine/Source/Programs/AutomationTool/AutomationUtils/UE4Build.cs
+++ b/Engine/Source/Programs/AutomationTool/AutomationUtils/UE4Build.cs
@@ -125,7 +125,7 @@ namespace AutomationTool
 			// Don't build UBT if we're running with pre-compiled binaries and if there's a debugger attached to this process.
 			// With the debugger attached, even though deleting the exe will work, the pdb files are still locked and the build will fail.
 			// Also, if we're running from VS then since UAT references UBT, we already have the most up-to-date version of UBT.exe
-			if (!bIsUBTReady && !GlobalCommandLine.NoCompile && !System.Diagnostics.Debugger.IsAttached)
+			if (!bIsUBTReady && GlobalCommandLine.Compile && !System.Diagnostics.Debugger.IsAttached)
 			{
 				DeleteFile(UBTExecutable);
 
@@ -1545,7 +1545,7 @@ namespace AutomationTool
 		/// </summary>
 		public void AddUBTFilesToBuildProducts()
 		{
-			if (GlobalCommandLine.NoCompile)
+			if (!GlobalCommandLine.Compile)
 			{
 				Log("We are being asked to copy the UBT build products, but we are running precompiled, so this does not make much sense.");
 			}
diff --git a/Engine/Source/Programs/AutomationTool/GUBP/GUBP.Automation.csproj b/Engine/Source/Programs/AutomationTool/GUBP/GUBP.Automation.csproj
--- a/Engine/Source/Programs/AutomationTool/GUBP/GUBP.Automation.csproj
+++ b/Engine/Source/Programs/AutomationTool/GUBP/GUBP.Automation.csproj
@@ -9,7 +9,7 @@
     <AppDesignerFolder>Properties</AppDesignerFolder>
     <RootNamespace>AutomationTool</RootNamespace>
     <AssemblyName>GUBP.Automation</AssemblyName>
-    <TargetFrameworkVersion>v4.0</TargetFrameworkVersion>
+    <TargetFrameworkVersion>v4.5</TargetFrameworkVersion>
     <FileAlignment>512</FileAlignment>
   </PropertyGroup>
   <PropertyGroup Condition=" '$(Configuration)|$(Platform)' == 'Debug|AnyCPU' ">
diff --git a/Engine/Source/Programs/AutomationTool/HTML5/HTML5.Automation.csproj b/Engine/Source/Programs/AutomationTool/HTML5/HTML5.Automation.csproj
--- a/Engine/Source/Programs/AutomationTool/HTML5/HTML5.Automation.csproj
+++ b/Engine/Source/Programs/AutomationTool/HTML5/HTML5.Automation.csproj
@@ -9,7 +9,7 @@
     <AppDesignerFolder>Properties</AppDesignerFolder>
     <RootNamespace>HTML5.Automation</RootNamespace>
     <AssemblyName>HTML5.Automation</AssemblyName>
-    <TargetFrameworkVersion>v4.0</TargetFrameworkVersion>
+    <TargetFrameworkVersion>v4.5</TargetFrameworkVersion>
     <FileAlignment>512</FileAlignment>
   </PropertyGroup>
   <PropertyGroup Condition=" '$(Configuration)|$(Platform)' == 'Debug|AnyCPU' ">
diff --git a/Engine/Source/Programs/AutomationTool/IOS/IOS.Automation.csproj b/Engine/Source/Programs/AutomationTool/IOS/IOS.Automation.csproj
--- a/Engine/Source/Programs/AutomationTool/IOS/IOS.Automation.csproj
+++ b/Engine/Source/Programs/AutomationTool/IOS/IOS.Automation.csproj
@@ -9,7 +9,7 @@
     <AppDesignerFolder>Properties</AppDesignerFolder>
     <RootNamespace>IOS.Automation</RootNamespace>
     <AssemblyName>IOS.Automation</AssemblyName>
-    <TargetFrameworkVersion>v4.0</TargetFrameworkVersion>
+    <TargetFrameworkVersion>v4.5</TargetFrameworkVersion>
     <FileAlignment>512</FileAlignment>
   </PropertyGroup>
   <PropertyGroup Condition=" '$(Configuration)|$(Platform)' == 'Debug|AnyCPU' ">
diff --git a/Engine/Source/Programs/AutomationTool/Linux/Linux.Automation.csproj b/Engine/Source/Programs/AutomationTool/Linux/Linux.Automation.csproj
--- a/Engine/Source/Programs/AutomationTool/Linux/Linux.Automation.csproj
+++ b/Engine/Source/Programs/AutomationTool/Linux/Linux.Automation.csproj
@@ -9,7 +9,7 @@
     <AppDesignerFolder>Properties</AppDesignerFolder>
     <RootNamespace>Linux.Automation</RootNamespace>
     <AssemblyName>Linux.Automation</AssemblyName>
-    <TargetFrameworkVersion>v4.0</TargetFrameworkVersion>
+    <TargetFrameworkVersion>v4.5</TargetFrameworkVersion>
     <FileAlignment>512</FileAlignment>
   </PropertyGroup>
   <PropertyGroup Condition=" '$(Configuration)|$(Platform)' == 'Debug|AnyCPU' ">
diff --git a/Engine/Source/Programs/AutomationTool/Mac/Mac.Automation.csproj b/Engine/Source/Programs/AutomationTool/Mac/Mac.Automation.csproj
--- a/Engine/Source/Programs/AutomationTool/Mac/Mac.Automation.csproj
+++ b/Engine/Source/Programs/AutomationTool/Mac/Mac.Automation.csproj
@@ -9,7 +9,7 @@
     <AppDesignerFolder>Properties</AppDesignerFolder>
     <RootNamespace>Mac.Automation</RootNamespace>
     <AssemblyName>Mac.Automation</AssemblyName>
-    <TargetFrameworkVersion>v4.0</TargetFrameworkVersion>
+    <TargetFrameworkVersion>v4.5</TargetFrameworkVersion>
     <FileAlignment>512</FileAlignment>
   </PropertyGroup>
   <PropertyGroup Condition=" '$(Configuration)|$(Platform)' == 'Debug|AnyCPU' ">
diff --git a/Engine/Source/Programs/AutomationTool/Scripts/AutomationScripts.Automation.csproj b/Engine/Source/Programs/AutomationTool/Scripts/AutomationScripts.Automation.csproj
--- a/Engine/Source/Programs/AutomationTool/Scripts/AutomationScripts.Automation.csproj
+++ b/Engine/Source/Programs/AutomationTool/Scripts/AutomationScripts.Automation.csproj
@@ -9,7 +9,7 @@
     <AppDesignerFolder>Properties</AppDesignerFolder>
     <RootNamespace>AutomationScripts.Automation</RootNamespace>
     <AssemblyName>AutomationScripts.Automation</AssemblyName>
-    <TargetFrameworkVersion>v4.0</TargetFrameworkVersion>
+    <TargetFrameworkVersion>v4.5</TargetFrameworkVersion>
     <FileAlignment>512</FileAlignment>
   </PropertyGroup>
   <PropertyGroup Condition=" '$(Configuration)|$(Platform)' == 'Debug|AnyCPU' ">
diff --git a/Engine/Source/Programs/AutomationTool/Win/Win.Automation.csproj b/Engine/Source/Programs/AutomationTool/Win/Win.Automation.csproj
--- a/Engine/Source/Programs/AutomationTool/Win/Win.Automation.csproj
+++ b/Engine/Source/Programs/AutomationTool/Win/Win.Automation.csproj
@@ -9,7 +9,7 @@
     <AppDesignerFolder>Properties</AppDesignerFolder>
     <RootNamespace>Win.Automation</RootNamespace>
     <AssemblyName>Win.Automation</AssemblyName>
-    <TargetFrameworkVersion>v4.0</TargetFrameworkVersion>
+    <TargetFrameworkVersion>v4.5</TargetFrameworkVersion>
     <FileAlignment>512</FileAlignment>
   </PropertyGroup>
   <PropertyGroup Condition=" '$(Configuration)|$(Platform)' == 'Debug|AnyCPU' ">
diff --git a/Engine/Source/Programs/AutomationTool/app.config b/Engine/Source/Programs/AutomationTool/app.config
--- a/Engine/Source/Programs/AutomationTool/app.config
+++ b/Engine/Source/Programs/AutomationTool/app.config
@@ -1,3 +1,3 @@
 <?xml version="1.0"?>
 <configuration>
-<startup><supportedRuntime version="v4.0" sku=".NETFramework,Version=v4.0"/></startup></configuration>
+<startup><supportedRuntime version="v4.0" sku=".NETFramework,Version=v4.5"/></startup></configuration>
diff --git a/Engine/Source/Programs/AutomationToolLauncher/AutomationToolLauncher.csproj b/Engine/Source/Programs/AutomationToolLauncher/AutomationToolLauncher.csproj
--- a/Engine/Source/Programs/AutomationToolLauncher/AutomationToolLauncher.csproj
+++ b/Engine/Source/Programs/AutomationToolLauncher/AutomationToolLauncher.csproj
@@ -9,7 +9,7 @@
     <AppDesignerFolder>Properties</AppDesignerFolder>
     <RootNamespace>AutomationToolLauncher</RootNamespace>
     <AssemblyName>AutomationToolLauncher</AssemblyName>
-    <TargetFrameworkVersion>v4.0</TargetFrameworkVersion>
+    <TargetFrameworkVersion>v4.5</TargetFrameworkVersion>
     <FileAlignment>512</FileAlignment>
   </PropertyGroup>
   <PropertyGroup Condition=" '$(Configuration)|$(Platform)' == 'Debug|AnyCPU' ">
diff --git a/Engine/Source/Programs/AutomationToolLauncher/app.config b/Engine/Source/Programs/AutomationToolLauncher/app.config
--- a/Engine/Source/Programs/AutomationToolLauncher/app.config
+++ b/Engine/Source/Programs/AutomationToolLauncher/app.config
@@ -3,4 +3,4 @@
    <runtime>
       <loadFromRemoteSources enabled="true"/>
    </runtime>
-</configuration>
+<startup><supportedRuntime version="v4.0" sku=".NETFramework,Version=v4.5"/></startup></configuration>
diff --git a/Engine/Source/Programs/DotNETCommon/DotNETUtilities/DotNETUtilities.csproj b/Engine/Source/Programs/DotNETCommon/DotNETUtilities/DotNETUtilities.csproj
--- a/Engine/Source/Programs/DotNETCommon/DotNETUtilities/DotNETUtilities.csproj
+++ b/Engine/Source/Programs/DotNETCommon/DotNETUtilities/DotNETUtilities.csproj
@@ -9,9 +9,8 @@
     <AppDesignerFolder>Properties</AppDesignerFolder>
     <RootNamespace>DotNETUtilities</RootNamespace>
     <AssemblyName>DotNETUtilities</AssemblyName>
-    <TargetFrameworkVersion>v4.0</TargetFrameworkVersion>
+    <TargetFrameworkVersion>v4.5</TargetFrameworkVersion>
     <FileAlignment>512</FileAlignment>
-    <TargetFrameworkProfile>Client</TargetFrameworkProfile>
   </PropertyGroup>
   <PropertyGroup Condition=" '$(Configuration)|$(Platform)' == 'Debug|AnyCPU' ">
     <DebugSymbols>true</DebugSymbols>
diff --git a/Engine/Source/Programs/UnrealBuildTool/System/ProjectFileGenerator.cs b/Engine/Source/Programs/UnrealBuildTool/System/ProjectFileGenerator.cs
--- a/Engine/Source/Programs/UnrealBuildTool/System/ProjectFileGenerator.cs
+++ b/Engine/Source/Programs/UnrealBuildTool/System/ProjectFileGenerator.cs
@@ -9,6 +9,8 @@ using Microsoft.Win32;
 using System.Linq;
 using System.Diagnostics;
 using Tools.DotNETCommon.ExecutingAssembly;
+using System.Xml.Linq;
+using System.Xml;
 
 namespace UnrealBuildTool
 {
@@ -270,7 +272,7 @@ namespace UnrealBuildTool
 					var Project = new VCSharpProjectFile(RelativeFileName);
 					Project.ShouldBuildForAllSolutionTargets = true;
 					AddExistingProjectFile(Project, bForceDevelopmentConfiguration: true);
-
+                    AutomationProjectFiles.Add( Project );
 					Folder.ChildProjects.Add( Project );
 				}
 			}
@@ -2066,7 +2068,7 @@ namespace UnrealBuildTool
 		/// <returns>True if successful</returns>
 		protected virtual bool WriteProjectFiles()
 		{
-			using(ProgressWriter Progress = new ProgressWriter("Writing project files...", true))
+            using(ProgressWriter Progress = new ProgressWriter("Writing project files...", true))
 			{
 				var TotalProjectFileCount = GeneratedProjectFiles.Count + 1;	// +1 for the master project file, which we'll save next
 
@@ -2086,6 +2088,30 @@ namespace UnrealBuildTool
 				WriteMasterProjectFile( UBTProject: UBTProject );
 				Progress.Write(TotalProjectFileCount, TotalProjectFileCount);
 			}
+
+            // Write AutomationReferences file
+            if (AutomationProjectFiles.Any())
+            {
+                var NS = XNamespace.Get("http://schemas.microsoft.com/developer/msbuild/2003");
+
+                var AutomationToolDir = Path.Combine(EngineRelativePath, "Source", "Programs", "AutomationTool");
+                new XDocument(
+                    new XElement(NS + "Project",
+                        new XAttribute("ToolsVersion", "12.0"),
+                        new XAttribute("DefaultTargets", "Build"),
+                        new XElement(NS + "ItemGroup",
+                            from AutomationProject in AutomationProjectFiles
+                            select new XElement(NS + "ProjectReference",
+                                new XAttribute("Include", Utils.MakePathRelativeTo(AutomationProject.ProjectFilePath, AutomationToolDir, true)),
+                                new XElement(NS + "Project", (AutomationProject as VCSharpProjectFile).ProjectGUID.ToString("B")),
+                                new XElement(NS + "Name", Path.GetFileNameWithoutExtension(AutomationProject.RelativeProjectFilePath)),
+                                new XElement(NS + "Private", "false")
+                            )
+                        )
+                    )
+                ).Save(Path.Combine( AutomationToolDir, "AutomationTool.csproj.References" ));
+            }
+
 			return true;
 		}
 
@@ -2305,6 +2331,8 @@ namespace UnrealBuildTool
 		/// aren't generating them ourselves.  Note that these may *not* always be C++ project files (e.g. C#)
 		protected readonly List<ProjectFile> OtherProjectFiles = new List<ProjectFile>();
 
+        protected readonly List<ProjectFile> AutomationProjectFiles = new List<ProjectFile>();
+
 		/// List of top-level folders in the master project file
 		protected MasterProjectFolder RootFolder;
 	}
diff --git a/Engine/Source/Programs/UnrealBuildTool/UnrealBuildTool.csproj b/Engine/Source/Programs/UnrealBuildTool/UnrealBuildTool.csproj
--- a/Engine/Source/Programs/UnrealBuildTool/UnrealBuildTool.csproj
+++ b/Engine/Source/Programs/UnrealBuildTool/UnrealBuildTool.csproj
@@ -12,7 +12,7 @@
     <AssemblyName>UnrealBuildTool</AssemblyName>
     <RunPostBuildEvent>OnOutputUpdated</RunPostBuildEvent>
     <IsWebBootstrapper>false</IsWebBootstrapper>
-    <TargetFrameworkVersion>v4.0</TargetFrameworkVersion>
+    <TargetFrameworkVersion>v4.5</TargetFrameworkVersion>
     <FileUpgradeFlags>
     </FileUpgradeFlags>
     <OldToolsVersion>3.5</OldToolsVersion>
diff --git a/Engine/Source/Programs/UnrealBuildTool/app.config b/Engine/Source/Programs/UnrealBuildTool/app.config
--- a/Engine/Source/Programs/UnrealBuildTool/app.config
+++ b/Engine/Source/Programs/UnrealBuildTool/app.config
@@ -1,3 +1,3 @@
 <?xml version="1.0"?>
 <configuration>
-<startup><supportedRuntime version="v4.0" sku=".NETFramework,Version=v4.0,Profile=Client"/></startup></configuration>
+<startup><supportedRuntime version="v4.0" sku=".NETFramework,Version=v4.5"/></startup></configuration>
