#html5 move emscripten cache to platform intermediate dir. fix an issue with CL 261593 where PlatformIntermediateFolder was not set during packaging,  using explicit name instead.

diff --git a/Engine/Source/Programs/UnrealBuildTool/HTML5/HTML5SDKInfo.cs b/Engine/Source/Programs/UnrealBuildTool/HTML5/HTML5SDKInfo.cs
--- a/Engine/Source/Programs/UnrealBuildTool/HTML5/HTML5SDKInfo.cs
+++ b/Engine/Source/Programs/UnrealBuildTool/HTML5/HTML5SDKInfo.cs
@@ -77,7 +77,7 @@ namespace UnrealBuildTool
 		{
 			get 
 			{
-				var TempPath = Path.GetFullPath(Path.Combine(BuildConfiguration.RelativeEnginePath, BuildConfiguration.PlatformIntermediateFolder));
+				var TempPath = Path.GetFullPath(Path.Combine(BuildConfiguration.RelativeEnginePath, BuildConfiguration.BaseIntermediateFolder, "HTML5"));
 				if (!Directory.Exists(TempPath))
 				{
 					Directory.CreateDirectory(TempPath);
@@ -86,6 +86,19 @@ namespace UnrealBuildTool
 			}
 		}
 
+		static public string EMSCRIPTEN_CACHE
+		{
+			get
+			{
+				var TempPath = Path.GetFullPath(Path.Combine(BuildConfiguration.RelativeEnginePath, BuildConfiguration.BaseIntermediateFolder, "HTML5"));
+				if (!Directory.Exists(TempPath))
+				{
+					Directory.CreateDirectory(TempPath);
+				}
+				return Path.Combine(TempPath, "EmscriptenCache");; 
+			}
+		}
+
 		public static string  SetupEmscriptenTemp()
 		{
 			string HTML5Intermediatory = Path.GetFullPath(Path.Combine(BuildConfiguration.RelativeEnginePath, BuildConfiguration.PlatformIntermediateFolder));
@@ -99,7 +112,6 @@ namespace UnrealBuildTool
 
 		public static string SetUpEmscriptenConfigFile()
 		{
-			string UserFolder = Environment.GetFolderPath(Environment.SpecialFolder.UserProfile);
 			string ConfigFile = DOT_EMSCRIPTEN;
 
 			if (!File.Exists(ConfigFile) || !File.ReadAllText(ConfigFile).Contains("GENERATEDBYUE4"))
diff --git a/Engine/Source/Programs/UnrealBuildTool/HTML5/HTML5ToolChain.cs b/Engine/Source/Programs/UnrealBuildTool/HTML5/HTML5ToolChain.cs
--- a/Engine/Source/Programs/UnrealBuildTool/HTML5/HTML5ToolChain.cs
+++ b/Engine/Source/Programs/UnrealBuildTool/HTML5/HTML5ToolChain.cs
@@ -35,6 +35,8 @@ namespace UnrealBuildTool
 			// set some environment variable we'll need.
 			// Forces emcc to use our generated .emscripten config, not the one in the users home directory.
 			Environment.SetEnvironmentVariable("EM_CONFIG", HTML5SDKInfo.DOT_EMSCRIPTEN);
+			Environment.SetEnvironmentVariable("EM_CACHE", HTML5SDKInfo.EMSCRIPTEN_CACHE); 
+
 		}
 	
 		static string GetSharedArguments_Global(CPPTargetConfiguration TargetConfiguration, string Architecture, bool bEnableShadowVariableWarning)
