fix for SSH not intiializing properly on OS X 10.9.x
#ios

diff --git a/Engine/Source/Programs/UnrealBuildTool/System/RPCUtilHelper.cs b/Engine/Source/Programs/UnrealBuildTool/System/RPCUtilHelper.cs
--- a/Engine/Source/Programs/UnrealBuildTool/System/RPCUtilHelper.cs
+++ b/Engine/Source/Programs/UnrealBuildTool/System/RPCUtilHelper.cs
@@ -115,8 +115,8 @@ namespace UnrealBuildTool
 						string[] Commands = new string[] 
 						{
 							"+\"%s\"",
-							"sysctl -a | grep hw.memsize | awk '{print $2}'",
-							"sysctl -a | grep hw.logicalcpu: | awk '{print $2}'",
+							"sysctl hw.memsize | awk '{print $2}'",
+							"sysctl hw.logicalcpu | awk '{print $2}'",
 						};
 						Hashtable Results = Command("/", "date", string.Join(" && ", Commands), null);
 						if ((Int64)Results["ExitCode"] != 0)
