UEB-266 - updated clean action for UBT to properly clean items from the most recent reciept, additionally, it will also now clean up stale modules when dependencies on them have been removed

diff --git a/Engine/Source/Programs/UnrealBuildTool/Configuration/UEBuildTarget.cs b/Engine/Source/Programs/UnrealBuildTool/Configuration/UEBuildTarget.cs
--- a/Engine/Source/Programs/UnrealBuildTool/Configuration/UEBuildTarget.cs
+++ b/Engine/Source/Programs/UnrealBuildTool/Configuration/UEBuildTarget.cs
@@ -1045,6 +1045,16 @@ namespace UnrealBuildTool
 		/// <param name="Manifest">Manifest</param>
 		protected void CleanTarget(BuildManifest Manifest)
 		{
+			// Expand all the paths in the receipt; they'll currently use variables for the engine and project directories
+			IUEBuildPlatform BuildPlatform = UEBuildPlatform.GetBuildPlatform(Platform);
+			// Expand all the paths in the receipt; they'll currently use variables for the engine and project directories
+			BuildReceipt ReceiptWithFullPaths = BuildReceipt.Read(BuildReceipt.GetDefaultPath(ProjectDirectory, TargetName, Platform, Configuration, BuildPlatform.GetActiveArchitecture()));
+			if (ReceiptWithFullPaths == null)
+			{
+				ReceiptWithFullPaths = new BuildReceipt(Receipt);
+			}
+			ReceiptWithFullPaths.ExpandPathVariables(BuildConfiguration.RelativeEnginePath, ProjectDirectory);
+
 			{
 				Log.TraceVerbose("Cleaning target {0} - AppName {1}", TargetName, AppName);
 
@@ -1053,13 +1063,26 @@ namespace UnrealBuildTool
 
 				// Collect all files to delete.
 				var AdditionalFileExtensions = new string[] { ".lib", ".exp", ".dll.response" };
-				var AllFilesToDelete = new List<string>(Manifest.BuildProducts);
-				foreach (var FileManifestItem in Manifest.BuildProducts)
+				var AllFilesToDelete = new List<string>();
+				foreach(BuildProduct BuildProduct in ReceiptWithFullPaths.BuildProducts)
 				{
-					var FileExt = Path.GetExtension(FileManifestItem);
+					// If we're cleaning, don't add any precompiled binaries to the manifest. We don't want to delete them.
+					if(UEBuildConfiguration.bCleanProject && bUsePrecompiled && BuildProduct.IsPrecompiled)
+					{
+						continue;
+					}
+
+					// Don't add static libraries into the manifest unless we're explicitly building them; we don't submit them to Perforce.
+					if(!UEBuildConfiguration.bCleanProject && !bPrecompile && (BuildProduct.Type == BuildProductType.StaticLibrary || BuildProduct.Type == BuildProductType.ImportLibrary))
+					{
+						continue;
+					}
+
+					AllFilesToDelete.Add(BuildProduct.Path);
+					var FileExt = Path.GetExtension(BuildProduct.Path);
 					if (FileExt == ".dll" || FileExt == ".exe")
 					{
-						var ManifestFileWithoutExtension = Utils.GetPathWithoutExtension(FileManifestItem);
+						var ManifestFileWithoutExtension = Utils.GetPathWithoutExtension(BuildProduct.Path);
 						foreach (var AdditionalExt in AdditionalFileExtensions)
 						{
 							var AdditionalFileToDelete = ManifestFileWithoutExtension + AdditionalExt;
@@ -1067,6 +1090,10 @@ namespace UnrealBuildTool
 						}
 					}
 				}
+				if (OnlyModules.Count == 0)
+				{
+					AllFilesToDelete.Add(BuildReceipt.GetDefaultPath(ProjectDirectory, TargetName, Platform, Configuration, BuildPlatform.GetActiveArchitecture()));
+				}
 
 				//@todo. This does not clean up files that are no longer built by the target...				
 				// Delete all output files listed in the manifest as well as any additional files.
@@ -1244,6 +1271,95 @@ namespace UnrealBuildTool
 			}
 		}
 
+		class BuldProductComparer : IEqualityComparer<BuildProduct>
+		{
+			public bool Equals(BuildProduct x, BuildProduct y)
+			{
+				if (Object.ReferenceEquals(x, y))
+					return true;
+
+				if (Object.ReferenceEquals(x, null) || Object.ReferenceEquals(y, null))
+					return false;
+
+				return x.Path == y.Path && x.Type == y.Type;
+			}
+
+			public int GetHashCode(BuildProduct product)
+			{
+				if (Object.ReferenceEquals(product, null))
+					return 0;
+
+				int hashProductPath = product.Path == null ? 0 : product.Path.GetHashCode();
+				int hashProductType = product.Type.GetHashCode();
+				return hashProductPath ^ hashProductType;
+			}
+		}
+		/// <summary>
+		/// Cleans all removed module intermediate files
+		/// </summary>
+		public void CleanStaleModules()
+		{
+			// Expand all the paths in the receipt; they'll currently use variables for the engine and project directories
+			BuildReceipt ReceiptWithFullPaths = new BuildReceipt(Receipt);
+			ReceiptWithFullPaths.ExpandPathVariables(BuildConfiguration.RelativeEnginePath, ProjectDirectory);
+			IUEBuildPlatform BuildPlatform = UEBuildPlatform.GetBuildPlatform(Platform);
+			// Expand all the paths in the receipt; they'll currently use variables for the engine and project directories
+			BuildReceipt BuiltReceiptWithFullPaths = BuildReceipt.Read(BuildReceipt.GetDefaultPath(ProjectDirectory, TargetName, Platform, Configuration, BuildPlatform.GetActiveArchitecture()));
+			if (BuiltReceiptWithFullPaths == null)
+			{
+				return;
+			}
+			BuiltReceiptWithFullPaths.ExpandPathVariables(BuildConfiguration.RelativeEnginePath, ProjectDirectory);
+
+			{
+				// Collect all files to build
+				var AdditionalFileExtensions = new string[] { ".lib", ".exp", ".dll.response" };
+				var AllFilesToDelete = new List<string>();
+				foreach (BuildProduct BuildProduct in BuiltReceiptWithFullPaths.BuildProducts)
+				{
+					// If we're cleaning, don't add any precompiled binaries to the manifest. We don't want to delete them.
+					if (UEBuildConfiguration.bCleanProject && bUsePrecompiled && BuildProduct.IsPrecompiled)
+					{
+						continue;
+					}
+
+					// Don't add static libraries into the manifest unless we're explicitly building them; we don't submit them to Perforce.
+					if (!UEBuildConfiguration.bCleanProject && !bPrecompile && (BuildProduct.Type == BuildProductType.StaticLibrary || BuildProduct.Type == BuildProductType.ImportLibrary))
+					{
+						continue;
+					}
+
+					if (!ReceiptWithFullPaths.BuildProducts.Contains(BuildProduct, new UEBuildTarget.BuldProductComparer()))
+					{
+						AllFilesToDelete.Add(BuildProduct.Path);
+						var FileExt = Path.GetExtension(BuildProduct.Path);
+						if (FileExt == ".dll" || FileExt == ".exe")
+						{
+							var ManifestFileWithoutExtension = Utils.GetPathWithoutExtension(BuildProduct.Path);
+							foreach (var AdditionalExt in AdditionalFileExtensions)
+							{
+								var AdditionalFileToDelete = ManifestFileWithoutExtension + AdditionalExt;
+								AllFilesToDelete.Add(AdditionalFileToDelete);
+							}
+						}
+					}
+				}
+
+				//@todo. This does not clean up files that are no longer built by the target...				
+				// Delete all output files listed in the manifest as well as any additional files.
+				foreach (var FileToDelete in AllFilesToDelete)
+				{
+					if (File.Exists(FileToDelete))
+					{
+						Log.TraceVerbose("\t\tDeleting " + FileToDelete);
+						CleanFile(FileToDelete);
+					}
+				}
+
+				// delete intermediate directory of the extra files
+			}
+		}
+
 		/** Generates a list of external files which are required to build this target */
 		public void GenerateExternalFileList()
 		{
diff --git a/Engine/Source/Programs/UnrealBuildTool/System/BuildReceipt.cs b/Engine/Source/Programs/UnrealBuildTool/System/BuildReceipt.cs
--- a/Engine/Source/Programs/UnrealBuildTool/System/BuildReceipt.cs
+++ b/Engine/Source/Programs/UnrealBuildTool/System/BuildReceipt.cs
@@ -349,9 +349,16 @@ namespace UnrealBuildTool
 		/// <param name="FileName">Filename to read from</param>
 		public static BuildReceipt Read(string FileName)
 		{
-			using(StreamReader Reader = new StreamReader(FileName))
+			if (File.Exists(FileName))
 			{
-				return (BuildReceipt)Serializer.Deserialize(Reader);
+				using (StreamReader Reader = new StreamReader(FileName))
+				{
+					return (BuildReceipt)Serializer.Deserialize(Reader);
+				}
+			}
+			else
+			{
+				return null;
 			}
 		}
 
diff --git a/Engine/Source/Programs/UnrealBuildTool/System/UnrealBuildTool.cs b/Engine/Source/Programs/UnrealBuildTool/System/UnrealBuildTool.cs
--- a/Engine/Source/Programs/UnrealBuildTool/System/UnrealBuildTool.cs
+++ b/Engine/Source/Programs/UnrealBuildTool/System/UnrealBuildTool.cs
@@ -1913,7 +1913,13 @@ namespace UnrealBuildTool
                                     ActionsToExecute.Count
                                     );
 
-                            ToolChain.PreBuildSync();
+							// clean up any stale modules
+							foreach (UEBuildTarget Target in Targets)
+							{
+								Target.CleanStaleModules();
+							}
+							
+							ToolChain.PreBuildSync();
 
                     
                             // Cache indirect includes for all outdated C++ files.  We kick this off as a background thread so that it can
