fix for occasional failure to build UE4Game after generating projects
#ubt

diff --git a/Engine/Source/Programs/UnrealBuildTool/System/DynamicCompilation.cs b/Engine/Source/Programs/UnrealBuildTool/System/DynamicCompilation.cs
--- a/Engine/Source/Programs/UnrealBuildTool/System/DynamicCompilation.cs
+++ b/Engine/Source/Programs/UnrealBuildTool/System/DynamicCompilation.cs
@@ -297,12 +297,9 @@ namespace UnrealBuildTool
 				// since the previous time we compiled the assembly.  In that case, we'll always want to recompile it!
 				{
 					FileInfo AssemblySourcesListFile = new FileInfo( AssemblySourcesListFilePath );
-					using( var Writer = AssemblySourcesListFile.OpenWrite() )
+					using( var Writer = AssemblySourcesListFile.CreateText() )
 					{
-						using( var TextWriter = new StreamWriter( Writer ) )
-						{
-							SourceFileNames.ForEach( x => TextWriter.WriteLine( x ) );
-						}
+						SourceFileNames.ForEach( x => Writer.WriteLine( x ) );
 					}
 				}
 			}
