Change AutomationTool.csproj.References to use ToolsVersion = 4.0 instead of 12.0. Shot in the dark as to why the references don't seem to be working 100% of the time. But this should match the other project stuff anyway.

diff --git a/Engine/Source/Programs/UnrealBuildTool/System/ProjectFileGenerator.cs b/Engine/Source/Programs/UnrealBuildTool/System/ProjectFileGenerator.cs
--- a/Engine/Source/Programs/UnrealBuildTool/System/ProjectFileGenerator.cs
+++ b/Engine/Source/Programs/UnrealBuildTool/System/ProjectFileGenerator.cs
@@ -2097,7 +2097,7 @@ namespace UnrealBuildTool
                 var AutomationToolDir = Path.Combine(EngineRelativePath, "Source", "Programs", "AutomationTool");
                 new XDocument(
                     new XElement(NS + "Project",
-                        new XAttribute("ToolsVersion", "12.0"),
+                        new XAttribute("ToolsVersion", "4.0"),
                         new XAttribute("DefaultTargets", "Build"),
                         new XElement(NS + "ItemGroup",
                             from AutomationProject in AutomationProjectFiles
