fix for project generation when code signing
#ios

diff --git a/Engine/Source/Programs/UnrealBuildTool/System/XcodeProjectFileGenerator.cs b/Engine/Source/Programs/UnrealBuildTool/System/XcodeProjectFileGenerator.cs
--- a/Engine/Source/Programs/UnrealBuildTool/System/XcodeProjectFileGenerator.cs
+++ b/Engine/Source/Programs/UnrealBuildTool/System/XcodeProjectFileGenerator.cs
@@ -1073,12 +1073,27 @@ namespace UnrealBuildTool
 
 		private void PopulateTargets(List<XcodeProjectTarget> ProjectTargets, List<XcodeContainerItemProxy> ContainerItemProxies, List<XcodeTargetDependency> TargetDependencies, XcodeProjectTarget ProjectTarget, List<XcodeFramework> Frameworks)
 		{
+			List<ProjectTarget> TargetProjects = new List<ProjectTarget> ();
 			foreach (var Project in GeneratedProjectFiles)
 			{
 				foreach (var TargetProject in Project.ProjectTargets)
 				{
+					TargetProjects.Add (TargetProject);
+				}
+			}
+			if (bGeneratingGameProjectFiles && TargetProjects.Count == 0)
+			{
+				// add UE4Game as most likely this is a content-only project
+				ProjectTarget TargetProject = new global::UnrealBuildTool.ProjectTarget();
+				var Target = new TargetInfo(UnrealTargetPlatform.Mac, UnrealTargetConfiguration.Development);
+				TargetProject.TargetFilePath = "";
+				TargetProject.TargetRules = RulesCompiler.CreateTargetRules("UE4Game", Target, false, out TargetProject.TargetFilePath);
+				TargetProjects.Add(TargetProject);
+			}
+			foreach (var TargetProject in TargetProjects)
+			{
 					string TargetPath = TargetProject.TargetFilePath;
-					string TargetName = Utils.GetFilenameWithoutAnyExtensions(Path.GetFileName(TargetPath));
+					string TargetName = Utils.GetFilenameWithoutAnyExtensions (Path.GetFileName (TargetPath));
 					bool WantProjectFileForTarget = true;
 					bool IsEngineTarget = false;
 					if (bGeneratingGameProjectFiles || bGeneratingRocketProjectFiles)
@@ -1181,7 +1196,6 @@ namespace UnrealBuildTool
 						}
 					}
 				}
-			}
 		}
 
 		/// <summary>
