Improved Android ARMv7 debug flags and moved UELinkerExceptions to Launch module to prevent multiple linking on x86_64
#ue4
#android

diff --git a/Engine/Source/Programs/UnrealBuildTool/Android/AndroidToolChain.cs b/Engine/Source/Programs/UnrealBuildTool/Android/AndroidToolChain.cs
--- a/Engine/Source/Programs/UnrealBuildTool/Android/AndroidToolChain.cs
+++ b/Engine/Source/Programs/UnrealBuildTool/Android/AndroidToolChain.cs
@@ -397,6 +397,13 @@ namespace UnrealBuildTool
 				Result += " -mfloat-abi=softfp";
 				Result += " -mfpu=vfpv3-d16";			//@todo android: UE3 was just vfp. arm7a should all support v3 with 16 registers
 
+				// Add flags for on-device debugging	
+				if (CompileEnvironment.Config.Target.Configuration == CPPTargetConfiguration.Debug)
+				{
+					Result += " -fno-omit-frame-pointer";	// Disable removing the save/restore frame pointer for better debugging
+					Result += " -fno-function-sections";	// Improve breakpoint location
+				}
+
 				// Some switches interfere with on-device debugging
 				if (CompileEnvironment.Config.Target.Configuration != CPPTargetConfiguration.Debug)
 				{
@@ -688,12 +695,12 @@ namespace UnrealBuildTool
 		static void GenerateEmptyLinkFunctionsForRemovedModules(List<FileItem> SourceFiles, string ModuleName, string OutputDirectory)
 		{
 			// Only add to UELinkerFixups module
-			if (!ModuleName.Equals("UELinkerFixups"))
+			if (!ModuleName.Equals("Launch"))
 			{
 				return;
 			}
 
-			string LinkerExceptionsName = "UELinkerExceptions";
+			string LinkerExceptionsName = "../UELinkerExceptions";
 			string LinkerExceptionsCPPFilename = Path.Combine(OutputDirectory, LinkerExceptionsName + ".cpp");
 
 			// Create the cpp filename
