Core: Added logging about the command line used to invoke UHT and the duration

diff --git a/Engine/Source/Programs/UnrealBuildTool/System/ExternalExecution.cs b/Engine/Source/Programs/UnrealBuildTool/System/ExternalExecution.cs
--- a/Engine/Source/Programs/UnrealBuildTool/System/ExternalExecution.cs
+++ b/Engine/Source/Programs/UnrealBuildTool/System/ExternalExecution.cs
@@ -692,6 +692,8 @@ namespace UnrealBuildTool
 						CmdLine += " -FailIfGeneratedCodeChanges";
 					}
 
+					Log.TraceInformation("  Running UnrealHeaderTool {0}", CmdLine);
+
 					Stopwatch s = new Stopwatch();
 					s.Start();
 					UHTResult = (ECompilationResult) RunExternalExecutable(ExternalExecution.GetHeaderToolPath(), CmdLine);
@@ -713,7 +715,7 @@ namespace UnrealBuildTool
 						return false;
 					}
 
-					Log.TraceInformation( "Reflection code generated for {0}", ActualTargetName );
+					Log.TraceInformation("Reflection code generated for {0} in {1} seconds", ActualTargetName, s.Elapsed.TotalSeconds);
 					if( BuildConfiguration.bPrintPerformanceInfo )
 					{
 						Log.TraceInformation( "UnrealHeaderTool took {1}", ActualTargetName, (double)s.ElapsedMilliseconds/1000.0 );
