Better handling of multiple Android build targets (copies proper .so)
#ue4
#android

diff --git a/Engine/Source/Programs/UnrealBuildTool/Android/UEBuildAndroid.cs b/Engine/Source/Programs/UnrealBuildTool/Android/UEBuildAndroid.cs
--- a/Engine/Source/Programs/UnrealBuildTool/Android/UEBuildAndroid.cs
+++ b/Engine/Source/Programs/UnrealBuildTool/Android/UEBuildAndroid.cs
@@ -349,6 +349,12 @@ namespace UnrealBuildTool
 			}
 		}
 
+		private bool UseTegraGraphicsDebugger(UEBuildTarget InBuildTarget)
+		{
+			// Disable for now
+			return false;
+		}
+
 		public override void SetUpEnvironment(UEBuildTarget InBuildTarget)
 		{
 			// we want gcc toolchain 4.9, but fall back to 4.8 or 4.6 for now if it doesn't exist
@@ -403,9 +409,13 @@ namespace UnrealBuildTool
 			InBuildTarget.GlobalLinkEnvironment.Config.LibraryPaths.Add(UEBuildConfiguration.UEThirdPartySourceDirectory + "Android/cxa_demangle/x86");
 			InBuildTarget.GlobalLinkEnvironment.Config.LibraryPaths.Add(UEBuildConfiguration.UEThirdPartySourceDirectory + "Android/cxa_demangle/x64");
 
-			//@TODO: Tegra Gfx Debugger
-//			InBuildTarget.GlobalLinkEnvironment.Config.LibraryPaths.Add(UEBuildConfiguration.UEThirdPartySourceDirectory + "NVIDIA/TegraGfxDebugger");
-//			InBuildTarget.GlobalLinkEnvironment.Config.AdditionalLibraries.Add("Tegra_gfx_debugger");
+			//@TODO: Tegra Gfx Debugger - standardize locations - for now, change the hardcoded paths and force this to return true to test
+			if (UseTegraGraphicsDebugger(InBuildTarget))
+			{
+				//InBuildTarget.GlobalLinkEnvironment.Config.LibraryPaths.Add(UEBuildConfiguration.UEThirdPartySourceDirectory + "NVIDIA/TegraGfxDebugger");
+				//InBuildTarget.GlobalLinkEnvironment.Config.LibraryPaths.Add("F:/NVPACK/android-kk-egl-t124-a32/stub");
+				//InBuildTarget.GlobalLinkEnvironment.Config.AdditionalLibraries.Add("Nvidia_gfx_debugger_stub");
+			}
 
             InBuildTarget.GlobalLinkEnvironment.Config.AdditionalLibraries.Add("gnustl_shared");
             InBuildTarget.GlobalLinkEnvironment.Config.AdditionalLibraries.Add("gcc");
@@ -414,8 +424,11 @@ namespace UnrealBuildTool
 			InBuildTarget.GlobalLinkEnvironment.Config.AdditionalLibraries.Add("m");
 			InBuildTarget.GlobalLinkEnvironment.Config.AdditionalLibraries.Add("log");
 			InBuildTarget.GlobalLinkEnvironment.Config.AdditionalLibraries.Add("dl");
-			InBuildTarget.GlobalLinkEnvironment.Config.AdditionalLibraries.Add("GLESv2");
-			InBuildTarget.GlobalLinkEnvironment.Config.AdditionalLibraries.Add("EGL");
+			if (!UseTegraGraphicsDebugger(InBuildTarget))
+			{
+				InBuildTarget.GlobalLinkEnvironment.Config.AdditionalLibraries.Add("GLESv2");
+				InBuildTarget.GlobalLinkEnvironment.Config.AdditionalLibraries.Add("EGL");
+			}
 			InBuildTarget.GlobalLinkEnvironment.Config.AdditionalLibraries.Add("OpenSLES");
 			InBuildTarget.GlobalLinkEnvironment.Config.AdditionalLibraries.Add("android");
 			InBuildTarget.GlobalLinkEnvironment.Config.AdditionalLibraries.Add("cxa_demangle");
diff --git a/Engine/Source/Programs/UnrealBuildTool/Android/UEDeployAndroid.cs b/Engine/Source/Programs/UnrealBuildTool/Android/UEDeployAndroid.cs
--- a/Engine/Source/Programs/UnrealBuildTool/Android/UEDeployAndroid.cs
+++ b/Engine/Source/Programs/UnrealBuildTool/Android/UEDeployAndroid.cs
@@ -557,13 +557,15 @@ namespace UnrealBuildTool.Android
 			}
 		}
 
-		//@TODO: To enable the NVIDIA Gfx Debugger
+		//@TODO: To enable the NVIDIA Gfx Debugger - these paths need to be standardized
 		private static void CopyGfxDebugger(string UE4BuildPath, string UE4Arch)
 		{
-//			string Arch = GetNDKArch(UE4Arch);
-//			Directory.CreateDirectory(UE4BuildPath + "/libs/" + Arch);
-//			File.Copy("E:/Dev2/UE4-Clean/Engine/Source/ThirdParty/NVIDIA/TegraGfxDebugger/libNvPmApi.Core.so", UE4BuildPath + "/libs/" + Arch + "/libNvPmApi.Core.so");
-//			File.Copy("E:/Dev2/UE4-Clean/Engine/Source/ThirdParty/NVIDIA/TegraGfxDebugger/libTegra_gfx_debugger.so", UE4BuildPath + "/libs/" + Arch + "/libTegra_gfx_debugger.so");
+			/*
+			string Arch = GetNDKArch(UE4Arch);
+			Directory.CreateDirectory(UE4BuildPath + "/libs/" + Arch);
+			File.Copy("F:/NVPACK/android-kk-egl-t124-a32/Stripped_libNvPmApi.Core.so", UE4BuildPath + "/libs/" + Arch + "/libNvPmApi.Core.so", true);
+			File.Copy("F:/NVPACK/android-kk-egl-t124-a32/Stripped_libNvidia_gfx_debugger.so", UE4BuildPath + "/libs/" + Arch + "/libNvidia_gfx_debugger.so", true);
+			*/
 		}
 
 		private static void RunCommandLineProgramAndThrowOnError(string WorkingDirectory, string Command, string Params, string OverrideDesc=null, bool bUseShellExecute=false)
@@ -1160,6 +1162,39 @@ namespace UnrealBuildTool.Android
 			return Text.ToString();
 		}
 
+		private bool FilesAreDifferent(string SourceFilename, string DestFilename)
+		{
+			// source must exist
+			FileInfo SourceInfo = new FileInfo(SourceFilename);
+			if (!SourceInfo.Exists)
+			{
+				throw new BuildException("Can't make an APK without file [{0}]", SourceFilename);
+			}
+
+			// different if destination doesn't exist
+			FileInfo DestInfo = new FileInfo(DestFilename);
+			if (!DestInfo.Exists)
+			{
+				return true;
+			}
+
+			// file lengths differ?
+			if (SourceInfo.Length != DestInfo.Length)
+			{
+				return true;
+			}
+
+			// validate timestamps
+			TimeSpan Diff = DestInfo.LastWriteTimeUtc - SourceInfo.LastWriteTimeUtc;
+			if (Diff.TotalSeconds < -1 || Diff.TotalSeconds > 1)
+			{
+				return true;
+			}
+
+			// could check actual bytes just to be sure, but good enough
+			return false;
+		}
+
 		private void MakeApk(string ProjectName, string ProjectDirectory, string OutputPath, string EngineDirectory, bool bForDistribution, string CookFlavor, bool bMakeSeparateApks, bool bIncrementalPackage, bool bDisallowPackagingDataInApk)
 		{
 			Log.TraceInformation("\n===={0}====PREPARING TO MAKE APK=================================================================", DateTime.Now.ToString());
@@ -1279,8 +1314,8 @@ namespace UnrealBuildTool.Android
             }
 
             // Now we have to spin over all the arch/gpu combinations to make sure they all match
-
-            if (BuildList.Count() == 0)
+			int BuildListComboTotal = BuildList.Count();
+            if (BuildListComboTotal == 0)
             {
                 Log.TraceInformation("Output .apk file(s) are up to date (dependencies and build settings are up to date)");
                 return;
@@ -1389,16 +1424,16 @@ namespace UnrealBuildTool.Android
                 }
 
                 String FinalSOName;
+
                 if(HasNDKPath)
                 {
                     string LibDir = UE4BuildPath + "/jni/" + GetNDKArch(Arch);
                     FinalSOName = LibDir + "/libUE4.so";
 
-                    // check to see if libUE4.so is newer than last time we copied
-                    TimeSpan Diff = File.GetLastWriteTimeUtc(FinalSOName) - File.GetLastWriteTimeUtc(SourceSOName);
-                    if (!File.Exists(FinalSOName) || Diff.TotalSeconds < -1 || Diff.TotalSeconds > 1)
+                    // check to see if libUE4.so needs to be copied
+					if (BuildListComboTotal > 1 || FilesAreDifferent(SourceSOName, FinalSOName))
                     {
-                        Log.TraceInformation("\nCopying new .so file to jni folder...");
+                        Log.TraceInformation("\nCopying new .so {0} file to jni folder...", SourceSOName);
                         Directory.CreateDirectory(LibDir);
                         // copy the binary to the standard .so location
                         File.Copy(SourceSOName, FinalSOName, true);
@@ -1409,11 +1444,11 @@ namespace UnrealBuildTool.Android
                     // if no NDK, we don't need any of the debugger stuff, so we just copy the .so to where it will end up
                     FinalSOName = UE4BuildPath + "/libs/" + GetNDKArch(Arch) + "/libUE4.so";
 
-                    // check to see if libUE4.so is newer than last time we copied
-                    TimeSpan Diff = File.GetLastWriteTimeUtc(FinalSOName) - File.GetLastWriteTimeUtc(SourceSOName);
-                    if (!File.Exists(FinalSOName) || Diff.TotalSeconds < -1 || Diff.TotalSeconds > 1)
+                    // check to see if libUE4.so needs to be copied
+					if (BuildListComboTotal > 1 || FilesAreDifferent(SourceSOName, FinalSOName))
                     {
-                        Directory.CreateDirectory(Path.GetDirectoryName(FinalSOName));
+						Log.TraceInformation("\nCopying .so {0} file to jni folder...", SourceSOName);
+						Directory.CreateDirectory(Path.GetDirectoryName(FinalSOName));
                         File.Copy(SourceSOName, FinalSOName, true);
                     }
                 }
