#gearvr Fix for so deployment

diff --git a/Engine/Source/Programs/UnrealBuildTool/Android/UEDeployAndroid.cs b/Engine/Source/Programs/UnrealBuildTool/Android/UEDeployAndroid.cs
--- a/Engine/Source/Programs/UnrealBuildTool/Android/UEDeployAndroid.cs
+++ b/Engine/Source/Programs/UnrealBuildTool/Android/UEDeployAndroid.cs
@@ -1488,6 +1488,7 @@ namespace UnrealBuildTool.Android
                 // copy libgnustl_shared.so to library (use 4.8 if possible, otherwise 4.6)
                 CopySTL(UE4BuildPath, Arch);
                 CopyGfxDebugger(UE4BuildPath, Arch);
+				CopyPluginLibs(EngineDirectory, UE4BuildPath, Arch); 
 
                 Log.TraceInformation("\n===={0}====PERFORMING FINAL APK PACKAGE OPERATION================================================", DateTime.Now.ToString());
 
@@ -1597,5 +1598,44 @@ namespace UnrealBuildTool.Android
 				Log.TraceInformation(Line.Data);
 			}
 		}
-	}
+
+        private static void CopyPluginLibs(string EngineDirectory, string UE4BuildPath, string UE4Arch)
+        {
+            ConfigCacheIni Ini = GetConfigCacheIni("Engine");
+            string Arch = GetNDKArch(UE4Arch);
+            string PluginsDir = Path.GetFullPath(Path.Combine(EngineDirectory, "Plugins"));
+            string ThirdPartyDir = Path.GetFullPath(Path.Combine(EngineDirectory, "Source/ThirdParty"));
+
+            // Check for GearVR enabled
+            bool bPackageForGearVR;
+            Ini.GetBool("/Script/AndroidRuntimeSettings.AndroidRuntimeSettings", "bPackageForGearVR", out bPackageForGearVR);
+
+            Console.WriteLine("bPackageForGearVR? {0}, arch {1}", bPackageForGearVR, Arch);
+
+            // Note: only support ARMv7 at the moment
+            if (bPackageForGearVR && Arch == "armeabi-v7a")
+            {
+                string VRLibFile = ThirdPartyDir + "/Oculus/LibOVRMobile/LibOVRMobile_060/VrApi/Libs/Android/armeabi-v7a/libvrapi.so";
+
+                if (File.Exists(VRLibFile))
+                {
+                    Console.WriteLine("File {0} found", VRLibFile);
+
+                    Console.WriteLine("Dir {0} created", UE4BuildPath + "/libs/" + Arch);
+
+                    Directory.CreateDirectory(UE4BuildPath + "/libs/" + Arch);
+                    File.Copy(VRLibFile, UE4BuildPath + "/libs/" + Arch + "/libvrapi.so", true);
+
+                    Console.WriteLine("File {0} copied to {1}", VRLibFile, UE4BuildPath + "/libs/" + Arch + "/libvrapi.so");
+                    
+                }
+				else
+				{
+					Console.WriteLine("Failed to find the GearVR library required for packaging: {0}, VRLibFile);
+				}
+            }
+
+            // Add others here...
+        }
+    }
 }
