Fix UBT.

diff --git a/Engine/Source/Programs/UnrealBuildTool/Android/UEDeployAndroid.cs b/Engine/Source/Programs/UnrealBuildTool/Android/UEDeployAndroid.cs
--- a/Engine/Source/Programs/UnrealBuildTool/Android/UEDeployAndroid.cs
+++ b/Engine/Source/Programs/UnrealBuildTool/Android/UEDeployAndroid.cs
@@ -1631,7 +1631,7 @@ namespace UnrealBuildTool.Android
                 }
 				else
 				{
-					Console.WriteLine("Failed to find the GearVR library required for packaging: {0}, VRLibFile);
+					Console.WriteLine("Failed to find the GearVR library required for packaging: {0}", VRLibFile);
 				}
             }
 
