UE-18578: Hot-Reload: UBT is not properly detecting hot-reload compilation.

diff --git a/Engine/Source/Programs/UnrealBuildTool/UnrealBuildTool.csproj b/Engine/Source/Programs/UnrealBuildTool/UnrealBuildTool.csproj
--- a/Engine/Source/Programs/UnrealBuildTool/UnrealBuildTool.csproj
+++ b/Engine/Source/Programs/UnrealBuildTool/UnrealBuildTool.csproj
@@ -69,6 +69,7 @@
     <ErrorReport>prompt</ErrorReport>
     <CodeAnalysisIgnoreBuiltInRuleSets>true</CodeAnalysisIgnoreBuiltInRuleSets>
     <CodeAnalysisIgnoreBuiltInRules>true</CodeAnalysisIgnoreBuiltInRules>
+    <Prefer32Bit>false</Prefer32Bit>
   </PropertyGroup>
   <PropertyGroup Condition="'$(Configuration)|$(Platform)' == 'Debug|x86'">
     <DefineConstants>TRACE;DEBUG</DefineConstants>
