UE-18712: In editor Compile & Hot Reload always fails.

Fixed hot reload problem caused by CL 2615855.

diff --git a/Engine/Source/Programs/UnrealBuildTool/System/UnrealBuildTool.cs b/Engine/Source/Programs/UnrealBuildTool/System/UnrealBuildTool.cs
--- a/Engine/Source/Programs/UnrealBuildTool/System/UnrealBuildTool.cs
+++ b/Engine/Source/Programs/UnrealBuildTool/System/UnrealBuildTool.cs
@@ -1924,10 +1924,13 @@ namespace UnrealBuildTool
                                     ActionsToExecute.Count
                                     );
 
-							// clean up any stale modules
-							foreach (UEBuildTarget Target in Targets)
+							if (!bIsHotReload)
 							{
-								Target.CleanStaleModules();
+								// clean up any stale modules
+								foreach (UEBuildTarget Target in Targets)
+								{
+									Target.CleanStaleModules();
+								}
 							}
 							
 							ToolChain.PreBuildSync();
