Attempt to fix exception when trying to resolve include paths on Linux. If we get an exception because the filename exceeds MAX_PATH, it can't be the header we're looking for.

diff --git a/Engine/Source/Programs/UnrealBuildTool/System/CPPHeaders.cs b/Engine/Source/Programs/UnrealBuildTool/System/CPPHeaders.cs
--- a/Engine/Source/Programs/UnrealBuildTool/System/CPPHeaders.cs
+++ b/Engine/Source/Programs/UnrealBuildTool/System/CPPHeaders.cs
@@ -132,11 +132,10 @@ namespace UnrealBuildTool
 						{
 							FullFilePath = Path.GetFullPath(RelativeFilePath);
 						}
-						catch (Exception Exception)
+						catch (Exception)
 						{
-							throw new BuildException(Exception, "Failed to get full path for include: \"{0}\"", RelativeFilePath);
 						}
-						if( DirectoryLookupCache.FileExists( FullFilePath ) )
+						if( FullFilePath != null && DirectoryLookupCache.FileExists( FullFilePath ) )
 						{
 							Result = FileItem.GetItemByFullPath( FullFilePath );
 							break;
