UE-18654 - Fix uses of GetExecutingAssembly in UAT/UBT

diff --git a/Engine/Source/Programs/AutomationTool/AutomationUtils/ProjectUtils.cs b/Engine/Source/Programs/AutomationTool/AutomationUtils/ProjectUtils.cs
--- a/Engine/Source/Programs/AutomationTool/AutomationUtils/ProjectUtils.cs
+++ b/Engine/Source/Programs/AutomationTool/AutomationUtils/ProjectUtils.cs
@@ -8,7 +8,8 @@ using System.IO;
 using UnrealBuildTool;
 using System.Diagnostics;
 using Tools.DotNETCommon.CaselessDictionary;
-using Tools.DotNETCommon.ExecutingAssembly;
+using Tools.DotNETCommon;
+using System.Reflection;
 
 namespace AutomationTool
 {
@@ -194,10 +195,10 @@ namespace AutomationTool
 			string oldCWD = Directory.GetCurrentDirectory();
 			if (BuildConfiguration.RelativeEnginePath == "../../Engine/")
 			{
-				string EngineSourceDirectory = Path.Combine(ExecutingAssembly.GetDirectory(), "..", "..", "..", "Engine", "Source");
+				string EngineSourceDirectory = Path.Combine(Path.GetDirectoryName(Assembly.GetEntryAssembly().GetOriginalLocation()), "..", "..", "..", "Engine", "Source");
 				if (!Directory.Exists(EngineSourceDirectory)) // only set the directory if it exists, this should only happen if we are launching the editor from an artist sync
 				{
-					EngineSourceDirectory = Path.Combine(ExecutingAssembly.GetDirectory(), "..", "..", "..", "Engine", "Binaries");
+                    EngineSourceDirectory = Path.Combine(Path.GetDirectoryName(Assembly.GetEntryAssembly().GetOriginalLocation()), "..", "..", "..", "Engine", "Binaries");
 				}
 				Directory.SetCurrentDirectory(EngineSourceDirectory);
 			}
diff --git a/Engine/Source/Programs/DotNETCommon/DotNETUtilities/AssemblyUtils.cs b/Engine/Source/Programs/DotNETCommon/DotNETUtilities/AssemblyUtils.cs
new file mode 100644
--- /dev/null
+++ b/Engine/Source/Programs/DotNETCommon/DotNETUtilities/AssemblyUtils.cs
@@ -0,0 +1,22 @@
+﻿// Copyright 1998-2015 Epic Games, Inc. All Rights Reserved.
+
+using System;
+using System.IO;
+using System.Reflection;
+
+namespace Tools.DotNETCommon
+{
+	public static class AssemblyUtils
+	{
+        /// <summary>
+        /// Gets the original location (path and filename) of an assembly.
+        /// This method is using Assembly.CodeBase property to properly resolve original
+        /// assembly path in case shadow copying is enabled.
+        /// </summary>
+        /// <returns>Absolute path and filename to the assembly.</returns>
+        public static string GetOriginalLocation(this Assembly ThisAssembly)
+        {
+            return new Uri(ThisAssembly.CodeBase).LocalPath;
+        }
+	}
+}
diff --git a/Engine/Source/Programs/DotNETCommon/DotNETUtilities/DotNETUtilities.csproj b/Engine/Source/Programs/DotNETCommon/DotNETUtilities/DotNETUtilities.csproj
--- a/Engine/Source/Programs/DotNETCommon/DotNETUtilities/DotNETUtilities.csproj
+++ b/Engine/Source/Programs/DotNETCommon/DotNETUtilities/DotNETUtilities.csproj
@@ -49,7 +49,7 @@
       <Link>Properties\MetaData.cs</Link>
     </Compile>
     <Compile Include="CaselessDictionary.cs" />
-    <Compile Include="ExecutingAssembly.cs" />
+    <Compile Include="AssemblyUtils.cs" />
     <Compile Include="FileContentsCacheType.cs" />
     <Compile Include="FileSystem.cs" />
     <Compile Include="HarvestEnvVars.cs" />
diff --git a/Engine/Source/Programs/DotNETCommon/DotNETUtilities/ExecutingAssembly.cs b/Engine/Source/Programs/DotNETCommon/DotNETUtilities/ExecutingAssembly.cs
deleted file mode 100644
diff --git a/Engine/Source/Programs/DotNETCommon/DotNETUtilities/HarvestEnvVars.cs b/Engine/Source/Programs/DotNETCommon/DotNETUtilities/HarvestEnvVars.cs
--- a/Engine/Source/Programs/DotNETCommon/DotNETUtilities/HarvestEnvVars.cs
+++ b/Engine/Source/Programs/DotNETCommon/DotNETUtilities/HarvestEnvVars.cs
@@ -4,6 +4,7 @@ using System;
 using System.Collections.Generic;
 using System.Diagnostics;
 using System.IO;
+using System.Reflection;
 using System.Text;
 using System.Xml;
 using System.Xml.Serialization;
@@ -51,7 +52,7 @@ namespace Tools.DotNETCommon.HarvestEnvVars
 			string EnvReaderBatchFileName = EnvOutputFileName + ".bat";
 			try
 			{
-				var EnvVarsToXMLExePath = Path.Combine(ExecutingAssembly.ExecutingAssembly.GetDirectory(), "EnvVarsToXML.exe");
+				var EnvVarsToXMLExePath = Path.Combine(Path.GetDirectoryName(Assembly.GetEntryAssembly().GetOriginalLocation()), "EnvVarsToXML.exe");
 
 				// Convert every path to short filenames to ensure we don't accidentally write out a non-ASCII batch file
 				var ShortBatchFileName       = FileSystem.FileSystem.GetShortPathName(BatchFileName);
diff --git a/Engine/Source/Programs/UnrealBuildTool/System/DynamicCompilation.cs b/Engine/Source/Programs/UnrealBuildTool/System/DynamicCompilation.cs
--- a/Engine/Source/Programs/UnrealBuildTool/System/DynamicCompilation.cs
+++ b/Engine/Source/Programs/UnrealBuildTool/System/DynamicCompilation.cs
@@ -7,14 +7,14 @@ using System.CodeDom.Compiler;
 using Microsoft.CSharp;
 using System.Reflection;
 using System.Diagnostics;
-using Tools.DotNETCommon.ExecutingAssembly;
+using Tools.DotNETCommon;
 
 namespace UnrealBuildTool
 {
 	public class DynamicCompilation
 	{
 		/// File information for UnrealBuildTool.exe, cached at program start
-		private static FileInfo UBTExecutableFileInfo = new FileInfo( ExecutingAssembly.GetFilename() );
+		private static FileInfo UBTExecutableFileInfo = new FileInfo( Assembly.GetEntryAssembly().GetOriginalLocation() );
 
 		/*
 		 * Checks to see if the assembly needs compilation
diff --git a/Engine/Source/Programs/UnrealBuildTool/System/ProjectFileGenerator.cs b/Engine/Source/Programs/UnrealBuildTool/System/ProjectFileGenerator.cs
--- a/Engine/Source/Programs/UnrealBuildTool/System/ProjectFileGenerator.cs
+++ b/Engine/Source/Programs/UnrealBuildTool/System/ProjectFileGenerator.cs
@@ -8,7 +8,7 @@ using System.Reflection;
 using Microsoft.Win32;
 using System.Linq;
 using System.Diagnostics;
-using Tools.DotNETCommon.ExecutingAssembly;
+using Tools.DotNETCommon;
 using System.Xml.Linq;
 using System.Xml;
 
@@ -307,7 +307,7 @@ namespace UnrealBuildTool
 					{
 						if (BuildHostPlatform.Current.Platform == UnrealTargetPlatform.Mac)
 						{
-							MasterProjectRelativePath = Path.GetFullPath(Path.Combine(ExecutingAssembly.GetDirectory(), "..", "..", "..", "Engine"));
+							MasterProjectRelativePath = Path.GetFullPath(Path.Combine(Path.GetDirectoryName(Assembly.GetEntryAssembly().GetOriginalLocation()), "..", "..", "..", "Engine"));
 							GameProjectName = "UE4Game";
 						}
 						if (!Directory.Exists(MasterProjectRelativePath + "/Source"))
diff --git a/Engine/Source/Programs/UnrealBuildTool/System/UProjectInfo.cs b/Engine/Source/Programs/UnrealBuildTool/System/UProjectInfo.cs
--- a/Engine/Source/Programs/UnrealBuildTool/System/UProjectInfo.cs
+++ b/Engine/Source/Programs/UnrealBuildTool/System/UProjectInfo.cs
@@ -8,7 +8,7 @@ using System.Diagnostics;
 using System.Threading;
 using System.Reflection;
 using System.Linq;
-using Tools.DotNETCommon.ExecutingAssembly;
+using Tools.DotNETCommon;
 
 namespace UnrealBuildTool
 {
@@ -128,7 +128,7 @@ namespace UnrealBuildTool
 			List<string> DirectoriesToSearch = new List<string>();
 
 			// Find all the .uprojectdirs files contained in the root folder and add their entries to the search array
-			string RootDirectory = Path.Combine(ExecutingAssembly.GetDirectory(), "..", "..", "..");
+			string RootDirectory = Path.Combine( Path.GetDirectoryName(Assembly.GetEntryAssembly().GetOriginalLocation()), "..", "..", "..");
 			string EngineSourceDirectory = Path.GetFullPath( Path.Combine(RootDirectory, "Engine", "Source") );
 
 			foreach (var File in Directory.EnumerateFiles(RootDirectory, "*.uprojectdirs", SearchOption.TopDirectoryOnly))
diff --git a/Engine/Source/Programs/UnrealBuildTool/System/UnrealBuildTool.cs b/Engine/Source/Programs/UnrealBuildTool/System/UnrealBuildTool.cs
--- a/Engine/Source/Programs/UnrealBuildTool/System/UnrealBuildTool.cs
+++ b/Engine/Source/Programs/UnrealBuildTool/System/UnrealBuildTool.cs
@@ -9,7 +9,7 @@ using System.Threading;
 using System.Reflection;
 using System.Linq;
 using System.Runtime.Serialization.Formatters.Binary;
-using Tools.DotNETCommon.ExecutingAssembly;
+using Tools.DotNETCommon;
 
 namespace UnrealBuildTool
 {
@@ -192,7 +192,7 @@ namespace UnrealBuildTool
 
         static public string GetUBTPath()
         {
-            string UnrealBuildToolPath = Path.Combine(ExecutingAssembly.GetDirectory(), "UnrealBuildTool.exe");
+            string UnrealBuildToolPath = Assembly.GetExecutingAssembly().GetOriginalLocation();
             return UnrealBuildToolPath;
         }
 
@@ -634,20 +634,19 @@ namespace UnrealBuildTool
             // Change the working directory to be the Engine/Source folder. We are likely running from Engine/Binaries/DotNET
             // This is critical to be done early so any code that relies on the current directory being Engine/Source will work.
             // UEBuildConfiguration.PostReset is susceptible to this, so we must do this before configs are loaded.
-            string EngineSourceDirectory = Path.Combine(ExecutingAssembly.GetDirectory(), "..", "..", "..", "Engine", "Source");
+            string EngineSourceDirectory = Path.Combine(Path.GetDirectoryName(Assembly.GetEntryAssembly().GetOriginalLocation()), "..", "..", "..", "Engine", "Source");
 
             [REDACTED_EMAIL]: This is a workaround for recompiling game code in editor
             // The working directory when launching is *not* what we would expect
             if (Directory.Exists(EngineSourceDirectory) == false)
             {
-                // We are assuming UBT always runs from <>/Engine/...
-                EngineSourceDirectory = ExecutingAssembly.GetDirectory();
+                // We are assuming UBT always runs from <>/Engine/Binaries/DotNET/...
+                EngineSourceDirectory = Assembly.GetExecutingAssembly().GetOriginalLocation();
                 EngineSourceDirectory = EngineSourceDirectory.Replace("\\", "/");
-                Int32 EngineIdx = EngineSourceDirectory.IndexOf("/Engine/");
+                Int32 EngineIdx = EngineSourceDirectory.IndexOf("/Engine/Binaries/DotNET/", StringComparison.InvariantCultureIgnoreCase);
                 if (EngineIdx != 0)
                 {
-                    EngineSourceDirectory = EngineSourceDirectory.Substring(0, EngineIdx + 8);
-                    EngineSourceDirectory += "Source";
+                    EngineSourceDirectory = Path.Combine(EngineSourceDirectory.Substring(0, EngineIdx), "Engine", "Source");
                 }
             }
             if (Directory.Exists(EngineSourceDirectory)) // only set the directory if it exists, this should only happen if we are launching the editor from an artist sync
@@ -941,7 +940,7 @@ namespace UnrealBuildTool
                 Mutex SingleInstanceMutex = null;
                 if (bUseMutex)
                 {
-                    int LocationHash = ExecutingAssembly.GetFilename().GetHashCode();
+                    int LocationHash = Assembly.GetEntryAssembly().GetOriginalLocation().GetHashCode();
 
                     String MutexName;
                     if (bAutoSDKOnly || bValidatePlatforms)
@@ -957,7 +956,7 @@ namespace UnrealBuildTool
                     SingleInstanceMutex = new Mutex(true, MutexName, out bCreatedMutex);
                     if (!bCreatedMutex)
                     {
-                        throw new BuildException("Mutex {0} already set, indicating that a conflicting instance of {1} is already running.", MutexName, ExecutingAssembly.GetFilename());
+                        throw new BuildException("Mutex {0} already set, indicating that a conflicting instance of {1} is already running.", MutexName, Assembly.GetEntryAssembly().GetOriginalLocation());
                     }
                 }
 
diff --git a/Engine/Source/Programs/UnrealBuildTool/Utilities/XmlConfigLoader.cs b/Engine/Source/Programs/UnrealBuildTool/Utilities/XmlConfigLoader.cs
--- a/Engine/Source/Programs/UnrealBuildTool/Utilities/XmlConfigLoader.cs
+++ b/Engine/Source/Programs/UnrealBuildTool/Utilities/XmlConfigLoader.cs
@@ -13,7 +13,7 @@ using System.Threading;
 using System.Xml;
 using System.Xml.Linq;
 using System.Xml.Schema;
-using Tools.DotNETCommon.ExecutingAssembly;
+using Tools.DotNETCommon;
 
 namespace UnrealBuildTool
 {
@@ -502,7 +502,7 @@ namespace UnrealBuildTool
 			 *	priority. Not defined classes and fields are left alone.
 			 */
 
-			var UE4EnginePath = new FileInfo(Path.Combine(ExecutingAssembly.GetDirectory(), "..", "..")).FullName;
+			var UE4EnginePath = new FileInfo(Path.Combine(Path.GetDirectoryName(Assembly.GetEntryAssembly().GetOriginalLocation()), "..", "..")).FullName;
 
 			ConfigLocationHierarchy = new XmlConfigLocation[]
 			{
@@ -961,7 +961,7 @@ namespace UnrealBuildTool
 		/// <returns>Standard BuildConfiguration.xml schema path.</returns>
 		public static string GetXSDPath()
 		{
-			return new FileInfo(Path.Combine(ExecutingAssembly.GetDirectory(), "..", "..", "Saved", "UnrealBuildTool", "BuildConfiguration.Schema.xsd")).FullName;
+			return new FileInfo(Path.Combine(Path.GetDirectoryName(Assembly.GetEntryAssembly().GetOriginalLocation()), "..", "..", "Saved", "UnrealBuildTool", "BuildConfiguration.Schema.xsd")).FullName;
 		}
 
 		private static XElement CreateXSDElementForField(FieldInfo Field)
