Remove the ENGINE_VERSION define from Version.h and ObjectVersion.cpp; it's superceded by GEngineVersion. EngineVersion.cpp now defines the constants GEngineVersionChangelist and GCompatibleWithEngineVersionChangelist, which are used to construct GEngineVersion and GCompatibleWithEngineVersion. Modifying those constants allows updating the engine version separately to the defines used for resource files (allowing UGS to update the version number without triggering all files to be relinked).

diff --git a/Engine/Plugins/NetcodeUnitTest/NetcodeUnitTest/Source/NetcodeUnitTest/Private/Net/NUTUtilNet.cpp b/Engine/Plugins/NetcodeUnitTest/NetcodeUnitTest/Source/NetcodeUnitTest/Private/Net/NUTUtilNet.cpp
--- a/Engine/Plugins/NetcodeUnitTest/NetcodeUnitTest/Source/NetcodeUnitTest/Private/Net/NUTUtilNet.cpp
+++ b/Engine/Plugins/NetcodeUnitTest/NetcodeUnitTest/Source/NetcodeUnitTest/Private/Net/NUTUtilNet.cpp
@@ -374,8 +374,10 @@ bool NUTNet::CreateFakePlayer(UWorld* InWorld, UNetDriver*& InNetDriver, FString
 			if (bOldProtocol)
 #endif
 			{
-				*ControlChanBunch << GEngineMinNetVersion;
-				*ControlChanBunch << GEngineNetVersion;
+				int32 EngineMinNetVersion = GEngineMinNetVersion;
+				*ControlChanBunch << EngineMinNetVersion;
+				int32 EngineNetVersion = GEngineNetVersion;
+				*ControlChanBunch << EngineNetVersion;
 				*ControlChanBunch << (FGuid&)GetDefault<UGeneralProjectSettings>()->ProjectID;
 			}
 #if TARGET_UE4_CL >= CL_FNETWORKVERSION
diff --git a/Engine/Source/Developer/DesktopPlatform/Private/DesktopPlatformBase.cpp b/Engine/Source/Developer/DesktopPlatform/Private/DesktopPlatformBase.cpp
--- a/Engine/Source/Developer/DesktopPlatform/Private/DesktopPlatformBase.cpp
+++ b/Engine/Source/Developer/DesktopPlatform/Private/DesktopPlatformBase.cpp
@@ -313,7 +313,7 @@ bool FDesktopPlatformBase::TryGetEngineVersion(const FString& RootDir, FEngineVe
 					{
 						PatchVersion = FCString::Atoi(*Tokens[3]);
 					}
-					else if(Tokens[2] == "ENGINE_VERSION")
+					else if(Tokens[2] == "BUILT_FROM_CHANGELIST")
 					{
 						Changelist = FCString::Atoi(*Tokens[3]);
 					}
diff --git a/Engine/Source/Programs/AutomationTool/AutomationUtils/UE4Build.cs b/Engine/Source/Programs/AutomationTool/AutomationUtils/UE4Build.cs
--- a/Engine/Source/Programs/AutomationTool/AutomationUtils/UE4Build.cs
+++ b/Engine/Source/Programs/AutomationTool/AutomationUtils/UE4Build.cs
@@ -489,23 +489,6 @@ namespace AutomationTool
 				}
 				Result.Add(VerFile);
 			}
-			{
-				string VerFile = CombinePaths(CmdEnv.LocalRoot, "Engine", "Source", "Runtime", "Core", "Private", "UObject", "ObjectVersion.cpp");
-				if (bDoUpdateVersionFiles)
-				{
-					Log("Updating {0} with:", VerFile);
-					Log(" #define	ENGINE_VERSION  {0}", ChangelistNumber.ToString());
-
-					VersionFileUpdater ObjectVersionCpp = new VersionFileUpdater(VerFile);
-					ObjectVersionCpp.ReplaceLine("#define	ENGINE_VERSION	", ChangelistNumber.ToString());
-                    ObjectVersionCpp.Commit();
-                }
-				else
-				{
-					Log("{0} will not be updated because P4 is not enabled.", VerFile);
-				}
-                Result.Add(VerFile);
-            }
 
             string EngineVersionFile = CombinePaths(CmdEnv.LocalRoot, "Engine", "Source", "Runtime", "Launch", "Resources", "Version.h");
             {
@@ -513,13 +496,11 @@ namespace AutomationTool
 				if (bDoUpdateVersionFiles)
 				{
 					Log("Updating {0} with:", VerFile);
-					Log(" #define	ENGINE_VERSION  {0}", ChangelistNumber.ToString());
 					Log(" #define	BRANCH_NAME  {0}", Branch);
 					Log(" #define	BUILT_FROM_CHANGELIST  {0}", ChangelistString);
 					Log(" #define   ENGINE_IS_LICENSEE_VERSION  {0}", bIsLicenseeVersion ? "1" : "0");
 
 					VersionFileUpdater VersionH = new VersionFileUpdater(VerFile);
-					VersionH.ReplaceLine("#define ENGINE_VERSION ", ChangelistNumber.ToString());
 					VersionH.ReplaceLine("#define BRANCH_NAME ", "\"" + Branch + "\"");
 					VersionH.ReplaceLine("#define BUILT_FROM_CHANGELIST ", ChangelistString);
 					VersionH.ReplaceOrAddLine("#define ENGINE_IS_LICENSEE_VERSION ", bIsLicenseeVersion ? "1" : "0");
diff --git a/Engine/Source/Programs/AutomationTool/Scripts/Tests.Automation.cs b/Engine/Source/Programs/AutomationTool/Scripts/Tests.Automation.cs
--- a/Engine/Source/Programs/AutomationTool/Scripts/Tests.Automation.cs
+++ b/Engine/Source/Programs/AutomationTool/Scripts/Tests.Automation.cs
@@ -1328,17 +1328,10 @@ public class ZeroEngineVersions : BuildCommand
 {
 	public override void ExecuteBuild()
 	{
-		string ObjectVersionFilename = CmdEnv.LocalRoot + @"/Engine/Source/Runtime/Core/Private/UObject/ObjectVersion.cpp";
 		string VersionFilename = CmdEnv.LocalRoot + @"/Engine/Source/Runtime/Launch/Resources/Version.h";
 		if (P4Env.Changelist > 0)
 		{
-			var Stat = P4.FStat(ObjectVersionFilename);
-			if (Stat.IsValid && Stat.Action != P4Action.None)
-			{
-				Log("Reverting {0}", ObjectVersionFilename);
-				P4.Revert(ObjectVersionFilename);
-			}
-			Stat = P4.FStat(VersionFilename);
+			var Stat = P4.FStat(VersionFilename);
 			if (Stat.IsValid && Stat.Action != P4Action.None)
 			{
 				Log("Reverting {0}", VersionFilename);
@@ -1346,7 +1339,6 @@ public class ZeroEngineVersions : BuildCommand
 			}
 
 			Log("Gettting engine version files @{0}", P4Env.Changelist);
-			P4.Sync(String.Format("-f {0}@{1}", ObjectVersionFilename, P4Env.Changelist));
 			P4.Sync(String.Format("-f {0}@{1}", VersionFilename, P4Env.Changelist));
 		}
 
diff --git a/Engine/Source/Programs/UnrealBuildTool/Mac/MacToolChain.cs b/Engine/Source/Programs/UnrealBuildTool/Mac/MacToolChain.cs
--- a/Engine/Source/Programs/UnrealBuildTool/Mac/MacToolChain.cs
+++ b/Engine/Source/Programs/UnrealBuildTool/Mac/MacToolChain.cs
@@ -424,15 +424,15 @@ namespace UnrealBuildTool
 
 		private int LoadEngineCL()
 		{
-			string[] VersionHeader = Utils.ReadAllText("../Source/Runtime/Launch/Resources/Version.h").Replace("\r\n", "\n").Replace("\t", " ").Split('\n');
-			foreach (string Line in VersionHeader)
+			BuildVersion Version;
+			if(Utils.TryReadBuildVersion("../Build/Build.version", out Version))
 			{
-				if (Line.StartsWith("#define ENGINE_VERSION "))
-				{
-					return int.Parse(Line.Split(' ')[2]);
-				}
+				return Version.Changelist;
+			}
+			else
+			{
+				return 0;
 			}
-			return 0;
 		}
 
 		private string LoadEngineDisplayVersion(bool bIgnorePatchVersion = false)
diff --git a/Engine/Source/Programs/UnrealBuildTool/System/UnrealBuildTool.cs b/Engine/Source/Programs/UnrealBuildTool/System/UnrealBuildTool.cs
--- a/Engine/Source/Programs/UnrealBuildTool/System/UnrealBuildTool.cs
+++ b/Engine/Source/Programs/UnrealBuildTool/System/UnrealBuildTool.cs
@@ -1155,6 +1155,13 @@ namespace UnrealBuildTool
 						throw new BuildException("When running UnrealCodeAnalyzer, please specify module to analyze in UCAModuleToAnalyze field in BuildConfiguration.xml");
 					}
 
+					// Get the build version
+					BuildVersion Version;
+					if(!Utils.TryReadBuildVersion("../Build/Build.version", out Version))
+					{
+						Log.TraceWarning("Could not read build.version");
+					}
+
                     // Send an event with basic usage dimensions
                     // [RCL] 2014-11-03 FIXME: this is incorrect since we can have more than one action
                     Telemetry.SendEvent("CommonAttributes.2",
@@ -1185,7 +1192,7 @@ namespace UnrealBuildTool
                             : "Build",
                         "Platform", CheckPlatform.ToString(),
                         "Configuration", CheckConfiguration.ToString(),
-                        "EngineVersion", Utils.GetEngineVersionFromObjVersionCPP().ToString()
+                        "EngineVersion", (Version == null)? "0" : Version.Changelist.ToString()
                         );
 
 
diff --git a/Engine/Source/Programs/UnrealBuildTool/System/Utils.cs b/Engine/Source/Programs/UnrealBuildTool/System/Utils.cs
--- a/Engine/Source/Programs/UnrealBuildTool/System/Utils.cs
+++ b/Engine/Source/Programs/UnrealBuildTool/System/Utils.cs
@@ -13,10 +13,27 @@ using System.Runtime.InteropServices;
 using System.Linq;
 using Tools.DotNETCommon.CaselessDictionary;
 using Tools.DotNETCommon.HarvestEnvVars;
+using System.Web.Script.Serialization;
 
 namespace UnrealBuildTool
 {
-	public abstract class Utils
+	/// <summary>
+	/// Holds information about the current engine version
+	/// </summary>
+	public class BuildVersion
+	{
+		public int MajorVersion;
+		public int MinorVersion;
+		public int PatchVersion;
+		public int Changelist;
+		public int IsLicenseeVersion;
+		public string BranchName;
+	}
+
+	/// <summary>
+	/// Utility functions
+	/// </summary>
+	public static class Utils
 	{
 		/// <summary>
 		/// Whether we are currently running on Mono platform.  We cache this statically because it is a bit slow to check.
@@ -790,29 +807,37 @@ namespace UnrealBuildTool
 			return true;
 		}
 
-
 		/// <summary>
-		/// Finds the Engine Version from ObjVersion.cpp.
+		/// Try to read the Build/Build.version file from disk
 		/// </summary>
-		/// <remarks>
-		/// If UBT eventually gets the engine version passed from the build environment this code should be scrapped for that instead.
-		/// </remarks>
-		/// <returns></returns>
-		public static int GetEngineVersionFromObjVersionCPP()
+		/// <param name="Version">The version information</param>
+		/// <returns>True if the version was read sucessfully, false otherwise</returns>
+		public static bool TryReadBuildVersion(string FileName, out BuildVersion Version)
 		{
+			JsonObject Object;
 			try
 			{
-				return
-					(from line in File.ReadLines("Runtime/Core/Private/UObject/ObjectVersion.cpp", Encoding.ASCII)
-						where line.StartsWith("#define	ENGINE_VERSION")
-						select int.Parse(line.Split()[2])).Single();
+				Object = JsonObject.FromFile(FileName);
 			}
-			catch (Exception ex)
+			catch(Exception)
 			{
-				// Don't do a stack trace so we don't pollute the logs with spurious exception data, as we don't crash on this case.
-				Log.TraceWarning("Could not parse Engine Version from ObjectVersion.cpp: {0}", ex.Message);
-				return 0;
+				Version = null;
+				return false;
+			}
+
+			BuildVersion NewVersion = new BuildVersion();
+			if(!Object.TryGetIntegerField("MajorVersion", out NewVersion.MajorVersion) || !Object.TryGetIntegerField("MinorVersion", out NewVersion.MinorVersion) || !Object.TryGetIntegerField("PatchVersion", out NewVersion.PatchVersion))
+			{
+				Version = null;
+				return false;
 			}
+
+			Object.TryGetIntegerField("Changelist", out NewVersion.Changelist);
+			Object.TryGetIntegerField("IsLicenseeVersion", out NewVersion.IsLicenseeVersion);
+			Object.TryGetStringField("BranchName", out NewVersion.BranchName);
+
+			Version = NewVersion;
+			return true;
 		}
 
 		/// <summary>
diff --git a/Engine/Source/Runtime/Core/Private/Misc/EngineVersion.cpp b/Engine/Source/Runtime/Core/Private/Misc/EngineVersion.cpp
--- a/Engine/Source/Runtime/Core/Private/Misc/EngineVersion.cpp
+++ b/Engine/Source/Runtime/Core/Private/Misc/EngineVersion.cpp
@@ -192,16 +192,17 @@ void operator<<(FArchive &Ar, FEngineVersion &Version)
 	Ar << Version.Branch;
 }
 
-// Make sure this is defined
-#ifndef ENGINE_IS_LICENSEE_VERSION
-	#define ENGINE_IS_LICENSEE_VERSION 0
-#endif
+/** Changelist constants. The two lines below may be modified by tools. */
+const int32 GEngineVersionChangelist = BUILT_FROM_CHANGELIST;
+const int32 GCompatibleWithEngineVersionChangelist = MODULE_API_VERSION;
 
-// Licensee changelists part of the engine version has the top bit set to 1
-#define ENGINE_VERSION_INTERNAL_OR_LICENSEE (BUILT_FROM_CHANGELIST | (ENGINE_IS_LICENSEE_VERSION << 31))
+/** Version numbers for networking */
+const int32 GEngineNetVersion			= GEngineVersionChangelist;
+const int32 GEngineMinNetVersion		= 7038;
+const int32 GEngineNegotiationVersion	= 3077;
 
 // Global instance of the current engine version
-const FEngineVersion GEngineVersion(ENGINE_MAJOR_VERSION, ENGINE_MINOR_VERSION, ENGINE_PATCH_VERSION, ENGINE_VERSION_INTERNAL_OR_LICENSEE, BRANCH_NAME);
+const FEngineVersion GEngineVersion(ENGINE_MAJOR_VERSION, ENGINE_MINOR_VERSION, ENGINE_PATCH_VERSION, (GEngineVersionChangelist | (ENGINE_IS_LICENSEE_VERSION << 31)), BRANCH_NAME);
 
 // Version which this engine maintains strict API and package compatibility with
-const FEngineVersion GCompatibleWithEngineVersion(ENGINE_MAJOR_VERSION, ENGINE_MINOR_VERSION, ENGINE_PATCH_VERSION, (MODULE_API_VERSION | (ENGINE_IS_LICENSEE_VERSION << 31)), BRANCH_NAME);
+const FEngineVersion GCompatibleWithEngineVersion(ENGINE_MAJOR_VERSION, ENGINE_MINOR_VERSION, ENGINE_PATCH_VERSION, (GCompatibleWithEngineVersionChangelist | (ENGINE_IS_LICENSEE_VERSION << 31)), BRANCH_NAME);
diff --git a/Engine/Source/Runtime/Core/Private/UObject/ObjectVersion.cpp b/Engine/Source/Runtime/Core/Private/UObject/ObjectVersion.cpp
--- a/Engine/Source/Runtime/Core/Private/UObject/ObjectVersion.cpp
+++ b/Engine/Source/Runtime/Core/Private/UObject/ObjectVersion.cpp
@@ -6,15 +6,6 @@
 
 #include "CorePrivatePCH.h"
 
-// Defined separately so the build script can get to it easily (DO NOT CHANGE THIS MANUALLY)
-// caution when merging, we want the UE3 version in ENGINE_VERSION_UE3
-#define	ENGINE_VERSION	0
-
-/** Version numbers for networking */
-int32   GEngineNetVersion			= ENGINE_VERSION;
-int32	GEngineMinNetVersion		= 7038;
-int32	GEngineNegotiationVersion	= 3077;
-
 // @see ObjectVersion.h for the list of changes/defines
-int32	GPackageFileUE4Version = VER_LATEST_ENGINE_UE4;
-int32	GPackageFileLicenseeUE4Version = VER_LATEST_ENGINE_LICENSEEUE4;
+const int32 GPackageFileUE4Version = VER_LATEST_ENGINE_UE4;
+const int32 GPackageFileLicenseeUE4Version = VER_LATEST_ENGINE_LICENSEEUE4;
diff --git a/Engine/Source/Runtime/Core/Public/Misc/EngineVersion.h b/Engine/Source/Runtime/Core/Public/Misc/EngineVersion.h
--- a/Engine/Source/Runtime/Core/Public/Misc/EngineVersion.h
+++ b/Engine/Source/Runtime/Core/Public/Misc/EngineVersion.h
@@ -46,6 +46,14 @@ private:
 	FString Branch;
 };
 
+/** Version used for networking; the P4 changelist number. */
+CORE_API extern const int32 GEngineNetVersion;
+
+/** Earliest engine build that is network compatible with this one. */
+CORE_API extern const int32 GEngineMinNetVersion;
+
+/** Base protocol version to negotiate in network play. */
+CORE_API extern const int32 GEngineNegotiationVersion;		
 
 /** Global instance of the current engine version. */
 CORE_API extern const FEngineVersion GEngineVersion;
diff --git a/Engine/Source/Runtime/Core/Public/UObject/ObjectVersion.h b/Engine/Source/Runtime/Core/Public/UObject/ObjectVersion.h
--- a/Engine/Source/Runtime/Core/Public/UObject/ObjectVersion.h
+++ b/Engine/Source/Runtime/Core/Public/UObject/ObjectVersion.h
@@ -586,9 +586,5 @@ enum EUnrealEngineObjectLicenseeUE4Version
 
 // Version access.
 
-extern CORE_API int32			GEngineNetVersion;				// Version used for networking; the P4 changelist number.
-extern CORE_API int32			GEngineMinNetVersion;			// Earliest engine build that is network compatible with this one.
-extern CORE_API int32			GEngineNegotiationVersion;		// Base protocol version to negotiate in network play.
-
-extern CORE_API int32			GPackageFileUE4Version;			// UE4 Version Number.
-extern CORE_API int32			GPackageFileLicenseeUE4Version;	// Licensee Version Number.
+extern const CORE_API int32			GPackageFileUE4Version;			// UE4 Version Number.
+extern const CORE_API int32			GPackageFileLicenseeUE4Version;	// Licensee Version Number.
diff --git a/Engine/Source/Runtime/Engine/Private/GameplayStatics.cpp b/Engine/Source/Runtime/Engine/Private/GameplayStatics.cpp
--- a/Engine/Source/Runtime/Engine/Private/GameplayStatics.cpp
+++ b/Engine/Source/Runtime/Engine/Private/GameplayStatics.cpp
@@ -1239,7 +1239,8 @@ bool UGameplayStatics::SaveGameToSlot(USaveGame* SaveGameObject, const FString&
 		MemoryWriter << SavegameFileVersion;
 
 		// Write out engine and UE4 version information
-		MemoryWriter << GPackageFileUE4Version;
+		int32 PackageFileUE4Version = GPackageFileUE4Version;
+		MemoryWriter << PackageFileUE4Version;
 		FEngineVersion SavedEngineVersion = GEngineVersion;
 		MemoryWriter << SavedEngineVersion;
 
diff --git a/Engine/Source/Runtime/Launch/Resources/Version.h b/Engine/Source/Runtime/Launch/Resources/Version.h
--- a/Engine/Source/Runtime/Launch/Resources/Version.h
+++ b/Engine/Source/Runtime/Launch/Resources/Version.h
@@ -6,15 +6,13 @@
 #define ENGINE_MINOR_VERSION	9
 #define ENGINE_PATCH_VERSION	0
 
-#define ENGINE_VERSION 0
-
 #define ENGINE_IS_LICENSEE_VERSION 0
 
-#define ENGINE_IS_PROMOTED_BUILD (ENGINE_VERSION > 0)
-
 #define BUILT_FROM_CHANGELIST 0
 #define BRANCH_NAME "UE4"
 
+#define ENGINE_IS_PROMOTED_BUILD (BUILT_FROM_CHANGELIST > 0)
+
 #define EPIC_COMPANY_NAME  "Epic Games, Inc."
 #define EPIC_COPYRIGHT_STRING "Copyright 1998-2015 Epic Games, Inc. All Rights Reserved."
 #define EPIC_PRODUCT_NAME "Unreal Engine"
diff --git a/Engine/Source/Runtime/NetworkFileSystem/Private/NetworkFileServerConnection.cpp b/Engine/Source/Runtime/NetworkFileSystem/Private/NetworkFileServerConnection.cpp
--- a/Engine/Source/Runtime/NetworkFileSystem/Private/NetworkFileServerConnection.cpp
+++ b/Engine/Source/Runtime/NetworkFileSystem/Private/NetworkFileServerConnection.cpp
@@ -813,8 +813,10 @@ bool FNetworkFileServerClientConnection::ProcessGetFileList( FArchive& In, FArch
 
 	// report the package version information
 	// The downside of this is that ALL cooked data will get tossed on package version changes
-	Out << GPackageFileUE4Version;
-	Out << GPackageFileLicenseeUE4Version;
+	int32 PackageFileUE4Version = GPackageFileUE4Version;
+	Out << PackageFileUE4Version;
+	int32 PackageFileLicenseeUE4Version = GPackageFileLicenseeUE4Version;
+	Out << PackageFileLicenseeUE4Version;
 
 	// Send *our* engine and game dirs
 	Out << LocalEngineDir;
