Corrected issue with trying to write over read-only libvrapi.so file in intermediate directory on Android
#ue4
#android

diff --git a/Engine/Source/Programs/UnrealBuildTool/Android/UEDeployAndroid.cs b/Engine/Source/Programs/UnrealBuildTool/Android/UEDeployAndroid.cs
--- a/Engine/Source/Programs/UnrealBuildTool/Android/UEDeployAndroid.cs
+++ b/Engine/Source/Programs/UnrealBuildTool/Android/UEDeployAndroid.cs
@@ -1614,23 +1614,28 @@ namespace UnrealBuildTool.Android
             // Note: only support ARMv7 at the moment
             if (bPackageForGearVR && Arch == "armeabi-v7a")
             {
-                string VRLibFile = ThirdPartyDir + "/Oculus/LibOVRMobile/LibOVRMobile_060/VrApi/Libs/Android/armeabi-v7a/libvrapi.so";
+                string SourceVRLibFile = ThirdPartyDir + "/Oculus/LibOVRMobile/LibOVRMobile_060/VrApi/Libs/Android/armeabi-v7a/libvrapi.so";
+				string FinalVRLibFile = UE4BuildPath + "/libs/" + Arch + "/libvrapi.so";
 
-                if (File.Exists(VRLibFile))
-                {
-                    Console.WriteLine("File {0} found", VRLibFile);
-
-                    Console.WriteLine("Dir {0} created", UE4BuildPath + "/libs/" + Arch);
-
-                    Directory.CreateDirectory(UE4BuildPath + "/libs/" + Arch);
-                    File.Copy(VRLibFile, UE4BuildPath + "/libs/" + Arch + "/libvrapi.so", true);
-
-                    Console.WriteLine("File {0} copied to {1}", VRLibFile, UE4BuildPath + "/libs/" + Arch + "/libvrapi.so");
-                    
+				if (File.Exists(SourceVRLibFile))
+				{
+					// check to see if newer than last time we copied
+					bool bFileExists = File.Exists(FinalVRLibFile);
+					TimeSpan Diff = File.GetLastWriteTimeUtc(FinalVRLibFile) - File.GetLastWriteTimeUtc(SourceVRLibFile);
+					if (!bFileExists || Diff.TotalSeconds < -1 || Diff.TotalSeconds > 1)
+					{
+						if (bFileExists)
+						{
+							File.Delete(FinalVRLibFile);
+						}
+						Directory.CreateDirectory(Path.GetDirectoryName(FinalVRLibFile));
+						File.Copy(SourceVRLibFile, FinalVRLibFile, true);
+	                    Console.WriteLine("File {0} copied to {1}", SourceVRLibFile, FinalVRLibFile);
+					}
                 }
 				else
 				{
-					Console.WriteLine("Failed to find the GearVR library required for packaging: {0}", VRLibFile);
+					Console.WriteLine("Failed to find the GearVR library required for packaging: {0}", SourceVRLibFile);
 				}
             }
 
