Fixed crash that could occur if source folder of a plugin was removed

GetModuleFilename returns an empty string if module can't be found, checking for null meant that it went ahead and tried to load it regardless.

diff --git a/Engine/Source/Programs/UnrealBuildTool/System/RulesCompiler.cs b/Engine/Source/Programs/UnrealBuildTool/System/RulesCompiler.cs
--- a/Engine/Source/Programs/UnrealBuildTool/System/RulesCompiler.cs
+++ b/Engine/Source/Programs/UnrealBuildTool/System/RulesCompiler.cs
@@ -1774,7 +1774,7 @@ namespace UnrealBuildTool
 		/// <returns>Compiled module rule info</returns>
 		public static bool TryCreateModuleRules( string ModuleName, TargetInfo Target, out ModuleRules Rules )
 		{
-			if(GetModuleFilename( ModuleName ) == null)
+			if(GetModuleFilename( ModuleName ) == "")
 			{
 				Rules = null;
 				return false;
