UBT: fix a number of project generators (UE-18832)

- In-tree code project targets were missing.

diff --git a/Engine/Source/Programs/UnrealBuildTool/System/CMakefileGenerator.cs b/Engine/Source/Programs/UnrealBuildTool/System/CMakefileGenerator.cs
--- a/Engine/Source/Programs/UnrealBuildTool/System/CMakefileGenerator.cs
+++ b/Engine/Source/Programs/UnrealBuildTool/System/CMakefileGenerator.cs
@@ -267,33 +267,41 @@ namespace UnrealBuildTool
 
 			string CMakeProjectCmdArg = "";
 
-			foreach (string TargetFilePath in DiscoverTargets())
+			foreach(var Project in GeneratedProjectFiles)
 			{
-				var TargetName = Utils.GetFilenameWithoutAnyExtensions(TargetFilePath);		// Remove both ".cs" and ".
-
-				foreach (UnrealTargetConfiguration CurConfiguration in Enum.GetValues(typeof(UnrealTargetConfiguration)))
+				foreach (var TargetFile in Project.ProjectTargets)
 				{
-					if (CurConfiguration != UnrealTargetConfiguration.Unknown && CurConfiguration != UnrealTargetConfiguration.Development)
+					if (TargetFile.TargetFilePath == null)
+					{
+						continue;
+					}
+
+					var TargetName = Utils.GetFilenameWithoutAnyExtensions(TargetFile.TargetFilePath);		// Remove both ".cs" and ".
+
+					foreach (UnrealTargetConfiguration CurConfiguration in Enum.GetValues(typeof(UnrealTargetConfiguration)))
 					{
-						if (UnrealBuildTool.IsValidConfiguration(CurConfiguration))
+						if (CurConfiguration != UnrealTargetConfiguration.Unknown && CurConfiguration != UnrealTargetConfiguration.Development)
 						{
-							if (TargetName == GameProjectName || TargetName == (GameProjectName + "Editor")) 
+							if (UnrealBuildTool.IsValidConfiguration(CurConfiguration))
 							{
-								CMakeProjectCmdArg = " -project=\"\\\"${GAME_PROJECT_FILE}\\\"\"";
+								if (TargetName == GameProjectName || TargetName == (GameProjectName + "Editor")) 
+								{
+									CMakeProjectCmdArg = " -project=\"\\\"${GAME_PROJECT_FILE}\\\"\"";
+								}
+								var ConfName = Enum.GetName(typeof(UnrealTargetConfiguration), CurConfiguration);
+								CMakefileContent.Append(String.Format("add_custom_target({0}-{3}-{1} ${{BUILD}} {2} {0} {3} {1} $(ARGS))\n", TargetName, ConfName, CMakeProjectCmdArg, HostArchitecture));
 							}
-							var ConfName = Enum.GetName(typeof(UnrealTargetConfiguration), CurConfiguration);
-							CMakefileContent.Append(String.Format("add_custom_target({0}-{3}-{1} ${{BUILD}} {2} {0} {3} {1} $(ARGS))\n", TargetName, ConfName, CMakeProjectCmdArg, HostArchitecture));
 						}
 					}
-				}
 
-				if (TargetName == GameProjectName || TargetName == (GameProjectName + "Editor")) 
-				{
-					CMakeProjectCmdArg = " -project=\"\\\"${GAME_PROJECT_FILE}\\\"\"";
-				}
-				if (HostArchitecture != null) 
-				{
-					CMakefileContent.Append (String.Format ("add_custom_target({0} ${{BUILD}} {1} {0} {2} Development $(ARGS) SOURCES ${{SOURCE_FILES}} ${{HEADER_FILES}} ${{CONFIG_FILES}})\n\n", TargetName, CMakeProjectCmdArg, HostArchitecture));
+					if (TargetName == GameProjectName || TargetName == (GameProjectName + "Editor")) 
+					{
+						CMakeProjectCmdArg = " -project=\"\\\"${GAME_PROJECT_FILE}\\\"\"";
+					}
+					if (HostArchitecture != null) 
+					{
+						CMakefileContent.Append (String.Format ("add_custom_target({0} ${{BUILD}} {1} {0} {2} Development $(ARGS) SOURCES ${{SOURCE_FILES}} ${{HEADER_FILES}} ${{CONFIG_FILES}})\n\n", TargetName, CMakeProjectCmdArg, HostArchitecture));
+					}
 				}
 			}
 
diff --git a/Engine/Source/Programs/UnrealBuildTool/System/KDevelopGenerator.cs b/Engine/Source/Programs/UnrealBuildTool/System/KDevelopGenerator.cs
--- a/Engine/Source/Programs/UnrealBuildTool/System/KDevelopGenerator.cs
+++ b/Engine/Source/Programs/UnrealBuildTool/System/KDevelopGenerator.cs
@@ -177,71 +177,79 @@ namespace UnrealBuildTool
 			FileContent.Append ("Executable=bash\n");
 			FileContent.Append ("Type=1\n\n");
 
-			foreach (string TargetFilePath in DiscoverTargets())
+			foreach(var Project in GeneratedProjectFiles)
 			{
-				var TargetName = Utils.GetFilenameWithoutAnyExtensions(TargetFilePath);		
-
-				// Remove both ".cs" and ".
-				foreach (UnrealTargetConfiguration CurConfiguration in Enum.GetValues(typeof(UnrealTargetConfiguration))) 
+				foreach (var TargetFile in Project.ProjectTargets)
 				{
-					if (CurConfiguration != UnrealTargetConfiguration.Unknown && CurConfiguration != UnrealTargetConfiguration.Development) 
+					if (TargetFile.TargetFilePath == null)
 					{
-						if (UnrealBuildTool.IsValidConfiguration (CurConfiguration)) 
-						{
-							var ConfName = Enum.GetName(typeof(UnrealTargetConfiguration), CurConfiguration);
-							FileContent.Append (String.Format("[CustomBuildSystem][BuildConfig{0}]\nBuildDir=file://{1}\n", BuildConfigIndex, UnrealRootPath ));
+						continue;
+					}
 
-							if (TargetName == GameProjectName)
-							{
-								FileContent.Append (String.Format ("Title={0}-Linux-{1}\n\n", TargetName, ConfName));
-								WriteCommandSubSection (ref FileContent, TargetName, ConfName, BuildConfigIndex, 0);
-								WriteCommandSubSection (ref FileContent, TargetName, ConfName, BuildConfigIndex, 1);
-								WriteCommandSubSection (ref FileContent, TargetName, ConfName, BuildConfigIndex, 3);
-							} 
-							else if (TargetName == (GameProjectName + "Editor")) 
-							{
-								FileContent.Append (String.Format ("Title={0}-Linux-{1}\n\n", TargetName, ConfName));
-								WriteCommandSubSection (ref FileContent, TargetName, ConfName, BuildConfigIndex, 0);
-								WriteCommandSubSection (ref FileContent, TargetName, ConfName, BuildConfigIndex, 1);
-								WriteCommandSubSection (ref FileContent, TargetName, ConfName, BuildConfigIndex, 3);
-							} 
-							else
+					var TargetName = Utils.GetFilenameWithoutAnyExtensions(TargetFile.TargetFilePath);
+
+					// Remove both ".cs" and ".
+					foreach (UnrealTargetConfiguration CurConfiguration in Enum.GetValues(typeof(UnrealTargetConfiguration))) 
+					{
+						if (CurConfiguration != UnrealTargetConfiguration.Unknown && CurConfiguration != UnrealTargetConfiguration.Development) 
+						{
+							if (UnrealBuildTool.IsValidConfiguration (CurConfiguration)) 
 							{
-								FileContent.Append (String.Format ("Title={0}-Linux-{1}\n\n", TargetName, ConfName));
-								WriteCommandSubSection (ref FileContent, TargetName, ConfName, BuildConfigIndex, 0);
-								WriteCommandSubSection (ref FileContent, TargetName, ConfName, BuildConfigIndex, 1);
-								WriteCommandSubSection (ref FileContent, TargetName, ConfName, BuildConfigIndex, 3);
+								var ConfName = Enum.GetName(typeof(UnrealTargetConfiguration), CurConfiguration);
+								FileContent.Append (String.Format("[CustomBuildSystem][BuildConfig{0}]\nBuildDir=file://{1}\n", BuildConfigIndex, UnrealRootPath ));
+
+								if (TargetName == GameProjectName)
+								{
+									FileContent.Append (String.Format ("Title={0}-Linux-{1}\n\n", TargetName, ConfName));
+									WriteCommandSubSection (ref FileContent, TargetName, ConfName, BuildConfigIndex, 0);
+									WriteCommandSubSection (ref FileContent, TargetName, ConfName, BuildConfigIndex, 1);
+									WriteCommandSubSection (ref FileContent, TargetName, ConfName, BuildConfigIndex, 3);
+								} 
+								else if (TargetName == (GameProjectName + "Editor")) 
+								{
+									FileContent.Append (String.Format ("Title={0}-Linux-{1}\n\n", TargetName, ConfName));
+									WriteCommandSubSection (ref FileContent, TargetName, ConfName, BuildConfigIndex, 0);
+									WriteCommandSubSection (ref FileContent, TargetName, ConfName, BuildConfigIndex, 1);
+									WriteCommandSubSection (ref FileContent, TargetName, ConfName, BuildConfigIndex, 3);
+								} 
+								else
+								{
+									FileContent.Append (String.Format ("Title={0}-Linux-{1}\n\n", TargetName, ConfName));
+									WriteCommandSubSection (ref FileContent, TargetName, ConfName, BuildConfigIndex, 0);
+									WriteCommandSubSection (ref FileContent, TargetName, ConfName, BuildConfigIndex, 1);
+									WriteCommandSubSection (ref FileContent, TargetName, ConfName, BuildConfigIndex, 3);
+								}
+								BuildConfigIndex++;
 							}
-							BuildConfigIndex++;
 						}
 					}
-				}
 
-				FileContent.Append (String.Format("[CustomBuildSystem][BuildConfig{0}]\nBuildDir=file://{1}\n", BuildConfigIndex, UnrealRootPath ));
-				if (TargetName == GameProjectName)
-				{
-					FileContent.Append (String.Format ("Title={0}\n\n", TargetName));
-					WriteCommandSubSection (ref FileContent, TargetName, "Development" , BuildConfigIndex, 0);
-					WriteCommandSubSection (ref FileContent, TargetName, "Development" , BuildConfigIndex, 1);
-					WriteCommandSubSection (ref FileContent, TargetName, "Development" , BuildConfigIndex, 3);
+					FileContent.Append (String.Format("[CustomBuildSystem][BuildConfig{0}]\nBuildDir=file://{1}\n", BuildConfigIndex, UnrealRootPath ));
+					if (TargetName == GameProjectName)
+					{
+						FileContent.Append (String.Format ("Title={0}\n\n", TargetName));
+						WriteCommandSubSection (ref FileContent, TargetName, "Development" , BuildConfigIndex, 0);
+						WriteCommandSubSection (ref FileContent, TargetName, "Development" , BuildConfigIndex, 1);
+						WriteCommandSubSection (ref FileContent, TargetName, "Development" , BuildConfigIndex, 3);
 
-				} 
-				else if (TargetName == (GameProjectName + "Editor")) 
-				{
-					FileContent.Append (String.Format ("Title={0}\n\n", TargetName));
-					WriteCommandSubSection (ref FileContent, TargetName, "Development" , BuildConfigIndex, 0);
-					WriteCommandSubSection (ref FileContent, TargetName, "Development" , BuildConfigIndex, 1);
-					WriteCommandSubSection (ref FileContent, TargetName, "Development" , BuildConfigIndex, 3);
-				}
-				else
-				{
-					FileContent.Append (String.Format ("Title={0}\n\n", TargetName));
-					WriteCommandSubSection (ref FileContent, TargetName, "Development" , BuildConfigIndex, 0);
-					WriteCommandSubSection (ref FileContent, TargetName, "Development" , BuildConfigIndex, 1);
-					WriteCommandSubSection (ref FileContent, TargetName, "Development" , BuildConfigIndex, 3);
+					} 
+					else if (TargetName == (GameProjectName + "Editor")) 
+					{
+						FileContent.Append (String.Format ("Title={0}\n\n", TargetName));
+						WriteCommandSubSection (ref FileContent, TargetName, "Development" , BuildConfigIndex, 0);
+						WriteCommandSubSection (ref FileContent, TargetName, "Development" , BuildConfigIndex, 1);
+						WriteCommandSubSection (ref FileContent, TargetName, "Development" , BuildConfigIndex, 3);
+					}
+					else
+					{
+						FileContent.Append (String.Format ("Title={0}\n\n", TargetName));
+						WriteCommandSubSection (ref FileContent, TargetName, "Development" , BuildConfigIndex, 0);
+						WriteCommandSubSection (ref FileContent, TargetName, "Development" , BuildConfigIndex, 1);
+						WriteCommandSubSection (ref FileContent, TargetName, "Development" , BuildConfigIndex, 3);
+					}
+					BuildConfigIndex++;
 				}
-				BuildConfigIndex++;
-			}			
+			}
 		}
 
 		/// <summary>
diff --git a/Engine/Source/Programs/UnrealBuildTool/System/MakefileGenerator.cs b/Engine/Source/Programs/UnrealBuildTool/System/MakefileGenerator.cs
--- a/Engine/Source/Programs/UnrealBuildTool/System/MakefileGenerator.cs
+++ b/Engine/Source/Programs/UnrealBuildTool/System/MakefileGenerator.cs
@@ -89,23 +89,32 @@ namespace UnrealBuildTool
             );
             String MakeProjectCmdArg = "";
             String MakeBuildCommand = "";
-            foreach (string Target in DiscoverTargets())
+            foreach(var Project in GeneratedProjectFiles)
             {
-                var Basename = Path.GetFileNameWithoutExtension(Target);
-                Basename = Path.GetFileNameWithoutExtension(Basename);
-                foreach (UnrealTargetConfiguration CurConfiguration in Enum.GetValues(typeof(UnrealTargetConfiguration)))
-                {
-                    if (CurConfiguration != UnrealTargetConfiguration.Unknown && CurConfiguration != UnrealTargetConfiguration.Development)
-                    {
-                        if (UnrealBuildTool.IsValidConfiguration(CurConfiguration))
-                        {
-                            var Confname = Enum.GetName(typeof(UnrealTargetConfiguration), CurConfiguration);
-                            MakefileContent.Append(String.Format(" \\\n\t{0}-Linux-{1} ", Basename, Confname));
-                        }
-                    }
-                }
-                MakefileContent.Append(" \\\n\t" + Basename);
-            }
+				foreach (var TargetFile in Project.ProjectTargets)
+				{
+					if (TargetFile.TargetFilePath == null)
+					{
+						continue;
+					}
+
+					string TargetFileName = Path.GetFileNameWithoutExtension(TargetFile.TargetFilePath);
+					string Basename = TargetFileName.Substring(0, TargetFileName.LastIndexOf(".Target"));
+
+					foreach (UnrealTargetConfiguration CurConfiguration in Enum.GetValues(typeof(UnrealTargetConfiguration)))
+					{
+						if (CurConfiguration != UnrealTargetConfiguration.Unknown && CurConfiguration != UnrealTargetConfiguration.Development)
+						{
+							if (UnrealBuildTool.IsValidConfiguration(CurConfiguration))
+							{
+								var Confname = Enum.GetName(typeof(UnrealTargetConfiguration), CurConfiguration);
+								MakefileContent.Append(String.Format(" \\\n\t{0}-Linux-{1} ", Basename, Confname));
+							}
+						}
+					}
+					MakefileContent.Append(" \\\n\t" + Basename);
+				}
+            }            
             MakefileContent.Append("\\\n\tconfigure");
 
             MakefileContent.Append("\n\n" + BuildCommand + ProjectBuildCommand + "\n" +
@@ -115,34 +124,52 @@ namespace UnrealBuildTool
                 "DebugSet: RequiredTools UnrealFrontend-Linux-Debug UE4Editor-Linux-Debug\n\n"
             );
 
-            foreach (string Target in DiscoverTargets())
-            {
-                var Basename = Path.GetFileNameWithoutExtension(Target);
-                Basename = Path.GetFileNameWithoutExtension(Basename);
-                foreach (UnrealTargetConfiguration CurConfiguration in Enum.GetValues(typeof(UnrealTargetConfiguration)))
-                {
-                    if (Basename == GameProjectName || Basename == (GameProjectName + "Editor")) 
-                    {
-                        MakeProjectCmdArg = " -project=\"\\\"$(GAMEPROJECTFILE)\\\"\"";
-                        MakeBuildCommand = "$(PROJECTBUILD)";
-                    }
-                    else
-                    {
-                        MakeBuildCommand = "$(BUILD)";
-                    }
-
-                    if (CurConfiguration != UnrealTargetConfiguration.Unknown && CurConfiguration != UnrealTargetConfiguration.Development)
-                    {
-                        if (UnrealBuildTool.IsValidConfiguration(CurConfiguration))
-                        {
-                            var Confname = Enum.GetName(typeof(UnrealTargetConfiguration), CurConfiguration);
-                            MakefileContent.Append(String.Format("\n{1}-Linux-{2}:\n\t {0} {1} Linux {2} {3} $(ARGS)\n", MakeBuildCommand, Basename, Confname, MakeProjectCmdArg));
-                        }
-                    }
-                }
-
-                MakefileContent.Append(String.Format("\n{1}:\n\t {0} {1} Linux Development {2} $(ARGS)\n", MakeBuildCommand, Basename, MakeProjectCmdArg));
-            }
+			foreach(var Project in GeneratedProjectFiles)
+			{
+				foreach (var TargetFile in Project.ProjectTargets)
+				{
+					if (TargetFile.TargetFilePath == null)
+					{
+						continue;
+					}
+
+					string TargetFileName = Path.GetFileNameWithoutExtension(TargetFile.TargetFilePath);
+					string Basename = TargetFileName.Substring(0, TargetFileName.LastIndexOf(".Target"));
+
+					if (Basename == GameProjectName || Basename == (GameProjectName + "Editor")) 
+					{
+						MakeProjectCmdArg = " -project=\"\\\"$(GAMEPROJECTFILE)\\\"\"";
+						MakeBuildCommand = "$(PROJECTBUILD)";
+					}
+					else
+					{
+						MakeBuildCommand = "$(BUILD)";
+					}
+					
+					foreach (UnrealTargetConfiguration CurConfiguration in Enum.GetValues(typeof(UnrealTargetConfiguration)))
+					{
+						if (Basename == GameProjectName || Basename == (GameProjectName + "Editor")) 
+						{
+							MakeProjectCmdArg = " -project=\"\\\"$(GAMEPROJECTFILE)\\\"\"";
+							MakeBuildCommand = "$(PROJECTBUILD)";
+						}
+						else
+						{
+							MakeBuildCommand = "$(BUILD)";
+						}
+
+						if (CurConfiguration != UnrealTargetConfiguration.Unknown && CurConfiguration != UnrealTargetConfiguration.Development)
+						{
+							if (UnrealBuildTool.IsValidConfiguration(CurConfiguration))
+							{
+								var Confname = Enum.GetName(typeof(UnrealTargetConfiguration), CurConfiguration);
+								MakefileContent.Append(String.Format("\n{1}-Linux-{2}:\n\t {0} {1} Linux {2} {3} $(ARGS)\n", MakeBuildCommand, Basename, Confname, MakeProjectCmdArg));
+							}
+						}
+					}
+					MakefileContent.Append(String.Format("\n{1}:\n\t {0} {1} Linux Development {2} $(ARGS)\n", MakeBuildCommand, Basename, MakeProjectCmdArg));
+				}
+			}
 
             MakefileContent.Append("\nconfigure:\n");
             if (!String.IsNullOrEmpty (GameProjectName))
diff --git a/Engine/Source/Programs/UnrealBuildTool/System/QMakefileGenerator.cs b/Engine/Source/Programs/UnrealBuildTool/System/QMakefileGenerator.cs
--- a/Engine/Source/Programs/UnrealBuildTool/System/QMakefileGenerator.cs
+++ b/Engine/Source/Programs/UnrealBuildTool/System/QMakefileGenerator.cs
@@ -122,7 +122,8 @@ namespace UnrealBuildTool
 
 			string QMakeGameProjectFile = "";
 
-			foreach (var CurProject in GeneratedProjectFiles) {
+			foreach (var CurProject in GeneratedProjectFiles) 
+			{
 
 				QMakefileProjectFile QMakeProject = CurProject as QMakefileProjectFile;
 				if (QMakeProject == null)
@@ -359,35 +360,43 @@ namespace UnrealBuildTool
 
 			string QMakeProjectCmdArg = "";
 
-			foreach (string TargetFilePath in DiscoverTargets())
+			foreach(var Project in GeneratedProjectFiles)
 			{
-				var TargetName = Utils.GetFilenameWithoutAnyExtensions(TargetFilePath);		// Remove both ".cs" and ".
-
-				foreach (UnrealTargetConfiguration CurConfiguration in Enum.GetValues(typeof(UnrealTargetConfiguration)))
+				foreach (var TargetFile in Project.ProjectTargets)
 				{
-					if (CurConfiguration != UnrealTargetConfiguration.Unknown && CurConfiguration != UnrealTargetConfiguration.Development)
+					if (TargetFile.TargetFilePath == null)
 					{
-						if (UnrealBuildTool.IsValidConfiguration(CurConfiguration))
-						{
+						continue;
+					}
 
-							if (TargetName == GameProjectName || TargetName == (GameProjectName + "Editor")) 
+					var TargetName = Utils.GetFilenameWithoutAnyExtensions(TargetFile.TargetFilePath);		// Remove both ".cs" and ".
+
+					foreach (UnrealTargetConfiguration CurConfiguration in Enum.GetValues(typeof(UnrealTargetConfiguration)))
+					{
+						if (CurConfiguration != UnrealTargetConfiguration.Unknown && CurConfiguration != UnrealTargetConfiguration.Development)
+						{
+							if (UnrealBuildTool.IsValidConfiguration(CurConfiguration))
 							{
-								QMakeProjectCmdArg = " -project=\"\\\"$$gameProjectFile\\\"\"";
+
+								if (TargetName == GameProjectName || TargetName == (GameProjectName + "Editor")) 
+								{
+									QMakeProjectCmdArg = " -project=\"\\\"$$gameProjectFile\\\"\"";
+								}
+								var ConfName = Enum.GetName(typeof(UnrealTargetConfiguration), CurConfiguration);
+								QMakeFileContent.Append(String.Format("{0}-Linux-{1}.commands = $$build {0} Linux {1} {2} $$args\n", TargetName, ConfName, QMakeProjectCmdArg));
+								QMakeTargetList += "\t" + TargetName + "-Linux-" + ConfName + " \\\n"; // , TargetName, ConfName);
 							}
-							var ConfName = Enum.GetName(typeof(UnrealTargetConfiguration), CurConfiguration);
-                            QMakeFileContent.Append(String.Format("{0}-Linux-{1}.commands = $$build {0} Linux {1} {2} $$args\n", TargetName, ConfName, QMakeProjectCmdArg));
-							QMakeTargetList += "\t" + TargetName + "-Linux-" + ConfName + " \\\n"; // , TargetName, ConfName);
 						}
 					}
-				}
 
-				if (TargetName == GameProjectName || TargetName == (GameProjectName + "Editor")) 
-				{
-					QMakeProjectCmdArg = " -project=\"\\\"$$gameProjectFile\\\"\"";
-				}
+					if (TargetName == GameProjectName || TargetName == (GameProjectName + "Editor")) 
+					{
+						QMakeProjectCmdArg = " -project=\"\\\"$$gameProjectFile\\\"\"";
+					}
 
-                QMakeFileContent.Append(String.Format("{0}.commands = $$build {0} Linux Development {1} $$args\n\n", TargetName, QMakeProjectCmdArg));
-				QMakeTargetList += "\t" + TargetName + " \\\n";
+					QMakeFileContent.Append(String.Format("{0}.commands = $$build {0} Linux Development {1} $$args\n\n", TargetName, QMakeProjectCmdArg));
+					QMakeTargetList += "\t" + TargetName + " \\\n";
+				}
 			}
 
 			QMakeFileContent.Append (QMakeTargetList.TrimEnd('\\'));
