Cache the path to a target's build receipt, rather than reevaluating it several times.

diff --git a/Engine/Source/Programs/UnrealBuildTool/Configuration/UEBuildTarget.cs b/Engine/Source/Programs/UnrealBuildTool/Configuration/UEBuildTarget.cs
--- a/Engine/Source/Programs/UnrealBuildTool/Configuration/UEBuildTarget.cs
+++ b/Engine/Source/Programs/UnrealBuildTool/Configuration/UEBuildTarget.cs
@@ -809,6 +809,9 @@ namespace UnrealBuildTool
 		/** The receipt for this target, which contains a record of this build. */
 		private BuildReceipt Receipt;
 
+		/** Filename for the receipt for this target. */
+		private string ReceiptFileName;
+
 		/** Force output of the receipt to an additional filename */
 		[NonSerialized]
 		private string ForceReceiptFileName;
@@ -921,6 +924,9 @@ namespace UnrealBuildTool
 				EngineIntermediateDirectory = Path.GetFullPath(Path.Combine(BuildConfiguration.RelativeEnginePath, BuildConfiguration.PlatformIntermediateFolder, AppName, Configuration.ToString()));
 			}
 
+			// Get the receipt path for this target
+			ReceiptFileName = BuildReceipt.GetDefaultPath(ProjectDirectory, TargetName, Platform, Configuration, UEBuildPlatform.GetBuildPlatform(Platform).GetActiveArchitecture());
+
 			RemoteRoot = InDesc.RemoteRoot;
 
 			OnlyModules = InDesc.OnlyModules;
@@ -1061,7 +1067,7 @@ namespace UnrealBuildTool
 			// Expand all the paths in the receipt; they'll currently use variables for the engine and project directories
 			IUEBuildPlatform BuildPlatform = UEBuildPlatform.GetBuildPlatform(Platform);
 			// Expand all the paths in the receipt; they'll currently use variables for the engine and project directories
-			BuildReceipt ReceiptWithFullPaths = BuildReceipt.Read(BuildReceipt.GetDefaultPath(ProjectDirectory, TargetName, Platform, Configuration, BuildPlatform.GetActiveArchitecture()));
+			BuildReceipt ReceiptWithFullPaths = BuildReceipt.Read(ReceiptFileName);
 			if (ReceiptWithFullPaths == null)
 			{
 				ReceiptWithFullPaths = new BuildReceipt(Receipt);
@@ -1105,7 +1111,7 @@ namespace UnrealBuildTool
 				}
 				if (OnlyModules.Count == 0)
 				{
-					AllFilesToDelete.Add(BuildReceipt.GetDefaultPath(ProjectDirectory, TargetName, Platform, Configuration, BuildPlatform.GetActiveArchitecture()));
+					AllFilesToDelete.Add(ReceiptFileName);
 				}
 
 				//@todo. This does not clean up files that are no longer built by the target...				
@@ -1317,7 +1323,7 @@ namespace UnrealBuildTool
 			ReceiptWithFullPaths.ExpandPathVariables(BuildConfiguration.RelativeEnginePath, ProjectDirectory);
 			IUEBuildPlatform BuildPlatform = UEBuildPlatform.GetBuildPlatform(Platform);
 			// Expand all the paths in the receipt; they'll currently use variables for the engine and project directories
-			BuildReceipt BuiltReceiptWithFullPaths = BuildReceipt.Read(BuildReceipt.GetDefaultPath(ProjectDirectory, TargetName, Platform, Configuration, BuildPlatform.GetActiveArchitecture()));
+			BuildReceipt BuiltReceiptWithFullPaths = BuildReceipt.Read(ReceiptFileName);
 			if (BuiltReceiptWithFullPaths == null)
 			{
 				return;
@@ -1534,7 +1540,7 @@ namespace UnrealBuildTool
 			IUEBuildPlatform BuildPlatform = UEBuildPlatform.GetBuildPlatform(Platform);
 			if(OnlyModules.Count == 0)
 			{
-				Manifest.AddBuildProduct(BuildReceipt.GetDefaultPath(ProjectDirectory, TargetName, Platform, Configuration, BuildPlatform.GetActiveArchitecture()));
+				Manifest.AddBuildProduct(ReceiptFileName);
 			}
 
 			if (UEBuildConfiguration.bCleanProject)
@@ -1581,7 +1587,6 @@ namespace UnrealBuildTool
 				IUEBuildPlatform BuildPlatform = UEBuildPlatform.GetBuildPlatform(Platform);
 				if(OnlyModules == null || OnlyModules.Count == 0)
 				{
-					string ReceiptFileName = BuildReceipt.GetDefaultPath(ProjectDirectory, TargetName, Platform, Configuration, BuildPlatform.GetActiveArchitecture());
 					Directory.CreateDirectory(Path.GetDirectoryName(ReceiptFileName));
 					Receipt.Write(ReceiptFileName);
 				}
