Include dynamic libraries in extenal dependencies list.
#ue4

diff --git a/Engine/Source/Programs/UnrealBuildTool/Configuration/UEBuildTarget.cs b/Engine/Source/Programs/UnrealBuildTool/Configuration/UEBuildTarget.cs
--- a/Engine/Source/Programs/UnrealBuildTool/Configuration/UEBuildTarget.cs
+++ b/Engine/Source/Programs/UnrealBuildTool/Configuration/UEBuildTarget.cs
@@ -1419,22 +1419,29 @@ namespace UnrealBuildTool
 				LibraryPaths.Add(Directory.GetCurrentDirectory());
 				LibraryPaths.AddRange(Rules.PublicLibraryPaths.Where(x => !x.StartsWith("$(")).Select(x => Path.GetFullPath(x.Replace('/', '\\'))));
 
+				// Get all the extensions to look for
+				List<string> LibraryExtensions = new List<string>();
+				LibraryExtensions.Add(BuildPlatform.GetBinaryExtension(UEBuildBinaryType.StaticLibrary));
+				LibraryExtensions.Add(BuildPlatform.GetBinaryExtension(UEBuildBinaryType.DynamicLinkLibrary));
+
 				// Add all the libraries
-				string LibraryExtension = BuildPlatform.GetBinaryExtension(UEBuildBinaryType.StaticLibrary);
-				foreach(string LibraryName in Rules.PublicAdditionalLibraries)
+				foreach(string LibraryExtension in LibraryExtensions)
 				{
-					foreach(string LibraryPath in LibraryPaths)
+					foreach(string LibraryName in Rules.PublicAdditionalLibraries)
 					{
-						string LibraryFileName = Path.Combine(LibraryPath, LibraryName);
-						if(File.Exists(LibraryFileName))
+						foreach(string LibraryPath in LibraryPaths)
 						{
-							FileNames.Add(LibraryFileName);
-						}
+							string LibraryFileName = Path.Combine(LibraryPath, LibraryName);
+							if(File.Exists(LibraryFileName))
+							{
+								FileNames.Add(LibraryFileName);
+							}
 
-						string UnixLibraryFileName = Path.Combine(LibraryPath, "lib" + LibraryName + LibraryExtension);
-						if(File.Exists(UnixLibraryFileName))
-						{
-							FileNames.Add(UnixLibraryFileName);
+							string UnixLibraryFileName = Path.Combine(LibraryPath, "lib" + LibraryName + LibraryExtension);
+							if(File.Exists(UnixLibraryFileName))
+							{
+								FileNames.Add(UnixLibraryFileName);
+							}
 						}
 					}
 				}
