Added code to perform a clean up for generated files - can be an issue if project names change during development for example.
Also added code to clean up empty directories in the same area to reduce any mess that could be left behind.

diff --git a/Engine/Source/Programs/UnrealBuildTool/Android/UEDeployAndroid.cs b/Engine/Source/Programs/UnrealBuildTool/Android/UEDeployAndroid.cs
--- a/Engine/Source/Programs/UnrealBuildTool/Android/UEDeployAndroid.cs
+++ b/Engine/Source/Programs/UnrealBuildTool/Android/UEDeployAndroid.cs
@@ -1211,8 +1211,9 @@ namespace UnrealBuildTool.Android
 	                     
             // Generate Java files
             string PackageName = GetPackageName(ProjectName);
-            string TemplateDestinationBase = Path.Combine(ProjectDirectory, "Build", "Android", "src" , PackageName.Replace('.', '/'));
+            string TemplateDestinationBase = Path.Combine(ProjectDirectory, "Build", "Android", "src" , PackageName.Replace('.', Path.DirectorySeparatorChar));
             MakeDirectoryIfRequired(TemplateDestinationBase);
+
             // We'll be writing the OBB data into the same location as the download service files
             string UE4OBBDataFileName = GetUE4JavaOBBDataFileName(TemplateDestinationBase);
             string UE4DownloadShimFileName = GetUE4JavaDownloadShimFileName(UE4JavaFilePath);
@@ -1239,6 +1240,47 @@ namespace UnrealBuildTool.Android
                 { "$$PackageName$$",PackageName }
             });
 
+            // Sometimes old files get left behind if things change, so we'll do a clean up pass
+            {
+                string CleanUpBaseDir = Path.Combine(ProjectDirectory, "Build", "Android", "src");
+                var files = Directory.EnumerateFiles(CleanUpBaseDir, "*.java", SearchOption.AllDirectories);
+                
+                Log.TraceInformation("Cleaning up files based on template dir {0}", TemplateDestinationBase);
+
+                foreach(var filename in files)
+                {
+                    if (filename == UE4DownloadShimFileName)  // we always need the shim, and it'll get rewritten if needed anyway
+                        continue;
+
+                    string filePath = Path.GetDirectoryName(filename);  // grab the file's path
+                    if(filePath != TemplateDestinationBase)             // and check to make sure it isn't the same as the Template directory we calculated earlier
+                    {
+                        Log.TraceInformation("Cleaning up file {0} with path {1}", filename, filePath);
+                        File.Delete(filename);
+
+                        // Check to see if this file also exists in our target destination, and if so nuke it too
+                        string DestFilename = Path.Combine(UE4BuildPath, Utils.MakePathRelativeTo(filePath, UE4BuildFilesPath));
+                        if(File.Exists(filename))
+                        {
+                            File.Delete(filename);
+                        }
+                    }
+                }
+
+                // Directory clean up code
+                var directories = Directory.EnumerateDirectories(CleanUpBaseDir, "*", SearchOption.AllDirectories).OrderByDescending(x => x);
+                foreach(var directory in directories)
+                {
+                    if(Directory.Exists(directory) && Directory.GetFiles(directory, "*.*", SearchOption.AllDirectories).Count() == 0)
+                    {
+                        Log.TraceInformation("Cleaning Directory {0} as empty.", directory);
+                        Directory.Delete(directory, true);
+                    }
+                };
+
+
+            }
+
 
 			// cache if we want data in the Apk
 			bool bPackageDataInsideApk = PackageDataInsideApk(bDisallowPackagingDataInApk);
