Rename BuildReceipt.cs to TargetReceipt.cs to clarify what it's purpose is.

diff --git a/Engine/Source/Programs/UnrealBuildTool/System/BuildReceipt.cs b/Engine/Source/Programs/UnrealBuildTool/System/BuildReceipt.cs
deleted file mode 100644
diff --git a/Engine/Source/Programs/UnrealBuildTool/System/TargetReceipt.cs b/Engine/Source/Programs/UnrealBuildTool/System/TargetReceipt.cs
new file mode 100644
--- /dev/null
+++ b/Engine/Source/Programs/UnrealBuildTool/System/TargetReceipt.cs
@@ -0,0 +1,377 @@
+// Copyright 1998-2015 Epic Games, Inc. All Rights Reserved.
+
+using System;
+using System.Collections.Generic;
+using System.IO;
+using System.Diagnostics;
+using System.Linq;
+using System.Xml.Serialization;
+using System.ComponentModel;
+
+namespace UnrealBuildTool
+{
+	[Serializable]
+	public class BuildProperty
+	{
+		[XmlAttribute]
+		public string Name;
+
+		[XmlAttribute]
+		public string Value;
+
+		private BuildProperty()
+		{
+		}
+
+		public BuildProperty(string InName, string InValue)
+		{
+			Name = InName;
+			Value = InValue;
+		}
+	}
+
+	public enum BuildProductType
+	{
+		Executable,
+		DynamicLibrary,
+		StaticLibrary,
+		ImportLibrary,
+		SymbolFile,
+		RequiredResource,
+	}
+
+	[Serializable]
+	public class BuildProduct
+	{
+		[XmlAttribute]
+		public string Path;
+
+		[XmlAttribute]
+		public BuildProductType Type;
+
+		[XmlAttribute, DefaultValue(false)]
+		public bool IsPrecompiled;
+
+		private BuildProduct()
+		{
+		}
+
+		public BuildProduct(string InPath, BuildProductType InType)
+		{
+			Path = InPath;
+			Type = InType;
+		}
+
+		public BuildProduct(BuildProduct Other)
+		{
+			Path = Other.Path;
+			Type = Other.Type;
+			IsPrecompiled = Other.IsPrecompiled;
+		}
+
+		public override string ToString()
+		{
+			return Path;
+		}
+	}
+
+	[Serializable]
+	public class RuntimeDependency
+	{
+		[XmlAttribute]
+		public string Path;
+
+		[XmlAttribute, DefaultValue(null)]
+		public string StagePath;
+
+		private RuntimeDependency()
+		{
+		}
+
+		public RuntimeDependency(string InPath)
+		{
+			Path = InPath;
+		}
+
+		public RuntimeDependency(string InPath, string InStagePath)
+		{
+			Path = InPath;
+			StagePath = InStagePath;
+		}
+
+		public RuntimeDependency(RuntimeDependency InOther)
+		{
+			Path = InOther.Path;
+			StagePath = InOther.StagePath;
+		}
+
+		public override string ToString()
+		{
+			return (StagePath == null)? Path : String.Format("{0} -> {1}", Path, StagePath);
+		}
+	}
+
+	/// <summary>
+	/// Stores a record of a built target, with all metadata that other tools may need to know about the build.
+	/// </summary>
+	[Serializable]
+	public class BuildReceipt
+	{
+		[XmlArrayItem("Property")]
+		public List<BuildProperty> Properties = new List<BuildProperty>();
+
+		[XmlArrayItem("BuildProduct")]
+		public List<BuildProduct> BuildProducts = new List<BuildProduct>();
+
+		[XmlArrayItem("RuntimeDependency")]
+		public List<RuntimeDependency> RuntimeDependencies = new List<RuntimeDependency>();
+
+		// if packaging in a mode where some files aren't required, set this to false
+		public bool bRequireDependenciesToExist = true;
+
+		/// <summary>
+		/// Default constructor
+		/// </summary>
+		public BuildReceipt()
+		{
+		}
+
+		/// <summary>
+		/// Copy constructor
+		/// </summary>
+		/// <param name="InOther">Receipt to copy from</param>
+		public BuildReceipt(BuildReceipt Other)
+		{ 
+			foreach(BuildProduct OtherBuildProduct in Other.BuildProducts)
+			{
+				BuildProducts.Add(new BuildProduct(OtherBuildProduct));
+			}
+			foreach(RuntimeDependency OtherRuntimeDependency in Other.RuntimeDependencies)
+			{
+				RuntimeDependencies.Add(new RuntimeDependency(OtherRuntimeDependency));
+			}
+		}
+
+		/// <summary>
+		/// Sets a property with the given name
+		/// </summary>
+		/// <param name="Name">Name of the property; case sensitive.</param>
+		/// <param name="Value">Value for the property</param>
+		public void SetProperty(string Name, string Value)
+		{
+			BuildProperty Property = Properties.FirstOrDefault(x => x.Name == Name);
+			if(Property == null)
+			{
+				Properties.Add(new BuildProperty(Name, Value));
+			}
+			else
+			{
+				Property.Value = Value;
+			}
+		}
+
+		/// <summary>
+		/// Gets the value associated with a property
+		/// </summary>
+		/// <param name="Name">Name of the property; case sensitive.</param>
+		/// <param name="DefaultValue">Default value for the property if it's not found</param>
+		public string GetProperty(string Name, string DefaultValue)
+		{
+			BuildProperty Property = Properties.FirstOrDefault(x => x.Name == Name);
+			if(Property == null)
+			{
+				return DefaultValue;
+			}
+			else
+			{
+				return Property.Value;
+			}
+		}
+
+		/// <summary>
+		/// Adds a build product to the receipt. Does not check whether it already exists.
+		/// </summary>
+		/// <param name="Path">Path to the build product.</param>
+		/// <param name="Type">Type of build product.</param>
+		/// <returns>The BuildProduct object that was created</returns>
+		public BuildProduct AddBuildProduct(string Path, BuildProductType Type)
+		{
+			BuildProduct NewBuildProduct = new BuildProduct(Path, Type);
+			BuildProducts.Add(NewBuildProduct);
+			return NewBuildProduct;
+		}
+
+		/// <summary>
+		/// Constructs a runtime dependency object and adds it to the receipt.
+		/// </summary>
+		/// <param name="Path">Source path for the dependency</param>
+		/// <param name="StagePath">Location for the dependency in the staged build</param>
+		/// <returns>The RuntimeDependency object that was created</returns>
+		public RuntimeDependency AddRuntimeDependency(string Path, string StagePath)
+		{
+			RuntimeDependency NewRuntimeDependency = new RuntimeDependency(Path, StagePath);
+			RuntimeDependencies.Add(NewRuntimeDependency);
+			return NewRuntimeDependency;
+		}
+
+		/// <summary>
+		/// Merges another receipt to this one.
+		/// </summary>
+		/// <param name="Other">Receipt which should be merged</param>
+		public void Merge(BuildReceipt Other)
+		{
+			foreach(BuildProduct OtherBuildProduct in Other.BuildProducts)
+			{
+				BuildProducts.Add(OtherBuildProduct);
+			}
+			foreach(RuntimeDependency OtherRuntimeDependency in Other.RuntimeDependencies)
+			{
+				if(!RuntimeDependencies.Any(x => x.Path == OtherRuntimeDependency.Path && x.StagePath == OtherRuntimeDependency.StagePath))
+				{
+					RuntimeDependencies.Add(OtherRuntimeDependency);
+				}
+			}
+		}
+
+		/// <summary>
+		/// Expand all the path variables in the manifest
+		/// </summary>
+		/// <param name="EngineDir">Value for the $(EngineDir) variable</param>
+		/// <param name="ProjectDir">Value for the $(ProjectDir) variable</param>
+		public void ExpandPathVariables(string EngineDir, string ProjectDir)
+		{
+			ExpandPathVariables(EngineDir, ProjectDir, new Dictionary<string, string>());
+		}
+
+		/// <summary>
+		/// Control whether the dependencies are required while staging them
+		/// </summary>
+		/// <param name="InDependenciesAreRequired"></param>
+		public void SetDependenciesToBeRequired(bool InDependenciesAreRequired)
+		{
+			bRequireDependenciesToExist = InDependenciesAreRequired;
+		}
+
+		/// <summary>
+		/// Expand all the path variables in the manifest, including a list of supplied variable values.
+		/// </summary>
+		/// <param name="EngineDir">Value for the $(EngineDir) variable</param>
+		/// <param name="ProjectDir">Value for the $(ProjectDir) variable</param>
+		public void ExpandPathVariables(string EngineDir, string ProjectDir, IDictionary<string, string> OtherVariables)
+		{
+			// Build a dictionary containing the standard variable expansions
+			Dictionary<string, string> Variables = new Dictionary<string, string>(OtherVariables);
+			Variables["EngineDir"] = Path.GetFullPath(EngineDir).TrimEnd(Path.DirectorySeparatorChar);
+			Variables["ProjectDir"] = Path.GetFullPath(ProjectDir).TrimEnd(Path.DirectorySeparatorChar);
+
+			// Replace all the variables in the paths
+			foreach(BuildProduct BuildProduct in BuildProducts)
+			{
+				BuildProduct.Path = Utils.ExpandVariables(BuildProduct.Path, Variables);
+			}
+			foreach(RuntimeDependency RuntimeDependency in RuntimeDependencies)
+			{
+				RuntimeDependency.Path = Utils.ExpandVariables(RuntimeDependency.Path, Variables);
+				if(RuntimeDependency.StagePath != null)
+				{
+					RuntimeDependency.StagePath = Utils.ExpandVariables(RuntimeDependency.StagePath, Variables);
+				}
+			}
+		}
+
+		/// <summary>
+		/// Inserts standard $(EngineDir) and $(ProjectDir) variables into any path strings, so it can be used on different machines.
+		/// </summary>
+		/// <param name="EngineDir">The engine directory. Relative paths are ok.</param>
+		/// <param name="ProjectDir">The project directory. Relative paths are ok.</param>
+		public void InsertStandardPathVariables(string EngineDir, string ProjectDir)
+		{
+			string EnginePrefix = Path.GetFullPath(EngineDir).TrimEnd(Path.DirectorySeparatorChar) + Path.DirectorySeparatorChar;
+			string ProjectPrefix = Path.GetFullPath(ProjectDir).TrimEnd(Path.DirectorySeparatorChar) + Path.DirectorySeparatorChar;
+
+			foreach(BuildProduct BuildProduct in BuildProducts)
+			{
+				BuildProduct.Path = InsertStandardPathVariablesToString(BuildProduct.Path, EnginePrefix, ProjectPrefix);
+			}
+			foreach(RuntimeDependency RuntimeDependency in RuntimeDependencies)
+			{
+				RuntimeDependency.Path = InsertStandardPathVariablesToString(RuntimeDependency.Path, EnginePrefix, ProjectPrefix);
+				if(RuntimeDependency.StagePath != null)
+				{
+					RuntimeDependency.StagePath = InsertStandardPathVariablesToString(RuntimeDependency.StagePath, EnginePrefix, ProjectPrefix);
+				}
+			}
+		}
+
+		/// <summary>
+		/// Inserts $(EngineDir) and $(ProjectDir) variables into a path string, so it can be used on different machines.
+		/// </summary>
+		/// <param name="InputPath">Input path</param>
+		/// <param name="EngineDir">The engine directory. Relative paths are ok.</param>
+		/// <param name="ProjectDir">The project directory. Relative paths are ok.</param>
+		/// <returns>New string with the base directory replaced, or the original string</returns>
+		static string InsertStandardPathVariablesToString(string InputPath, string EnginePrefix, string ProjectPrefix)
+		{
+			string Result = InputPath;
+			if(!InputPath.StartsWith("$("))
+			{
+				string FullInputPath = Path.GetFullPath(InputPath);
+				if(FullInputPath.StartsWith(EnginePrefix))
+				{
+					Result = "$(EngineDir)" + FullInputPath.Substring(EnginePrefix.Length - 1);
+				}
+				else if(FullInputPath.StartsWith(ProjectPrefix))
+				{
+					Result = "$(ProjectDir)" + FullInputPath.Substring(ProjectPrefix.Length - 1);
+				}
+			}
+			return Result;
+		}
+
+		/// <summary>
+		/// Returns the standard path to the build receipt for a given target
+		/// </summary>
+		/// <param name="DirectoryName">Base directory for the target being built; either the project directory or engine directory.</param>
+		/// <param name="TargetName">The target being built</param>
+		/// <param name="Configuration">The target configuration</param>
+		/// <param name="Platform">The target platform</param>
+		/// <returns>Path to the receipt for this target</returns>
+		public static string GetDefaultPath(string BaseDir, string TargetName, UnrealTargetPlatform Platform, UnrealTargetConfiguration Configuration, string BuildArchitecture)
+		{
+			return Path.Combine(BaseDir, "Build", "Receipts", String.Format("{0}-{1}-{2}{3}.target.xml", TargetName, Platform.ToString(), Configuration.ToString(), BuildArchitecture));
+		}
+
+		static XmlSerializer Serializer = XmlSerializer.FromTypes(new Type[]{ typeof(BuildReceipt) })[0];
+
+		/// <summary>
+		/// Read a receipt from disk.
+		/// </summary>
+		/// <param name="FileName">Filename to read from</param>
+		public static BuildReceipt Read(string FileName)
+		{
+			if (File.Exists(FileName))
+			{
+				using (StreamReader Reader = new StreamReader(FileName))
+				{
+					return (BuildReceipt)Serializer.Deserialize(Reader);
+				}
+			}
+			else
+			{
+				return null;
+			}
+		}
+
+		/// <summary>
+		/// Write the receipt to disk.
+		/// </summary>
+		/// <param name="FileName">Output filename</param>
+		public void Write(string FileName)
+		{
+			using(StreamWriter Writer = new StreamWriter(FileName))
+			{
+				Serializer.Serialize(Writer, this);
+			}
+		}
+	}
+}
diff --git a/Engine/Source/Programs/UnrealBuildTool/UnrealBuildTool.csproj b/Engine/Source/Programs/UnrealBuildTool/UnrealBuildTool.csproj
--- a/Engine/Source/Programs/UnrealBuildTool/UnrealBuildTool.csproj
+++ b/Engine/Source/Programs/UnrealBuildTool/UnrealBuildTool.csproj
@@ -124,6 +124,7 @@
     <Compile Include="Mac\UEBuildMac.cs" />
     <Compile Include="Mac\MacProjectGenerator.cs" />
     <Compile Include="Mac\UEDeployMac.cs" />
+    <Compile Include="System\TargetReceipt.cs" />
     <Compile Include="ToolChain\AppleToolChain.cs" />
   </ItemGroup>
   <ItemGroup>
@@ -193,7 +194,6 @@
     <Compile Include="System\BuildException.cs" />
     <Compile Include="Configuration\BuildConfiguration.cs" />
     <Compile Include="System\BuildHostPlatform.cs" />
-    <Compile Include="System\BuildReceipt.cs" />
     <Compile Include="System\FlatCPPIncludeDepencencyCache.cs" />
     <Compile Include="System\Distcc.cs" />
     <Compile Include="System\CodeLiteGenerator.cs" />
