Rename BuildReceipt to TargetReceipt, to clarify its purpose.

diff --git a/Engine/Source/Programs/AutomationTool/AutomationUtils/DeploymentContext.cs b/Engine/Source/Programs/AutomationTool/AutomationUtils/DeploymentContext.cs
--- a/Engine/Source/Programs/AutomationTool/AutomationUtils/DeploymentContext.cs
+++ b/Engine/Source/Programs/AutomationTool/AutomationUtils/DeploymentContext.cs
@@ -54,7 +54,7 @@ public class DeploymentContext //: ProjectParams
 	/// <summary>
 	/// Receipts for the build targets that should be staged
 	/// </summary>
-	public List<BuildReceipt> StageTargetReceipts;
+	public List<TargetReceipt> StageTargetReceipts;
 
 	/// <summary>
 	///  this is the root directory that contains the engine: d:\a\UE4\
@@ -234,7 +234,7 @@ public class DeploymentContext //: ProjectParams
 		Platform InSourcePlatform,
         Platform InTargetPlatform,
 		List<UnrealTargetConfiguration> InTargetConfigurations,
-		IEnumerable<BuildReceipt> InStageTargetReceipts,
+		IEnumerable<TargetReceipt> InStageTargetReceipts,
 		List<String> InStageExecutables,
 		bool InServer,
 		bool InCooked,
@@ -254,7 +254,7 @@ public class DeploymentContext //: ProjectParams
         CookSourcePlatform = InSourcePlatform;
 		StageTargetPlatform = InTargetPlatform;
 		StageTargetConfigurations = new List<UnrealTargetConfiguration>(InTargetConfigurations);
-		StageTargetReceipts = new List<BuildReceipt>(InStageTargetReceipts);
+		StageTargetReceipts = new List<TargetReceipt>(InStageTargetReceipts);
 		StageExecutables = InStageExecutables;
         IsCodeBasedProject = ProjectUtils.IsCodeBasedUProjectFile(RawProjectPath);
 		ShortProjectName = ProjectUtils.GetShortProjectName(RawProjectPath);
@@ -385,7 +385,7 @@ public class DeploymentContext //: ProjectParams
 		}
 	}
 
-	public void StageBuildProductsFromReceipt(BuildReceipt Receipt)
+	public void StageBuildProductsFromReceipt(TargetReceipt Receipt)
 	{
 		// Stage all the build products needed at runtime
 		foreach(BuildProduct BuildProduct in Receipt.BuildProducts)
@@ -407,7 +407,7 @@ public class DeploymentContext //: ProjectParams
 		}
 	}
 
-	public void StageRuntimeDependenciesFromReceipt(BuildReceipt Receipt)
+	public void StageRuntimeDependenciesFromReceipt(TargetReceipt Receipt)
 	{
 		// Also stage any additional runtime dependencies, like ThirdParty DLLs
 		foreach(RuntimeDependency RuntimeDependency in Receipt.RuntimeDependencies)
diff --git a/Engine/Source/Programs/AutomationTool/Mac/MacPlatform.Automation.cs b/Engine/Source/Programs/AutomationTool/Mac/MacPlatform.Automation.cs
--- a/Engine/Source/Programs/AutomationTool/Mac/MacPlatform.Automation.cs
+++ b/Engine/Source/Programs/AutomationTool/Mac/MacPlatform.Automation.cs
@@ -156,7 +156,7 @@ public class MacPlatform : Platform
 		// Stage the bootstrap executable
 		if(!Params.NoBootstrapExe)
 		{
-			foreach(BuildReceipt Receipt in SC.StageTargetReceipts)
+			foreach(TargetReceipt Receipt in SC.StageTargetReceipts)
 			{
 				BuildProduct Executable = Receipt.BuildProducts.FirstOrDefault(x => x.Type == BuildProductType.Executable);
 				if(Executable != null)
diff --git a/Engine/Source/Programs/AutomationTool/Scripts/BuildPluginCommand.Automation.cs b/Engine/Source/Programs/AutomationTool/Scripts/BuildPluginCommand.Automation.cs
--- a/Engine/Source/Programs/AutomationTool/Scripts/BuildPluginCommand.Automation.cs
+++ b/Engine/Source/Programs/AutomationTool/Scripts/BuildPluginCommand.Automation.cs
@@ -97,7 +97,7 @@ class BuildPlugin : BuildCommand
 				Arguments += String.Format(" -module {0}", ModuleName);
 			}
 
-			string ReceiptFileName = BuildReceipt.GetDefaultPath(Path.GetDirectoryName(PluginFileName), TargetName, Platform, Configuration, "");
+			string ReceiptFileName = TargetReceipt.GetDefaultPath(Path.GetDirectoryName(PluginFileName), TargetName, Platform, Configuration, "");
 			Arguments += String.Format(" -receipt {0}", CommandUtils.MakePathSafeToUseWithCommandLine(ReceiptFileName));
 			ReceiptFileNames.Add(ReceiptFileName);
 			
@@ -115,7 +115,7 @@ class BuildPlugin : BuildCommand
 		List<BuildProduct> BuildProducts = new List<BuildProduct>();
 		foreach(string ReceiptFileName in ReceiptFileNames)
 		{
-			BuildReceipt Receipt = BuildReceipt.Read(ReceiptFileName);
+			TargetReceipt Receipt = TargetReceipt.Read(ReceiptFileName);
 			BuildProducts.AddRange(Receipt.BuildProducts);
 		}
 		return BuildProducts;
diff --git a/Engine/Source/Programs/AutomationTool/Scripts/CopyBuildToStagingDirectory.Automation.cs b/Engine/Source/Programs/AutomationTool/Scripts/CopyBuildToStagingDirectory.Automation.cs
--- a/Engine/Source/Programs/AutomationTool/Scripts/CopyBuildToStagingDirectory.Automation.cs
+++ b/Engine/Source/Programs/AutomationTool/Scripts/CopyBuildToStagingDirectory.Automation.cs
@@ -266,7 +266,7 @@ public partial class Project : CommandUtils
 
 
 		// Stage any extra runtime dependencies from the receipts
-		foreach(BuildReceipt Receipt in SC.StageTargetReceipts)
+		foreach(TargetReceipt Receipt in SC.StageTargetReceipts)
 		{
 			SC.StageRuntimeDependenciesFromReceipt(Receipt);
 		}
@@ -1511,7 +1511,7 @@ public partial class Project : CommandUtils
 
 			string EngineDir = CommandUtils.CombinePaths(CommandUtils.CmdEnv.LocalRoot, "Engine");
 
-			List<BuildReceipt> TargetsToStage = new List<BuildReceipt>();
+			List<TargetReceipt> TargetsToStage = new List<TargetReceipt>();
 			foreach(string Target in ListToProcess)
 			{
 				foreach(UnrealTargetConfiguration Config in ConfigsToProcess)
@@ -1536,7 +1536,7 @@ public partial class Project : CommandUtils
                                 Architecture = BuildPlatform.GetActiveArchitecture();
                             }
                         }
-						string ReceiptFileName = BuildReceipt.GetDefaultPath(ReceiptBaseDir, Target, ReceiptPlatform, Config, Architecture);
+						string ReceiptFileName = TargetReceipt.GetDefaultPath(ReceiptBaseDir, Target, ReceiptPlatform, Config, Architecture);
 						if(!File.Exists(ReceiptFileName))
 						{
 							if (bRequireStagedFilesToExist)
@@ -1553,7 +1553,7 @@ public partial class Project : CommandUtils
 						}
 
 						// Convert the paths to absolute
-						BuildReceipt Receipt = BuildReceipt.Read(ReceiptFileName);
+						TargetReceipt Receipt = TargetReceipt.Read(ReceiptFileName);
 						Receipt.ExpandPathVariables(EngineDir, Path.GetDirectoryName(Params.RawProjectPath));
 						Receipt.SetDependenciesToBeRequired(bRequireStagedFilesToExist);
 						TargetsToStage.Add(Receipt);
diff --git a/Engine/Source/Programs/AutomationTool/Win/WinPlatform.Automation.cs b/Engine/Source/Programs/AutomationTool/Win/WinPlatform.Automation.cs
--- a/Engine/Source/Programs/AutomationTool/Win/WinPlatform.Automation.cs
+++ b/Engine/Source/Programs/AutomationTool/Win/WinPlatform.Automation.cs
@@ -36,7 +36,7 @@ public abstract class BaseWinPlatform : Platform
 		}
 
 		// Stage all the build products
-		foreach(BuildReceipt Receipt in SC.StageTargetReceipts)
+		foreach(TargetReceipt Receipt in SC.StageTargetReceipts)
 		{
 			SC.StageBuildProductsFromReceipt(Receipt);
 		}
@@ -47,7 +47,7 @@ public abstract class BaseWinPlatform : Platform
 		// Stage the bootstrap executable
 		if(!Params.NoBootstrapExe)
 		{
-			foreach(BuildReceipt Receipt in SC.StageTargetReceipts)
+			foreach(TargetReceipt Receipt in SC.StageTargetReceipts)
 			{
 				BuildProduct Executable = Receipt.BuildProducts.FirstOrDefault(x => x.Type == BuildProductType.Executable);
 				if(Executable != null)
diff --git a/Engine/Source/Programs/UnrealBuildTool/Android/AndroidToolChain.cs b/Engine/Source/Programs/UnrealBuildTool/Android/AndroidToolChain.cs
--- a/Engine/Source/Programs/UnrealBuildTool/Android/AndroidToolChain.cs
+++ b/Engine/Source/Programs/UnrealBuildTool/Android/AndroidToolChain.cs
@@ -1146,7 +1146,7 @@ namespace UnrealBuildTool
 			return Outputs.ToArray();
 		}
 
-		public override void AddFilesToReceipt(BuildReceipt Receipt, UEBuildBinary Binary)
+		public override void AddFilesToReceipt(TargetReceipt Receipt, UEBuildBinary Binary)
 		{
 			// the binary will have all of the .so's in the output files, we need to trim down to the shared apk (which is what needs to go into the manifest)
 			if (Binary.Config.Type != UEBuildBinaryType.StaticLibrary)
diff --git a/Engine/Source/Programs/UnrealBuildTool/Configuration/UEBuildBinary.cs b/Engine/Source/Programs/UnrealBuildTool/Configuration/UEBuildBinary.cs
--- a/Engine/Source/Programs/UnrealBuildTool/Configuration/UEBuildBinary.cs
+++ b/Engine/Source/Programs/UnrealBuildTool/Configuration/UEBuildBinary.cs
@@ -289,9 +289,9 @@ namespace UnrealBuildTool
 		/// </summary>
 		/// <param name="ToolChain">Toolchain for the target platform</param>
 		/// <param name="BuildPlatform">Platform that we're building for</param>
-		public virtual BuildReceipt MakeReceipt(IUEToolChain ToolChain)
+		public virtual TargetReceipt MakeReceipt(IUEToolChain ToolChain)
 		{
-			BuildReceipt Receipt = new BuildReceipt();
+			TargetReceipt Receipt = new TargetReceipt();
 
 			// Get the type of build products we're creating
 			BuildProductType Type = BuildProductType.RequiredResource;
@@ -335,7 +335,7 @@ namespace UnrealBuildTool
 		/// <param name="OutputFile">Build product to add</param>
 		/// <param name="DebugExtension">Extension for the matching debug file (may be null).</param>
 		/// <param name="Receipt">Receipt to add to</param>
-		static void AddBuildProductAndDebugFile(string OutputFile, BuildProductType OutputType, string DebugExtension, BuildReceipt Receipt)
+		static void AddBuildProductAndDebugFile(string OutputFile, BuildProductType OutputType, string DebugExtension, TargetReceipt Receipt)
 		{
 			Receipt.AddBuildProduct(OutputFile, OutputType);
 
@@ -633,9 +633,9 @@ namespace UnrealBuildTool
 		/// Overrides base class to add module runtime dependencies to the build receipt.
 		/// </summary>
 		/// <param name="ToolChain">The platform toolchain</param>
-		public override BuildReceipt MakeReceipt(IUEToolChain ToolChain)
+		public override TargetReceipt MakeReceipt(IUEToolChain ToolChain)
 		{
-			BuildReceipt Receipt = base.MakeReceipt(ToolChain);
+			TargetReceipt Receipt = base.MakeReceipt(ToolChain);
 
 			// Set the IsPrecompiled flag on all the build products if we're not actually building this binary
 			if(!Config.bAllowCompilation)
diff --git a/Engine/Source/Programs/UnrealBuildTool/Configuration/UEBuildTarget.cs b/Engine/Source/Programs/UnrealBuildTool/Configuration/UEBuildTarget.cs
--- a/Engine/Source/Programs/UnrealBuildTool/Configuration/UEBuildTarget.cs
+++ b/Engine/Source/Programs/UnrealBuildTool/Configuration/UEBuildTarget.cs
@@ -807,7 +807,7 @@ namespace UnrealBuildTool
 		public CaselessDictionary<FlatModuleCsDataType> FlatModuleCsData = new CaselessDictionary<FlatModuleCsDataType>();
 
 		/** The receipt for this target, which contains a record of this build. */
-		private BuildReceipt Receipt;
+		private TargetReceipt Receipt;
 
 		/** Filename for the receipt for this target. */
 		private string ReceiptFileName;
@@ -925,7 +925,7 @@ namespace UnrealBuildTool
 			}
 
 			// Get the receipt path for this target
-			ReceiptFileName = BuildReceipt.GetDefaultPath(ProjectDirectory, TargetName, Platform, Configuration, UEBuildPlatform.GetBuildPlatform(Platform).GetActiveArchitecture());
+			ReceiptFileName = TargetReceipt.GetDefaultPath(ProjectDirectory, TargetName, Platform, Configuration, UEBuildPlatform.GetBuildPlatform(Platform).GetActiveArchitecture());
 
 			RemoteRoot = InDesc.RemoteRoot;
 
@@ -1067,10 +1067,10 @@ namespace UnrealBuildTool
 			// Expand all the paths in the receipt; they'll currently use variables for the engine and project directories
 			IUEBuildPlatform BuildPlatform = UEBuildPlatform.GetBuildPlatform(Platform);
 			// Expand all the paths in the receipt; they'll currently use variables for the engine and project directories
-			BuildReceipt ReceiptWithFullPaths = BuildReceipt.Read(ReceiptFileName);
+			TargetReceipt ReceiptWithFullPaths = TargetReceipt.Read(ReceiptFileName);
 			if (ReceiptWithFullPaths == null)
 			{
-				ReceiptWithFullPaths = new BuildReceipt(Receipt);
+				ReceiptWithFullPaths = new TargetReceipt(Receipt);
 			}
 			ReceiptWithFullPaths.ExpandPathVariables(BuildConfiguration.RelativeEnginePath, ProjectDirectory);
 
@@ -1319,11 +1319,11 @@ namespace UnrealBuildTool
 		public void CleanStaleModules()
 		{
 			// Expand all the paths in the receipt; they'll currently use variables for the engine and project directories
-			BuildReceipt ReceiptWithFullPaths = new BuildReceipt(Receipt);
+			TargetReceipt ReceiptWithFullPaths = new TargetReceipt(Receipt);
 			ReceiptWithFullPaths.ExpandPathVariables(BuildConfiguration.RelativeEnginePath, ProjectDirectory);
 			IUEBuildPlatform BuildPlatform = UEBuildPlatform.GetBuildPlatform(Platform);
 			// Expand all the paths in the receipt; they'll currently use variables for the engine and project directories
-			BuildReceipt BuiltReceiptWithFullPaths = BuildReceipt.Read(ReceiptFileName);
+			TargetReceipt BuiltReceiptWithFullPaths = TargetReceipt.Read(ReceiptFileName);
 			if (BuiltReceiptWithFullPaths == null)
 			{
 				return;
@@ -1523,7 +1523,7 @@ namespace UnrealBuildTool
 			}
 
 			// Expand all the paths in the receipt; they'll currently use variables for the engine and project directories
-			BuildReceipt ReceiptWithFullPaths = new BuildReceipt(Receipt);
+			TargetReceipt ReceiptWithFullPaths = new TargetReceipt(Receipt);
 			ReceiptWithFullPaths.ExpandPathVariables(BuildConfiguration.RelativeEnginePath, ProjectDirectory);
 
 			foreach(BuildProduct BuildProduct in ReceiptWithFullPaths.BuildProducts)
@@ -1563,7 +1563,7 @@ namespace UnrealBuildTool
 		/** Creates the receipt for the target */
 		private void PrepareReceipt(IUEToolChain ToolChain)
 		{
-			Receipt = new BuildReceipt();
+			Receipt = new TargetReceipt();
 
 			// Add the target properties
 			Receipt.SetProperty("TargetName", TargetName);
@@ -1574,7 +1574,7 @@ namespace UnrealBuildTool
 			// Merge all the binary receipts into this
 			foreach(UEBuildBinary Binary in AppBinaries)
 			{
-				BuildReceipt BinaryReceipt = Binary.MakeReceipt(ToolChain);
+				TargetReceipt BinaryReceipt = Binary.MakeReceipt(ToolChain);
 				if(Binary.Config.Type == UEBuildBinaryType.StaticLibrary)
 				{
 					BinaryReceipt.RuntimeDependencies.Clear();
diff --git a/Engine/Source/Programs/UnrealBuildTool/HTML5/HTML5ToolChain.cs b/Engine/Source/Programs/UnrealBuildTool/HTML5/HTML5ToolChain.cs
--- a/Engine/Source/Programs/UnrealBuildTool/HTML5/HTML5ToolChain.cs
+++ b/Engine/Source/Programs/UnrealBuildTool/HTML5/HTML5ToolChain.cs
@@ -536,7 +536,7 @@ namespace UnrealBuildTool
 			throw new BuildException("HTML5 cannot compile C# files");
 		}
 
-        public override void AddFilesToReceipt(BuildReceipt Receipt, UEBuildBinary Binary)
+        public override void AddFilesToReceipt(TargetReceipt Receipt, UEBuildBinary Binary)
         {
             // we need to include the generated .mem and .symbols file.  
 			if(Binary.Config.Type != UEBuildBinaryType.StaticLibrary)
diff --git a/Engine/Source/Programs/UnrealBuildTool/IOS/IOSToolChain.cs b/Engine/Source/Programs/UnrealBuildTool/IOS/IOSToolChain.cs
--- a/Engine/Source/Programs/UnrealBuildTool/IOS/IOSToolChain.cs
+++ b/Engine/Source/Programs/UnrealBuildTool/IOS/IOSToolChain.cs
@@ -94,7 +94,7 @@ namespace UnrealBuildTool
 			SetupXcodePaths(true);
 		}
 
-		public override void AddFilesToReceipt(BuildReceipt Receipt, UEBuildBinary Binary)
+		public override void AddFilesToReceipt(TargetReceipt Receipt, UEBuildBinary Binary)
 		{
 			if (BuildConfiguration.bCreateStubIPA && Binary.Config.Type != UEBuildBinaryType.StaticLibrary)
 			{
diff --git a/Engine/Source/Programs/UnrealBuildTool/Mac/MacToolChain.cs b/Engine/Source/Programs/UnrealBuildTool/Mac/MacToolChain.cs
--- a/Engine/Source/Programs/UnrealBuildTool/Mac/MacToolChain.cs
+++ b/Engine/Source/Programs/UnrealBuildTool/Mac/MacToolChain.cs
@@ -1296,7 +1296,7 @@ namespace UnrealBuildTool
 
 		static private string BundleContentsDirectory = "";
 
-        public override void AddFilesToReceipt(BuildReceipt Receipt, UEBuildBinary Binary)
+        public override void AddFilesToReceipt(TargetReceipt Receipt, UEBuildBinary Binary)
 		{
 			// The cross-platform code adds .dSYMs for static libraries, which is just wrong, so
 			// eliminate them here for now.
diff --git a/Engine/Source/Programs/UnrealBuildTool/System/ExternalExecution.cs b/Engine/Source/Programs/UnrealBuildTool/System/ExternalExecution.cs
--- a/Engine/Source/Programs/UnrealBuildTool/System/ExternalExecution.cs
+++ b/Engine/Source/Programs/UnrealBuildTool/System/ExternalExecution.cs
@@ -230,17 +230,17 @@ namespace UnrealBuildTool
 			using (var TimestampTimer = new ScopedTimer("GetHeaderToolTimestamp"))
 			{
 				// Try to read the receipt for UHT.
-				string ReceiptPath = BuildReceipt.GetDefaultPath(BuildConfiguration.RelativeEnginePath, "UnrealHeaderTool", BuildHostPlatform.Current.Platform, UnrealTargetConfiguration.Development, null);
+				string ReceiptPath = TargetReceipt.GetDefaultPath(BuildConfiguration.RelativeEnginePath, "UnrealHeaderTool", BuildHostPlatform.Current.Platform, UnrealTargetConfiguration.Development, null);
 				if(!File.Exists(ReceiptPath))
 				{
 					Timestamp = DateTime.MaxValue;
 					return false;
 				}
 
-				BuildReceipt Receipt;
+				TargetReceipt Receipt;
 				try
 				{
-					Receipt = BuildReceipt.Read(ReceiptPath);
+					Receipt = TargetReceipt.Read(ReceiptPath);
 				}
 				catch(Exception)
 				{
@@ -268,7 +268,7 @@ namespace UnrealBuildTool
 		/// <returns>
 		/// True if all the files are valid.
 		/// </returns>
-		static bool CheckBinariesExist(BuildReceipt Receipt)
+		static bool CheckBinariesExist(TargetReceipt Receipt)
 		{
 			bool bExist = true;
 			foreach(BuildProduct BuildProduct in Receipt.BuildProducts)
@@ -291,7 +291,7 @@ namespace UnrealBuildTool
 		/// <returns>
 		/// True if all the files are valid.
 		/// </returns>
-		static bool CheckDynamicLibaryVersionsMatch(BuildReceipt Receipt)
+		static bool CheckDynamicLibaryVersionsMatch(TargetReceipt Receipt)
 		{
 			List<Tuple<string, int>> BinaryVersions = new List<Tuple<string,int>>();
 			foreach(BuildProduct BuildProduct in Receipt.BuildProducts)
@@ -323,7 +323,7 @@ namespace UnrealBuildTool
 		/// <returns>
 		/// True if all the files are valid.
 		/// </returns>
-		static DateTime GetTimestampFromBinaries(BuildReceipt Receipt)
+		static DateTime GetTimestampFromBinaries(TargetReceipt Receipt)
 		{
 			DateTime LatestWriteTime = DateTime.MinValue;
 			foreach(BuildProduct BuildProduct in Receipt.BuildProducts)
diff --git a/Engine/Source/Programs/UnrealBuildTool/System/TargetReceipt.cs b/Engine/Source/Programs/UnrealBuildTool/System/TargetReceipt.cs
--- a/Engine/Source/Programs/UnrealBuildTool/System/TargetReceipt.cs
+++ b/Engine/Source/Programs/UnrealBuildTool/System/TargetReceipt.cs
@@ -115,7 +115,7 @@ namespace UnrealBuildTool
 	/// Stores a record of a built target, with all metadata that other tools may need to know about the build.
 	/// </summary>
 	[Serializable]
-	public class BuildReceipt
+	public class TargetReceipt
 	{
 		[XmlArrayItem("Property")]
 		public List<BuildProperty> Properties = new List<BuildProperty>();
@@ -132,7 +132,7 @@ namespace UnrealBuildTool
 		/// <summary>
 		/// Default constructor
 		/// </summary>
-		public BuildReceipt()
+		public TargetReceipt()
 		{
 		}
 
@@ -140,7 +140,7 @@ namespace UnrealBuildTool
 		/// Copy constructor
 		/// </summary>
 		/// <param name="InOther">Receipt to copy from</param>
-		public BuildReceipt(BuildReceipt Other)
+		public TargetReceipt(TargetReceipt Other)
 		{ 
 			foreach(BuildProduct OtherBuildProduct in Other.BuildProducts)
 			{
@@ -218,7 +218,7 @@ namespace UnrealBuildTool
 		/// Merges another receipt to this one.
 		/// </summary>
 		/// <param name="Other">Receipt which should be merged</param>
-		public void Merge(BuildReceipt Other)
+		public void Merge(TargetReceipt Other)
 		{
 			foreach(BuildProduct OtherBuildProduct in Other.BuildProducts)
 			{
@@ -341,19 +341,19 @@ namespace UnrealBuildTool
 			return Path.Combine(BaseDir, "Build", "Receipts", String.Format("{0}-{1}-{2}{3}.target.xml", TargetName, Platform.ToString(), Configuration.ToString(), BuildArchitecture));
 		}
 
-		static XmlSerializer Serializer = XmlSerializer.FromTypes(new Type[]{ typeof(BuildReceipt) })[0];
+		static XmlSerializer Serializer = XmlSerializer.FromTypes(new Type[]{ typeof(TargetReceipt) })[0];
 
 		/// <summary>
 		/// Read a receipt from disk.
 		/// </summary>
 		/// <param name="FileName">Filename to read from</param>
-		public static BuildReceipt Read(string FileName)
+		public static TargetReceipt Read(string FileName)
 		{
 			if (File.Exists(FileName))
 			{
 				using (StreamReader Reader = new StreamReader(FileName))
 				{
-					return (BuildReceipt)Serializer.Deserialize(Reader);
+					return (TargetReceipt)Serializer.Deserialize(Reader);
 				}
 			}
 			else
diff --git a/Engine/Source/Programs/UnrealBuildTool/ToolChain/UEToolChain.cs b/Engine/Source/Programs/UnrealBuildTool/ToolChain/UEToolChain.cs
--- a/Engine/Source/Programs/UnrealBuildTool/ToolChain/UEToolChain.cs
+++ b/Engine/Source/Programs/UnrealBuildTool/ToolChain/UEToolChain.cs
@@ -38,7 +38,7 @@ namespace UnrealBuildTool
 
 		void SetUpGlobalEnvironment();
 
-		void AddFilesToReceipt(BuildReceipt Receipt, UEBuildBinary Binary);
+		void AddFilesToReceipt(TargetReceipt Receipt, UEBuildBinary Binary);
 
 		void SetupBundleDependencies(List<UEBuildBinary> Binaries, string GameName);
 
@@ -172,7 +172,7 @@ namespace UnrealBuildTool
 		{
 		}
 
-		public virtual void AddFilesToReceipt(BuildReceipt Receipt, UEBuildBinary Binary)
+		public virtual void AddFilesToReceipt(TargetReceipt Receipt, UEBuildBinary Binary)
 		{
 		}
 
diff --git a/Engine/Source/Programs/UnrealBuildTool/Windows/VCToolChain.cs b/Engine/Source/Programs/UnrealBuildTool/Windows/VCToolChain.cs
--- a/Engine/Source/Programs/UnrealBuildTool/Windows/VCToolChain.cs
+++ b/Engine/Source/Programs/UnrealBuildTool/Windows/VCToolChain.cs
@@ -1762,7 +1762,7 @@ namespace UnrealBuildTool
 			return IncludePaths;
 		}
 
-		public override void AddFilesToReceipt(BuildReceipt Receipt, UEBuildBinary Binary)
+		public override void AddFilesToReceipt(TargetReceipt Receipt, UEBuildBinary Binary)
 		{
 			if (Binary.Config.Type == UEBuildBinaryType.DynamicLinkLibrary)
 			{
