Catch exceptions if the target receipt can't be read.

diff --git a/Engine/Source/Programs/UnrealBuildTool/System/TargetReceipt.cs b/Engine/Source/Programs/UnrealBuildTool/System/TargetReceipt.cs
--- a/Engine/Source/Programs/UnrealBuildTool/System/TargetReceipt.cs
+++ b/Engine/Source/Programs/UnrealBuildTool/System/TargetReceipt.cs
@@ -351,9 +351,16 @@ namespace UnrealBuildTool
 		{
 			if (File.Exists(FileName))
 			{
-				using (StreamReader Reader = new StreamReader(FileName))
+				try
 				{
-					return (TargetReceipt)Serializer.Deserialize(Reader);
+					using (StreamReader Reader = new StreamReader(FileName))
+					{
+						return (TargetReceipt)Serializer.Deserialize(Reader);
+					}
+				}
+				catch(Exception)
+				{
+					return null;
 				}
 			}
 			else
