Handle cases where TargetReceipt.Read() returns null instead of throwing now.

diff --git a/Engine/Source/Programs/UnrealBuildTool/System/ExternalExecution.cs b/Engine/Source/Programs/UnrealBuildTool/System/ExternalExecution.cs
--- a/Engine/Source/Programs/UnrealBuildTool/System/ExternalExecution.cs
+++ b/Engine/Source/Programs/UnrealBuildTool/System/ExternalExecution.cs
@@ -241,6 +241,11 @@ namespace UnrealBuildTool
 				try
 				{
 					Receipt = TargetReceipt.Read(ReceiptPath);
+					if(Receipt == null)
+					{
+						Timestamp = DateTime.MaxValue;
+						return false;
+					}
 				}
 				catch(Exception)
 				{
