Added long suffix to file length in generated OBBData.java
PR #1405: added specific suffix for long (Contributed by [REDACTED_PERSON])
#ue4
#android

diff --git a/Engine/Source/Programs/UnrealBuildTool/Android/UEDeployAndroid.cs b/Engine/Source/Programs/UnrealBuildTool/Android/UEDeployAndroid.cs
--- a/Engine/Source/Programs/UnrealBuildTool/Android/UEDeployAndroid.cs
+++ b/Engine/Source/Programs/UnrealBuildTool/Android/UEDeployAndroid.cs
@@ -388,7 +388,7 @@ namespace UnrealBuildTool.Android
             {
                 obbData.Append("new XAPKFile(\ntrue, // true signifies a main file\n");
                 obbData.AppendFormat("\"{0}\", // the version of the APK that the file was uploaded against\n", GetOBBVersionNumber(StoreVersion));
-                obbData.AppendFormat("{0} // the length of the file in bytes\n", File.Exists(ObbSource) ? new FileInfo(ObbSource).Length : 0);
+                obbData.AppendFormat("{0}L // the length of the file in bytes\n", File.Exists(ObbSource) ? new FileInfo(ObbSource).Length : 0);
                 obbData.AppendFormat("){0}\n", first ? "," : "");
                 first = false;
             }
