Skip vrapi library for unsupported Android platforms
#ue4
#android

diff --git a/Engine/Source/Programs/UnrealBuildTool/Android/AndroidToolChain.cs b/Engine/Source/Programs/UnrealBuildTool/Android/AndroidToolChain.cs
--- a/Engine/Source/Programs/UnrealBuildTool/Android/AndroidToolChain.cs
+++ b/Engine/Source/Programs/UnrealBuildTool/Android/AndroidToolChain.cs
@@ -31,9 +31,9 @@ namespace UnrealBuildTool
 
 		static private Dictionary<string, string[]> LibrariesToSkip = new Dictionary<string, string[]> {
 			{ "-armv7", new string[] { } }, 
-			{ "-arm64", new string[] { "nvToolsExt", "oculus", "gpg", } }, 
-			{ "-x86",   new string[] { "nvToolsExt", "oculus", } }, 
-			{ "-x64",   new string[] { "nvToolsExt", "oculus", "gpg", } }, 
+			{ "-arm64", new string[] { "nvToolsExt", "oculus", "vrapi", "gpg", } }, 
+			{ "-x86",   new string[] { "nvToolsExt", "oculus", "vrapi", } }, 
+			{ "-x64",   new string[] { "nvToolsExt", "oculus", "vrapi", "gpg", } }, 
 		};
 
 		static private Dictionary<string, string[]> ModulesToSkip = new Dictionary<string, string[]> {
