Fixed ParseAction.Remove.

diff --git a/Engine/Source/Programs/UnrealBuildTool/Configuration/EngineConfiguration.cs b/Engine/Source/Programs/UnrealBuildTool/Configuration/EngineConfiguration.cs
--- a/Engine/Source/Programs/UnrealBuildTool/Configuration/EngineConfiguration.cs
+++ b/Engine/Source/Programs/UnrealBuildTool/Configuration/EngineConfiguration.cs
@@ -514,7 +514,7 @@ namespace UnrealBuildTool
 						IniValues Value;
 						if (CurrentSection.TryGetValue(Key, out Value))
 						{
-							var ExistingIndex = Value.FindIndex(X => (String.Compare(Key, X, true) == 0));
+							var ExistingIndex = Value.FindIndex(X => (String.Compare(SingleValue, X, true) == 0));
 							if (ExistingIndex >= 0)
 							{
 								Value.RemoveAt(ExistingIndex);
