Missing file.

diff --git a/Engine/Source/Programs/UnrealBuildTool/Configuration/UEBuildTarget.cs b/Engine/Source/Programs/UnrealBuildTool/Configuration/UEBuildTarget.cs
--- a/Engine/Source/Programs/UnrealBuildTool/Configuration/UEBuildTarget.cs
+++ b/Engine/Source/Programs/UnrealBuildTool/Configuration/UEBuildTarget.cs
@@ -2749,7 +2749,7 @@ namespace UnrealBuildTool
 		protected virtual void SetupPlugins()
 		{
 			// Filter the plugins list by the current project
-			ValidPlugins = Plugins.ReadAvailablePlugins(UnrealBuildTool.GetUProjectFile());
+			ValidPlugins = Plugins.ReadAvailablePlugins(BuildConfiguration.RelativeEnginePath, UnrealBuildTool.GetUProjectFile());
 
 			// Remove any plugins for platforms we don't have
 			List<string> ExcludeFolders = new List<string>();
