Remove redundant diagnostic spew on actions with PrintDebugInfo enabled.

diff --git a/Engine/Source/Programs/UnrealBuildTool/System/LocalExecutor.cs b/Engine/Source/Programs/UnrealBuildTool/System/LocalExecutor.cs
--- a/Engine/Source/Programs/UnrealBuildTool/System/LocalExecutor.cs
+++ b/Engine/Source/Programs/UnrealBuildTool/System/LocalExecutor.cs
@@ -170,10 +170,6 @@ namespace UnrealBuildTool
 				{
 					Log.TraceVerbose("Executing: {0} {1}", ExpandedCommandPath, ActionStartInfo.Arguments);
 				}
-                if (Action.bPrintDebugInfo)
-                {
-                    Log.TraceInformation("Executing: {0} {1}", ExpandedCommandPath, ActionStartInfo.Arguments);
-                }
 				// Log summary if wanted.
 				else if (Action.bShouldOutputStatusDescription)
 				{
