Invalidate include caches if they didn't get saved on the last run (e.g. by UBT termination).

diff --git a/Engine/Source/Programs/UnrealBuildTool/System/DependencyCache.cs b/Engine/Source/Programs/UnrealBuildTool/System/DependencyCache.cs
--- a/Engine/Source/Programs/UnrealBuildTool/System/DependencyCache.cs
+++ b/Engine/Source/Programs/UnrealBuildTool/System/DependencyCache.cs
@@ -125,13 +125,25 @@ namespace UnrealBuildTool
 			DependencyCache Result = null;
 			try
 			{
-				using (FileStream Stream = new FileStream(Cache.AbsolutePath, FileMode.Open, FileAccess.Read))
+				string CacheBuildMutexPath = Cache.AbsolutePath + ".buildmutex";
+
+				// If the .buildmutex file for the cache is present, it means that something went wrong between loading
+				// and saving the cache last time (most likely the UBT process being terminated), so we don't want to load
+				// it.
+				if (!File.Exists(CacheBuildMutexPath))
 				{
-					BinaryFormatter Formatter = new BinaryFormatter();
-					Result = Formatter.Deserialize(Stream) as DependencyCache;
+					using (File.Create(CacheBuildMutexPath))
+					{
+					}
+
+					using (FileStream Stream = new FileStream(Cache.AbsolutePath, FileMode.Open, FileAccess.Read))
+					{
+						BinaryFormatter Formatter = new BinaryFormatter();
+						Result = Formatter.Deserialize(Stream) as DependencyCache;
+					}
+					Result.CreateFileExistsInfo();
+					Result.ResetUnresolvedDependencies();
 				}
-				Result.CreateFileExistsInfo();
-				Result.ResetUnresolvedDependencies();
 			}
 			catch (Exception Ex)
 			{
@@ -202,6 +214,8 @@ namespace UnrealBuildTool
 					Log.TraceInformation("IncludeFileCache did not need to be saved (bIsDirty=false)");
 				}
 			}
+
+			File.Delete(CachePath + ".buildmutex");
 		}
 
 		/**
diff --git a/Engine/Source/Programs/UnrealBuildTool/System/FlatCPPIncludeDepencencyCache.cs b/Engine/Source/Programs/UnrealBuildTool/System/FlatCPPIncludeDepencencyCache.cs
--- a/Engine/Source/Programs/UnrealBuildTool/System/FlatCPPIncludeDepencencyCache.cs
+++ b/Engine/Source/Programs/UnrealBuildTool/System/FlatCPPIncludeDepencencyCache.cs
@@ -88,14 +88,26 @@ namespace UnrealBuildTool
 			FlatCPPIncludeDependencyCache Result = null;
 			try
 			{
-				using (FileStream Stream = new FileStream(Cache.AbsolutePath, FileMode.Open, FileAccess.Read))
-				{	
-					// @todo ubtmake: We can store the cache in a cheaper/smaller way using hash file names and indices into included headers, but it might actually slow down load times
-					// @todo ubtmake: If we can index PCHs here, we can avoid storing all of the PCH's included headers (PCH's action should have been invalidated, so we shouldn't even have to report the PCH's includes as our indirect includes)
-					BinaryFormatter Formatter = new BinaryFormatter();
-					Result = Formatter.Deserialize(Stream) as FlatCPPIncludeDependencyCache;
-					Result.CacheFileItem = Cache;
-					Result.bIsDirty = false;
+				string CacheBuildMutexPath = Cache.AbsolutePath + ".buildmutex";
+
+				// If the .buildmutex file for the cache is present, it means that something went wrong between loading
+				// and saving the cache last time (most likely the UBT process being terminated), so we don't want to load
+				// it.
+				if (!File.Exists(CacheBuildMutexPath))
+				{
+					using (File.Create(CacheBuildMutexPath))
+					{
+					}
+
+					using (FileStream Stream = new FileStream(Cache.AbsolutePath, FileMode.Open, FileAccess.Read))
+					{
+						// @todo ubtmake: We can store the cache in a cheaper/smaller way using hash file names and indices into included headers, but it might actually slow down load times
+						// @todo ubtmake: If we can index PCHs here, we can avoid storing all of the PCH's included headers (PCH's action should have been invalidated, so we shouldn't even have to report the PCH's includes as our indirect includes)
+						BinaryFormatter Formatter = new BinaryFormatter();
+						Result = Formatter.Deserialize(Stream) as FlatCPPIncludeDependencyCache;
+						Result.CacheFileItem = Cache;
+						Result.bIsDirty = false;
+					}
 				}
 			}
 			catch (Exception Ex)
@@ -159,6 +171,8 @@ namespace UnrealBuildTool
 					Log.TraceInformation("FlatCPPIncludeDependencyCache did not need to be saved (bIsDirty=false)");
 				}
 			}
+
+			File.Delete(CacheFileItem.AbsolutePath + ".buildmutex");
 		}
 
 
