Swallow exceptions when failing to delete the .buildmutex file.

diff --git a/Engine/Source/Programs/UnrealBuildTool/System/DependencyCache.cs b/Engine/Source/Programs/UnrealBuildTool/System/DependencyCache.cs
--- a/Engine/Source/Programs/UnrealBuildTool/System/DependencyCache.cs
+++ b/Engine/Source/Programs/UnrealBuildTool/System/DependencyCache.cs
@@ -215,7 +215,14 @@ namespace UnrealBuildTool
 				}
 			}
 
-			File.Delete(CachePath + ".buildmutex");
+			try
+			{
+				File.Delete(CachePath + ".buildmutex");
+			}
+			catch
+			{
+				// We don't care if we couldn't delete this file, as maybe it couldn't have been created in the first place.
+			}
 		}
 
 		/**
diff --git a/Engine/Source/Programs/UnrealBuildTool/System/FlatCPPIncludeDepencencyCache.cs b/Engine/Source/Programs/UnrealBuildTool/System/FlatCPPIncludeDepencencyCache.cs
--- a/Engine/Source/Programs/UnrealBuildTool/System/FlatCPPIncludeDepencencyCache.cs
+++ b/Engine/Source/Programs/UnrealBuildTool/System/FlatCPPIncludeDepencencyCache.cs
@@ -172,7 +172,14 @@ namespace UnrealBuildTool
 				}
 			}
 
-			File.Delete(CacheFileItem.AbsolutePath + ".buildmutex");
+			try
+			{
+				File.Delete(CacheFileItem.AbsolutePath + ".buildmutex");
+			}
+			catch
+			{
+				// We don't care if we couldn't delete this file, as maybe it couldn't have been created in the first place.
+			}
 		}
 
 
