fix for deletion of precompiled modules in the engine when in binary release
UE-19029, UE-19203
#uat

diff --git a/Engine/Source/Programs/UnrealBuildTool/Configuration/UEBuildTarget.cs b/Engine/Source/Programs/UnrealBuildTool/Configuration/UEBuildTarget.cs
--- a/Engine/Source/Programs/UnrealBuildTool/Configuration/UEBuildTarget.cs
+++ b/Engine/Source/Programs/UnrealBuildTool/Configuration/UEBuildTarget.cs
@@ -1091,6 +1091,12 @@ namespace UnrealBuildTool
 						continue;
 					}
 
+					// if we're cleaning and in an installed or rocket binary and the build product path matches the engine install location, we don't want to delete them
+					if (UEBuildConfiguration.bCleanProject && (UnrealBuildTool.IsEngineInstalled() || UnrealBuildTool.RunningRocket()) && BuildProduct.Path.Contains(Path.GetFullPath(BuildConfiguration.RelativeEnginePath)))
+					{
+						continue;
+					}
+
 					// Don't add static libraries into the manifest unless we're explicitly building them; we don't submit them to Perforce.
 					if(!UEBuildConfiguration.bCleanProject && !bPrecompile && (BuildProduct.Type == BuildProductType.StaticLibrary || BuildProduct.Type == BuildProductType.ImportLibrary))
 					{
@@ -1336,8 +1342,14 @@ namespace UnrealBuildTool
 				var AllFilesToDelete = new List<string>();
 				foreach (BuildProduct BuildProduct in BuiltReceiptWithFullPaths.BuildProducts)
 				{
-					// If we're cleaning, don't add any precompiled binaries to the manifest. We don't want to delete them.
-					if (UEBuildConfiguration.bCleanProject && bUsePrecompiled && BuildProduct.IsPrecompiled)
+					// don't add any precompiled binaries to the manifest. We don't want to delete them.
+					if (bUsePrecompiled && BuildProduct.IsPrecompiled)
+					{
+						continue;
+					}
+
+					// don't add any installed or rocket binary and the build product path matches the engine install location, we don't want to delete them
+					if ((UnrealBuildTool.IsEngineInstalled() || UnrealBuildTool.RunningRocket()) && BuildProduct.Path.Contains(Path.GetFullPath(BuildConfiguration.RelativeEnginePath)))
 					{
 						continue;
 					}
