Only use -fno-function-sections on Clang 3.6 or later on Android
#ue4
#android

diff --git a/Engine/Source/Programs/UnrealBuildTool/Android/AndroidToolChain.cs b/Engine/Source/Programs/UnrealBuildTool/Android/AndroidToolChain.cs
--- a/Engine/Source/Programs/UnrealBuildTool/Android/AndroidToolChain.cs
+++ b/Engine/Source/Programs/UnrealBuildTool/Android/AndroidToolChain.cs
@@ -401,7 +401,10 @@ namespace UnrealBuildTool
 				if (CompileEnvironment.Config.Target.Configuration == CPPTargetConfiguration.Debug)
 				{
 					Result += " -fno-omit-frame-pointer";	// Disable removing the save/restore frame pointer for better debugging
-					Result += " -fno-function-sections";	// Improve breakpoint location
+					if (ClangVersionFloat >= 3.6f)
+					{
+						Result += " -fno-function-sections";	// Improve breakpoint location
+					}
 				}
 
 				// Some switches interfere with on-device debugging
