integration of PR#1408
UE-19108

diff --git a/Engine/Source/Programs/UnrealBuildTool/ToolChain/AppleToolChain.cs b/Engine/Source/Programs/UnrealBuildTool/ToolChain/AppleToolChain.cs
--- a/Engine/Source/Programs/UnrealBuildTool/ToolChain/AppleToolChain.cs
+++ b/Engine/Source/Programs/UnrealBuildTool/ToolChain/AppleToolChain.cs
@@ -36,6 +36,7 @@ namespace UnrealBuildTool
 					if (Results != null)
 					{
 						DeveloperDir = (string)Results["CommandOutput"];
+						DeveloperDir = DeveloperDir.TrimEnd();
 					}
 				}
 
