Remove plugin spam from UProjectInfo.cs

diff --git a/Engine/Source/Programs/UnrealBuildTool/System/UProjectInfo.cs b/Engine/Source/Programs/UnrealBuildTool/System/UProjectInfo.cs
--- a/Engine/Source/Programs/UnrealBuildTool/System/UProjectInfo.cs
+++ b/Engine/Source/Programs/UnrealBuildTool/System/UProjectInfo.cs
@@ -310,16 +310,13 @@ namespace UnrealBuildTool
 		public static bool IsPluginEnabledForProject(PluginInfo Plugin, ProjectDescriptor Project, UnrealTargetPlatform Platform)
 		{
 			bool bEnabled = Plugin.Descriptor.bEnabledByDefault || Plugin.LoadedFrom == PluginLoadedFrom.GameProject;
-			Log.TraceVerbose("Plugin passes first enabled check : {0}", Plugin.FileName);
 			if(Project != null && Project.Plugins != null)
 			{
 				foreach(PluginReferenceDescriptor PluginReference in Project.Plugins)
 				{
 					if(String.Compare(PluginReference.Name, Plugin.Name, true) == 0)
 					{
-						Log.TraceVerbose("Plugin passes project ref check: {0}", Plugin.FileName);
 						bEnabled = PluginReference.IsEnabledForPlatform(Platform);
-						Log.TraceVerbose("Plugin " + Plugin.FileName + "now " + bEnabled.ToString());
 					}
 				}
 			}
