Replace TADP with AndroidWorks
#ue4
#android

diff --git a/Engine/Build/InstalledEngineFilters.ini b/Engine/Build/InstalledEngineFilters.ini
--- a/Engine/Build/InstalledEngineFilters.ini
+++ b/Engine/Build/InstalledEngineFilters.ini
@@ -183,7 +183,7 @@
 {Mac} /Engine/Binaries/ThirdParty/ARM/Mac/...
 {Linux} /Engine/Binaries/ThirdParty/ARM/Linux32/...
 
-; TADP installer
+; AndroidWorks installer
 {Win64} /Engine/Extras/Android/*.exe
 {Mac} /Engine/Extras/Android/*.dmg
 
diff --git a/Engine/Source/Programs/UnrealBuildTool/Android/AndroidProjectGenerator.cs b/Engine/Source/Programs/UnrealBuildTool/Android/AndroidProjectGenerator.cs
--- a/Engine/Source/Programs/UnrealBuildTool/Android/AndroidProjectGenerator.cs
+++ b/Engine/Source/Programs/UnrealBuildTool/Android/AndroidProjectGenerator.cs
@@ -57,7 +57,7 @@ namespace UnrealBuildTool
 			if (NsightVersion.ProductMajorPart > 3)
 			{
 				// Mark as Nsight 3.1 (project will be updated)
-				NsightVersionCode = 10;
+				NsightVersionCode = 11;
 				NsightInstalled = true;
 			}
 			else if (NsightVersion.ProductMajorPart == 3)
@@ -69,7 +69,7 @@ namespace UnrealBuildTool
 				if (NsightVersion.ProductMinorPart >= 1)
 				{
 					// Nsight 3.1+ should be valid (will update project if newer)
-					NsightVersionCode = 10;
+					NsightVersionCode = 11;
 				}
 			}
 			else if (NsightVersion.ProductMajorPart == 2)
