When compiling a modular target, write out a manifest of build products to each output directory listing the valid modules in that directory and their version. Allows faster switching of changelists when using UGS (because it's no longer necessary to re-link all modules), and safeguards against stale DLLs from a different target causing loader errors. Uses compiled-in API version if BUILT_FROM_CHANGELIST is non-zero.

diff --git a/Engine/Source/Programs/UnrealBuildTool/Configuration/UEBuildTarget.cs b/Engine/Source/Programs/UnrealBuildTool/Configuration/UEBuildTarget.cs
--- a/Engine/Source/Programs/UnrealBuildTool/Configuration/UEBuildTarget.cs
+++ b/Engine/Source/Programs/UnrealBuildTool/Configuration/UEBuildTarget.cs
@@ -812,6 +812,9 @@ namespace UnrealBuildTool
 		/** Filename for the receipt for this target. */
 		private string ReceiptFileName;
 
+		/** Version manifests to be written to each output folder */
+		private Dictionary<string, VersionManifest> FileNameToVersionManifest;
+
 		/** Force output of the receipt to an additional filename */
 		[NonSerialized]
 		private string ForceReceiptFileName;
@@ -1256,13 +1259,13 @@ namespace UnrealBuildTool
 
 					// Normal makefile
 					{					
-					var UBTMakefilePath = UnrealBuildTool.GetUBTMakefilePath( TargetDescs );
-					if (File.Exists(UBTMakefilePath))
-					{
-						Log.TraceVerbose("\tDeleting " + UBTMakefilePath);
-						CleanFile(UBTMakefilePath);
+						var UBTMakefilePath = UnrealBuildTool.GetUBTMakefilePath( TargetDescs );
+						if (File.Exists(UBTMakefilePath))
+						{
+							Log.TraceVerbose("\tDeleting " + UBTMakefilePath);
+							CleanFile(UBTMakefilePath);
+						}
 					}
-				}
 
 					// Hot reload makefile
 					{					
@@ -1275,7 +1278,7 @@ namespace UnrealBuildTool
 						}
 						UEBuildConfiguration.bHotReloadFromIDE = false;
 					}
-			}
+				}
 
 				// Delete the action history
 				{
@@ -1294,6 +1297,7 @@ namespace UnrealBuildTool
 					CleanUnrealHeaderTool();
 				}
 			}
+
 		}
 
 		class BuldProductComparer : IEqualityComparer<BuildProduct>
@@ -1572,6 +1576,74 @@ namespace UnrealBuildTool
 			}
 		}
 
+		/** Prepare all the version manifests for this target. See the VersionManifest class for an explanation of what these files are. */
+		public void PrepareVersionManifests()
+		{
+			FileNameToVersionManifest = new Dictionary<string,VersionManifest>(StringComparer.InvariantCultureIgnoreCase);
+			if(!bCompileMonolithic)
+			{
+				// Read the version file
+				BuildVersion Version;
+				if(!Utils.TryReadBuildVersion(Path.Combine(BuildConfiguration.RelativeEnginePath, "Build", "Build.version"), out Version))
+				{
+					Version = new BuildVersion();
+				}
+
+				// Create a unique identifier for this build, which can be used to identify modules when the changelist is constant. It's fine to share this between runs with the same makefile; 
+				// the output won't change. By default we leave it blank when compiling a subset of modules (for hot reload, etc...), otherwise it won't match anything else. When writing to a directory
+				// that already contains a manifest, we'll reuse the build id that's already in there (see below).
+				string BuildId = (OnlyModules.Count == 0)? Guid.NewGuid().ToString() : "";
+
+				// Create the receipts for each folder
+				foreach(UEBuildBinary Binary in AppBinaries)
+				{
+					string DirectoryName = Path.GetDirectoryName(Binary.Config.OutputFilePath);
+					bool bIsGameDirectory = !Utils.IsFileUnderDirectory(Binary.Config.OutputFilePath, BuildConfiguration.RelativeEnginePath);
+					string ManifestFileName = Path.Combine(DirectoryName, VersionManifest.GetStandardFileName(AppName, Platform, Configuration, UEBuildPlatform.GetBuildPlatform(Platform).GetActiveArchitecture(), bIsGameDirectory));
+	
+					VersionManifest Manifest;
+					if(!FileNameToVersionManifest.TryGetValue(ManifestFileName, out Manifest))
+					{
+						Manifest = new VersionManifest(Version.Changelist, BuildId);
+
+						VersionManifest ExistingManifest;
+						if(VersionManifest.TryRead(ManifestFileName, out ExistingManifest) && Version.Changelist == ExistingManifest.Changelist)
+						{
+							if(OnlyModules.Count > 0)
+							{
+								// We're just building an existing module; reuse the existing manifest AND build id.
+								Manifest = ExistingManifest;
+							}
+							else if(Version.Changelist != 0)
+							{
+								// We're rebuilding at the same changelist. Keep all the existing binaries.
+								Manifest.ModuleNameToFileName.Union(ExistingManifest.ModuleNameToFileName);
+							}
+						}
+
+						FileNameToVersionManifest.Add(ManifestFileName, Manifest);
+					}
+
+					foreach(string ModuleName in Binary.Config.ModuleNames)
+					{
+						Manifest.ModuleNameToFileName[ModuleName] = Path.GetFileName(Binary.Config.OutputFilePath);
+					}
+				}
+			}
+		}
+
+		/** Writes out the version manifest */
+		public void WriteVersionManifests()
+		{
+			if(FileNameToVersionManifest != null)
+			{
+				foreach(KeyValuePair<string, VersionManifest> ManifestPair in FileNameToVersionManifest)
+				{
+					ManifestPair.Value.Write(ManifestPair.Key);
+				}
+			}
+		}
+
 		/** Creates the receipt for the target */
 		private void PrepareReceipt(IUEToolChain ToolChain)
 		{
@@ -1719,6 +1791,9 @@ namespace UnrealBuildTool
 				return ECompilationResult.Succeeded;
 			}
 
+			// Create all the version manifests
+			PrepareVersionManifests();
+
 			// Create a receipt for the target
 			PrepareReceipt(TargetToolChain);
 
diff --git a/Engine/Source/Programs/UnrealBuildTool/System/PluginDescriptor.cs b/Engine/Source/Programs/UnrealBuildTool/System/PluginDescriptor.cs
--- a/Engine/Source/Programs/UnrealBuildTool/System/PluginDescriptor.cs
+++ b/Engine/Source/Programs/UnrealBuildTool/System/PluginDescriptor.cs
@@ -88,7 +88,7 @@ namespace UnrealBuildTool
 		/// <returns>New plugin descriptor</returns>
 		public static PluginDescriptor FromFile(string FileName)
 		{
-			JsonObject RawObject = JsonObject.FromFile(FileName);
+			JsonObject RawObject = JsonObject.Read(FileName);
 			try
 			{
 				PluginDescriptor Descriptor = new PluginDescriptor();
diff --git a/Engine/Source/Programs/UnrealBuildTool/System/ProjectDescriptor.cs b/Engine/Source/Programs/UnrealBuildTool/System/ProjectDescriptor.cs
--- a/Engine/Source/Programs/UnrealBuildTool/System/ProjectDescriptor.cs
+++ b/Engine/Source/Programs/UnrealBuildTool/System/ProjectDescriptor.cs
@@ -59,7 +59,7 @@ namespace UnrealBuildTool
 		/// <returns>New plugin descriptor</returns>
 		public static ProjectDescriptor FromFile(string FileName)
 		{
-			JsonObject RawObject = JsonObject.FromFile(FileName);
+			JsonObject RawObject = JsonObject.Read(FileName);
 			try
 			{
 				ProjectDescriptor Descriptor = new ProjectDescriptor();
diff --git a/Engine/Source/Programs/UnrealBuildTool/System/UnrealBuildTool.cs b/Engine/Source/Programs/UnrealBuildTool/System/UnrealBuildTool.cs
--- a/Engine/Source/Programs/UnrealBuildTool/System/UnrealBuildTool.cs
+++ b/Engine/Source/Programs/UnrealBuildTool/System/UnrealBuildTool.cs
@@ -1976,6 +1976,7 @@ namespace UnrealBuildTool
                             {
                                 foreach (UEBuildTarget Target in Targets)
                                 {
+									Target.WriteVersionManifests();
 									Target.WriteReceipt();
                                     ToolChain.PostBuildSync(Target);
                                 }
diff --git a/Engine/Source/Programs/UnrealBuildTool/System/Utils.cs b/Engine/Source/Programs/UnrealBuildTool/System/Utils.cs
--- a/Engine/Source/Programs/UnrealBuildTool/System/Utils.cs
+++ b/Engine/Source/Programs/UnrealBuildTool/System/Utils.cs
@@ -829,11 +829,7 @@ namespace UnrealBuildTool
 		public static bool TryReadBuildVersion(string FileName, out BuildVersion Version)
 		{
 			JsonObject Object;
-			try
-			{
-				Object = JsonObject.FromFile(FileName);
-			}
-			catch(Exception)
+			if(!JsonObject.TryRead(FileName, out Object))
 			{
 				Version = null;
 				return false;
diff --git a/Engine/Source/Programs/UnrealBuildTool/System/VersionManifest.cs b/Engine/Source/Programs/UnrealBuildTool/System/VersionManifest.cs
new file mode 100644
--- /dev/null
+++ b/Engine/Source/Programs/UnrealBuildTool/System/VersionManifest.cs
@@ -0,0 +1,141 @@
+﻿using System;
+using System.Collections.Generic;
+using System.IO;
+using System.Linq;
+using System.Text;
+using System.Threading.Tasks;
+
+namespace UnrealBuildTool
+{
+	/// <summary>
+	/// Stores the version (or a unique build ID) for the build products of an application in a certain folder.
+	/// 
+	/// Allows the runtime to identify which modules are used for which files, and at which version they're at. This prevents stale binaries from being
+	/// loaded by the runtime when making a local unversioned build, and allows faster incremental builds than compiling the build changelist into every 
+	/// module when making versioned builds.
+	/// </summary>
+	[Serializable]
+	class VersionManifest
+	{
+		public readonly int Changelist;
+		public readonly string BuildId;
+		public Dictionary<string, string> ModuleNameToFileName = new Dictionary<string, string>();
+
+		/// <summary>
+		/// Constructs the module map with the given changelist
+		/// </summary>
+		/// <param name="InChangelist"></param>
+		public VersionManifest(int InChangelist, string InBuildId)
+		{
+			Changelist = InChangelist;
+			BuildId = InBuildId;
+		}
+
+		/// <summary>
+		/// Merge another manifest into this one
+		/// </summary>
+		/// <param name="Other">The manifest to merge in</param>
+		public void Include(VersionManifest Other)
+		{
+			foreach(KeyValuePair<string, string> Pair in Other.ModuleNameToFileName)
+			{
+				if(!ModuleNameToFileName.ContainsKey(Pair.Key))
+				{
+					ModuleNameToFileName.Add(Pair.Key, Pair.Value);
+				}
+			}
+		}
+
+		/// <summary>
+		/// Gets the standard path for an manifest
+		/// </summary>
+		/// <param name="DirectoryName">The directory containing this manifest</param>
+		/// <param name="AppName">The modular app name being built</param>
+		/// <param name="Configuration">The target configuration</param>
+		/// <param name="Platform">The target platform</param>
+		/// <param name="BuildArchitecture">The architecture of the target platform</param>
+		/// <returns>Filename for the app receipt</returns>
+		public static string GetStandardFileName(string AppName, UnrealTargetPlatform Platform, UnrealTargetConfiguration Configuration, string BuildArchitecture, bool bIsGameDirectory)
+		{
+			string BaseName = AppName;
+			if(Configuration != UnrealTargetConfiguration.Development && !(Configuration == UnrealTargetConfiguration.DebugGame && !bIsGameDirectory))
+			{
+				BaseName += String.Format("-{0}-{1}", Platform.ToString(), Configuration.ToString());
+			}
+			return String.Format("{0}{1}.modules", BaseName, BuildArchitecture);
+		}
+
+		/// <summary>
+		/// Read an app receipt from disk
+		/// </summary>
+		/// <param name="FileName">Filename to read from</param>
+		/// <returns>The receipt that was read</returns>
+		public static VersionManifest Read(string FileName)
+		{
+			JsonObject Object = JsonObject.Read(FileName);
+
+			VersionManifest Receipt = new VersionManifest(Object.GetIntegerField("Changelist"), Object.GetStringField("BuildId"));
+
+			JsonObject Modules = Object.GetObjectField("Modules");
+			foreach(string ModuleName in Modules.KeyNames)
+			{
+				Receipt.ModuleNameToFileName.Add(ModuleName, Modules.GetStringField(ModuleName));
+			}
+			return Receipt;
+		}
+
+		/// <summary>
+		/// Tries to read a receipt from disk.
+		/// </summary>
+		/// <param name="FileName">The filename that was read</param>
+		/// <param name="Receipt">If successful, the receipt that was read. Null otherwise.</param>
+		/// <returns>True if the file was read succesfully.</returns>
+		public static bool TryRead(string FileName, out VersionManifest Result)
+		{
+			if(!File.Exists(FileName))
+			{
+				Result = null;
+				return false;
+			}
+			try
+			{
+				Result = Read(FileName);
+				return true;
+			}
+			catch(Exception)
+			{
+				Result = null;
+				return false;
+			}
+		}
+
+		/// <summary>
+		/// Write the receipt to disk.
+		/// </summary>
+		/// <param name="EngineDir">The engine directory</param>
+		/// <param name="ProjectDir">The project directory</param>
+		/// <param name="AppName">The modular app name being built</param>
+		/// <param name="Configuration">The target configuration</param>
+		/// <param name="Platform">The target platform</param>
+		/// <param name="BuildArchitecture">The architecture of the target platform</param>
+		public void Write(string FileName)
+		{
+			using(JsonWriter Writer = new JsonWriter(FileName))
+			{
+				Writer.WriteObjectStart();
+
+				Writer.WriteValue("Changelist", Changelist);
+				Writer.WriteValue("BuildId", BuildId);
+
+				Writer.WriteObjectStart("Modules");
+				foreach(KeyValuePair<string, string> ModuleNameToFileNamePair in ModuleNameToFileName)
+				{
+					Writer.WriteValue(ModuleNameToFileNamePair.Key, ModuleNameToFileNamePair.Value);
+				}
+				Writer.WriteObjectEnd();
+
+				Writer.WriteObjectEnd();
+			}
+		}
+	}
+}
diff --git a/Engine/Source/Programs/UnrealBuildTool/UnrealBuildTool.csproj b/Engine/Source/Programs/UnrealBuildTool/UnrealBuildTool.csproj
--- a/Engine/Source/Programs/UnrealBuildTool/UnrealBuildTool.csproj
+++ b/Engine/Source/Programs/UnrealBuildTool/UnrealBuildTool.csproj
@@ -39,7 +39,7 @@
     <Optimize>true</Optimize>
     <OutputPath>..\..\..\Binaries\DotNET\</OutputPath>
     <DefineConstants>TRACE</DefineConstants>
-    <Optimize>true</Optimize>
+    <Optimize>false</Optimize>
     <DebugType>pdbonly</DebugType>
     <PlatformTarget>AnyCPU</PlatformTarget>
     <ErrorReport>prompt</ErrorReport>
@@ -125,6 +125,7 @@
     <Compile Include="Mac\MacProjectGenerator.cs" />
     <Compile Include="Mac\UEDeployMac.cs" />
     <Compile Include="System\TargetReceipt.cs" />
+    <Compile Include="System\VersionManifest.cs" />
     <Compile Include="ToolChain\AppleToolChain.cs" />
   </ItemGroup>
   <ItemGroup>
diff --git a/Engine/Source/Programs/UnrealBuildTool/Utilities/JsonObject.cs b/Engine/Source/Programs/UnrealBuildTool/Utilities/JsonObject.cs
--- a/Engine/Source/Programs/UnrealBuildTool/Utilities/JsonObject.cs
+++ b/Engine/Source/Programs/UnrealBuildTool/Utilities/JsonObject.cs
@@ -22,13 +22,38 @@ namespace UnrealBuildTool
 			RawObject = new Dictionary<string,object>(InRawObject, StringComparer.InvariantCultureIgnoreCase);
 		}
 
-		public static JsonObject FromFile(string FileName)
+		public static JsonObject Read(string FileName)
 		{
 			string Text = File.ReadAllText(FileName);
-			Dictionary<string, object> CaseSensitiveRawObject = fastJSON.JSON.Instance.ToObject<Dictionary<string, object>>(Text);
+			Dictionary<string, object> CaseSensitiveRawObject = (Dictionary<string, object>)fastJSON.JSON.Instance.Parse(Text);
 			return new JsonObject(CaseSensitiveRawObject);
 		}
 
+		public static bool TryRead(string FileName, out JsonObject Result)
+		{
+			if(!File.Exists(FileName))
+			{
+				Result = null;
+				return false;
+			}
+
+			try
+			{
+				Result = Read(FileName);
+				return true;
+			}
+			catch(Exception)
+			{
+				Result = null;
+				return false;
+			}
+		}
+
+		public IEnumerable<string> KeyNames
+		{
+			get { return RawObject.Keys; }
+		}
+
 		public string GetStringField(string FieldName)
 		{
 			string StringValue;
@@ -67,9 +92,9 @@ namespace UnrealBuildTool
 		public bool TryGetStringArrayField(string FieldName, out string[] Result)
 		{
 			object RawValue;
-			if(RawObject.TryGetValue(FieldName, out RawValue) && (RawValue is object[]) && ((object[])RawValue).All(x => x is string))
+			if(RawObject.TryGetValue(FieldName, out RawValue) && (RawValue is IEnumerable<object>) && ((IEnumerable<object>)RawValue).All(x => x is string))
 			{
-				Result = Array.ConvertAll((object[])RawValue, x => (string)x);
+				Result = ((IEnumerable<object>)RawValue).Select(x => (string)x).ToArray();
 				return true;
 			}
 			else
@@ -104,6 +129,16 @@ namespace UnrealBuildTool
 			}
 		}
 
+		public int GetIntegerField(string FieldName)
+		{
+			int IntegerValue;
+			if(!TryGetIntegerField(FieldName, out IntegerValue))
+			{
+				throw new JsonParseException("Missing or invalid '{0}' field", FieldName);
+			}
+			return IntegerValue;
+		}
+
 		public bool TryGetIntegerField(string FieldName, out int Result)
 		{
 			object RawValue;
@@ -170,12 +205,37 @@ namespace UnrealBuildTool
 			return true;
 		}
 
+		public JsonObject GetObjectField(string FieldName)
+		{
+			JsonObject Result;
+			if(!TryGetObjectField(FieldName, out Result))
+			{
+				throw new JsonParseException("Missing or invalid '{0}' field", FieldName);
+			}
+			return Result;
+		}
+
+		public bool TryGetObjectField(string FieldName, out JsonObject Result)
+		{
+			object RawValue;
+			if(RawObject.TryGetValue(FieldName, out RawValue) && (RawValue is Dictionary<string, object>))
+			{
+				Result = new JsonObject((Dictionary<string, object>)RawValue);
+				return true;
+			}
+			else
+			{
+				Result = null;
+				return false;
+			}
+		}
+
 		public bool TryGetObjectArrayField(string FieldName, out JsonObject[] Result)
 		{
 			object RawValue;
-			if(RawObject.TryGetValue(FieldName, out RawValue) && (RawValue is object[]) && ((object[])RawValue).All(x => x is Dictionary<string, object>))
+			if(RawObject.TryGetValue(FieldName, out RawValue) && (RawValue is IEnumerable<object>) && ((IEnumerable<object>)RawValue).All(x => x is Dictionary<string, object>))
 			{
-				Result = Array.ConvertAll((object[])RawValue, x => new JsonObject((Dictionary<string, object>)x));
+				Result = ((IEnumerable<object>)RawValue).Select(x => new JsonObject((Dictionary<string, object>)x)).ToArray();
 				return true;
 			}
 			else
@@ -214,6 +274,17 @@ namespace UnrealBuildTool
 			bRequiresComma = false;
 		}
 
+		public void WriteObjectStart(string ObjectName)
+		{
+			WriteCommaNewline();
+
+			Writer.Write("{0}\"{1}\" : ", Indent, ObjectName);
+
+			bRequiresComma = false;
+
+			WriteObjectStart();
+		}
+
 		public void WriteObjectEnd()
 		{
 			Indent = Indent.Substring(0, Indent.Length - 1);
@@ -257,7 +328,7 @@ namespace UnrealBuildTool
 
 		public void WriteValue(string Name, string Value)
 		{
-			WriteValueInternal(Name, '"' + Value + '"');
+			WriteValueInternal(Name, '"' + Value.Replace("\\", "\\\\") + '"');
 		}
 
 		public void WriteValue(string Name, int Value)
diff --git a/Engine/Source/Runtime/Core/Private/Modules/ModuleManager.cpp b/Engine/Source/Runtime/Core/Private/Modules/ModuleManager.cpp
--- a/Engine/Source/Runtime/Core/Private/Modules/ModuleManager.cpp
+++ b/Engine/Source/Runtime/Core/Private/Modules/ModuleManager.cpp
@@ -940,34 +940,52 @@ void FModuleManager::FindModulePaths(const TCHAR* NamePattern, TMap<FName, FStri
 
 void FModuleManager::FindModulePathsInDirectory(const FString& InDirectoryName, bool bIsGameDirectory, const TCHAR* NamePattern, TMap<FName, FString> &OutModulePaths) const
 {
-	// Get the prefix and suffix for module filenames
-	FString ModulePrefix, ModuleSuffix;
-	GetModuleFilenameFormat(bIsGameDirectory, ModulePrefix, ModuleSuffix);
-
-	// Find all the files
-	TArray<FString> FullFileNames;
-	IFileManager::Get().FindFilesRecursive(FullFileNames, *InDirectoryName, *(ModulePrefix + NamePattern + ModuleSuffix), true, false);
-
-	// Parse all the matching module names
-	for (int32 Idx = 0; Idx < FullFileNames.Num(); Idx++)
+	if(QueryModulesDelegate.IsBound())
 	{
-		const FString &FullFileName = FullFileNames[Idx];
-	
-		// On Mac OS X the separate debug symbol format is the dSYM bundle, which is a bundle folder hierarchy containing a .dylib full of Mach-O formatted DWARF debug symbols, these are not loadable modules, so we mustn't ever try and use them. If we don't eliminate this here then it will appear in the module paths & cause errors later on which cannot be recovered from.
-	#if PLATFORM_MAC
-		if(FullFileName.Contains(".dSYM"))
+		// Use the delegate to query all the modules in this directory
+		TMap<FString, FString> Modules;
+		QueryModulesDelegate.Execute(InDirectoryName, bIsGameDirectory, Modules);
+
+		// Fill the output map with modules that match the wildcard
+		for(const TPair<FString, FString>& Pair: Modules)
 		{
-			continue;
+			if(Pair.Key.MatchesWildcard(NamePattern))
+			{
+				OutModulePaths.Add(FName(*Pair.Key), FPaths::Combine(*InDirectoryName, *Pair.Value));
+			}
 		}
-	#endif
-		
-		FString FileName = FPaths::GetCleanFilename(FullFileName);
-		if (FileName.StartsWith(ModulePrefix) && FileName.EndsWith(ModuleSuffix))
+	}
+	else
+	{
+		// Get the prefix and suffix for module filenames
+		FString ModulePrefix, ModuleSuffix;
+		GetModuleFilenameFormat(bIsGameDirectory, ModulePrefix, ModuleSuffix);
+
+		// Find all the files
+		TArray<FString> FullFileNames;
+		IFileManager::Get().FindFilesRecursive(FullFileNames, *InDirectoryName, *(ModulePrefix + NamePattern + ModuleSuffix), true, false);
+
+		// Parse all the matching module names
+		for (int32 Idx = 0; Idx < FullFileNames.Num(); Idx++)
 		{
-			FString ModuleName = FileName.Mid(ModulePrefix.Len(), FileName.Len() - ModulePrefix.Len() - ModuleSuffix.Len());
-			if (!ModuleName.EndsWith("-Debug") && !ModuleName.EndsWith("-Shipping") && !ModuleName.EndsWith("-Test") && !ModuleName.EndsWith("-DebugGame"))
+			const FString &FullFileName = FullFileNames[Idx];
+	
+			// On Mac OS X the separate debug symbol format is the dSYM bundle, which is a bundle folder hierarchy containing a .dylib full of Mach-O formatted DWARF debug symbols, these are not loadable modules, so we mustn't ever try and use them. If we don't eliminate this here then it will appear in the module paths & cause errors later on which cannot be recovered from.
+		#if PLATFORM_MAC
+			if(FullFileName.Contains(".dSYM"))
 			{
-				OutModulePaths.Add(FName(*ModuleName), FullFileName);
+				continue;
+			}
+		#endif
+		
+			FString FileName = FPaths::GetCleanFilename(FullFileName);
+			if (FileName.StartsWith(ModulePrefix) && FileName.EndsWith(ModuleSuffix))
+			{
+				FString ModuleName = FileName.Mid(ModulePrefix.Len(), FileName.Len() - ModulePrefix.Len() - ModuleSuffix.Len());
+				if (!ModuleName.EndsWith("-Debug") && !ModuleName.EndsWith("-Shipping") && !ModuleName.EndsWith("-Test") && !ModuleName.EndsWith("-DebugGame"))
+				{
+					OutModulePaths.Add(FName(*ModuleName), FullFileName);
+				}
 			}
 		}
 	}
diff --git a/Engine/Source/Runtime/Core/Public/Modules/ModuleManager.h b/Engine/Source/Runtime/Core/Public/Modules/ModuleManager.h
--- a/Engine/Source/Runtime/Core/Public/Modules/ModuleManager.h
+++ b/Engine/Source/Runtime/Core/Public/Modules/ModuleManager.h
@@ -216,6 +216,11 @@ public:
 	 */
 	void UnloadOrAbandonModuleWithCallback( const FName InModuleName, FOutputDevice &Ar, bool bAbandonOnly = false);
 
+	/** Delegate that's used by the module manager to find all the valid modules in a directory matching a pattern */
+	typedef TMap<FString, FString> FModuleNamesMap;
+	DECLARE_DELEGATE_ThreeParams( FQueryModulesDelegate, const FString&, bool, FModuleNamesMap& );
+	FQueryModulesDelegate QueryModulesDelegate;
+
 public:
 
 	/**
diff --git a/Engine/Source/Runtime/Launch/Private/LaunchEngineLoop.cpp b/Engine/Source/Runtime/Launch/Private/LaunchEngineLoop.cpp
--- a/Engine/Source/Runtime/Launch/Private/LaunchEngineLoop.cpp
+++ b/Engine/Source/Runtime/Launch/Private/LaunchEngineLoop.cpp
@@ -14,6 +14,8 @@
 #include "EngineVersion.h"
 
 #include "ModuleManager.h"
+#include "../Resources/Version.h"
+#include "VersionManifest.h"
 
 #if WITH_EDITOR
 	#include "EditorStyle.h"
@@ -1028,6 +1030,16 @@ int32 FEngineLoop::PreInit( const TCHAR* CmdLine )
 		UE_LOG(LogInit, Display, TEXT("RandInit(%d) SRandInit(%d)."), Cycles1, Cycles2);
 	}
 
+	// Set up the module list and version information, if it's not compiled-in
+#if !IS_MONOLITHIC && BUILT_FROM_CHANGELIST == 0
+	static FVersionedModuleEnumerator ModuleEnumerator;
+	if(ModuleEnumerator.RegisterWithModuleManager())
+	{
+		const FVersionManifest& Manifest = ModuleEnumerator.GetInitialManifest();
+		UE_LOG(LogInit, Log, TEXT("Using version manifest at CL %d with build ID '%s'"), Manifest.Changelist, *Manifest.BuildId);
+	}
+#endif
+
 #if !IS_PROGRAM
 	if ( !GIsGameAgnosticExe && FApp::HasGameName() && !FPaths::IsProjectFilePathSet() )
 	{
