Starting to check in experimental Fast Semantics support, starting with a simple switch (disabled).

diff --git a/Engine/Source/Programs/UnrealBuildTool/Configuration/BuildConfiguration.cs b/Engine/Source/Programs/UnrealBuildTool/Configuration/BuildConfiguration.cs
--- a/Engine/Source/Programs/UnrealBuildTool/Configuration/BuildConfiguration.cs
+++ b/Engine/Source/Programs/UnrealBuildTool/Configuration/BuildConfiguration.cs
@@ -49,6 +49,20 @@ namespace UnrealBuildTool
 		[XmlConfig]
 		public static bool bUseFastMonoCalls;
 
+        /// <summary>
+        /// New Xbox driver supports a "fast semantics" context type. This switches it on for the immediate context
+        /// EXPERIMENTAL - WILL CAUSE RENDERING ISSUES AND/OR CRASHES AT PRESENT!
+        /// </summary>
+        [XmlConfig]
+        public static bool bUseFastSemanticsImmediateContext;
+
+        /// <summary>
+        /// New Xbox driver supports a "fast semantics" context type. This switches it on for the deferred contexts
+        /// EXPERIMENTAL - WILL CAUSE RENDERING ISSUES AND/OR CRASHES AT PRESENT!
+        /// </summary>
+        [XmlConfig]
+        public static bool bUseFastSemanticsDeferredContexts;
+
         /// Async Compute context support. Requires Mono and Fastcalls
         /// </summary>
         [XmlConfig]
@@ -550,6 +564,11 @@ namespace UnrealBuildTool
             bUseFastMonoCalls = true;
             bUseAsyncComputeContext = bUseFastMonoCalls;
 
+            // Switches for fast semantics D3D contexts
+            // EXPERIMENTAL - WILL CAUSE RENDERING ISSUES AND/OR CRASHES AT PRESENT!
+            bUseFastSemanticsImmediateContext = false;
+            bUseFastSemanticsDeferredContexts = false;
+
 			// By default we use the Release C++ Runtime (CRT), even when compiling Debug builds.  This is because the Debug C++
 			// Runtime isn't very useful when debugging Unreal Engine projects, and linking against the Debug CRT libraries forces
 			// our third party library dependencies to also be compiled using the Debug CRT (and often perform more slowly.)  Often
