Makefiles enabled on Mac.
Explicit serialization functions used instead of reflection, to avoid deserialization problems on Mono.
Version number added to makefiles to allow versioning in future.
Makefiles still disabled if the host platform differs from the built platform, to continue to avoid any problems 'when building with RPCUtility'.

diff --git a/Engine/Source/Programs/UnrealBuildTool/Configuration/BuildConfiguration.cs b/Engine/Source/Programs/UnrealBuildTool/Configuration/BuildConfiguration.cs
--- a/Engine/Source/Programs/UnrealBuildTool/Configuration/BuildConfiguration.cs
+++ b/Engine/Source/Programs/UnrealBuildTool/Configuration/BuildConfiguration.cs
@@ -586,7 +586,7 @@ namespace UnrealBuildTool
 			// Enables support for fast include dependency scanning, as well as gathering data for 'UBT Makefiles', then quickly
 			// assembling builds in subsequent runs using data in those cached makefiles
 			// NOTE: This feature is new and has a number of known issues (search the code for '@todo ubtmake')
-			bUseUBTMakefiles = !Utils.IsRunningOnMono;	// @todo ubtmake: Needs support for Mac
+			bUseUBTMakefiles = true;
 
 			// Distcc requires some setup - so by default disable it so we don't break local or remote building
 			bAllowDistcc = false;
diff --git a/Engine/Source/Programs/UnrealBuildTool/Configuration/UEBuildClient.cs b/Engine/Source/Programs/UnrealBuildTool/Configuration/UEBuildClient.cs
--- a/Engine/Source/Programs/UnrealBuildTool/Configuration/UEBuildClient.cs
+++ b/Engine/Source/Programs/UnrealBuildTool/Configuration/UEBuildClient.cs
@@ -5,12 +5,18 @@ using System.Collections.Generic;
 using System.IO;
 using System.Linq;
 using System.Text;
+using System.Runtime.Serialization;
 
 namespace UnrealBuildTool
 {
 	[Serializable]
     public class UEBuildClient : UEBuildTarget
 	{
+		public UEBuildClient(SerializationInfo Info, StreamingContext Context)
+			: base(Info, Context)
+		{
+		}
+
 		public UEBuildClient(TargetDescriptor InDesc, TargetRules InRulesObject, string InTargetCsFilename)
             : base(InDesc, InRulesObject, "UE4Client", InTargetCsFilename)
         {
diff --git a/Engine/Source/Programs/UnrealBuildTool/Configuration/UEBuildEditor.cs b/Engine/Source/Programs/UnrealBuildTool/Configuration/UEBuildEditor.cs
--- a/Engine/Source/Programs/UnrealBuildTool/Configuration/UEBuildEditor.cs
+++ b/Engine/Source/Programs/UnrealBuildTool/Configuration/UEBuildEditor.cs
@@ -4,12 +4,18 @@ using System;
 using System.Collections.Generic;
 using System.IO;
 using System.Text;
+using System.Runtime.Serialization;
 
 namespace UnrealBuildTool
 {
 	[Serializable]
 	public class UEBuildEditor : UEBuildTarget
 	{
+		public UEBuildEditor(SerializationInfo Info, StreamingContext Context)
+			: base(Info, Context)
+		{
+		}
+
 		public UEBuildEditor(TargetDescriptor InDesc, TargetRules InRulesObject, string InTargetCsFilename)
 			: base(InDesc, InRulesObject, "UE4Editor", InTargetCsFilename)
 		{
diff --git a/Engine/Source/Programs/UnrealBuildTool/Configuration/UEBuildGame.cs b/Engine/Source/Programs/UnrealBuildTool/Configuration/UEBuildGame.cs
--- a/Engine/Source/Programs/UnrealBuildTool/Configuration/UEBuildGame.cs
+++ b/Engine/Source/Programs/UnrealBuildTool/Configuration/UEBuildGame.cs
@@ -5,12 +5,18 @@ using System.Collections.Generic;
 using System.IO;
 using System.Linq;
 using System.Text;
+using System.Runtime.Serialization;
 
 namespace UnrealBuildTool
 {
 	[Serializable]
 	public class UEBuildGame : UEBuildTarget
 	{
+		public UEBuildGame(SerializationInfo Info, StreamingContext Context)
+			: base(Info, Context)
+		{
+		}
+
 		public UEBuildGame(TargetDescriptor InDesc, TargetRules InRulesObject, string InTargetCsFilename)
 			: base(InDesc, InRulesObject, "UE4", InTargetCsFilename)
 		{
diff --git a/Engine/Source/Programs/UnrealBuildTool/Configuration/UEBuildServer.cs b/Engine/Source/Programs/UnrealBuildTool/Configuration/UEBuildServer.cs
--- a/Engine/Source/Programs/UnrealBuildTool/Configuration/UEBuildServer.cs
+++ b/Engine/Source/Programs/UnrealBuildTool/Configuration/UEBuildServer.cs
@@ -4,12 +4,18 @@ using System;
 using System.Collections.Generic;
 using System.IO;
 using System.Text;
+using System.Runtime.Serialization;
 
 namespace UnrealBuildTool
 {
 	[Serializable]
 	public class UEBuildServer : UEBuildTarget
 	{
+		public UEBuildServer(SerializationInfo Info, StreamingContext Context)
+			: base(Info, Context)
+		{
+		}
+
 		public UEBuildServer(TargetDescriptor InDesc, TargetRules InRulesObject, string InTargetCsFilename)
 			: base(InDesc, InRulesObject, "UE4Server", InTargetCsFilename)
 		{
diff --git a/Engine/Source/Programs/UnrealBuildTool/Configuration/UEBuildTarget.cs b/Engine/Source/Programs/UnrealBuildTool/Configuration/UEBuildTarget.cs
--- a/Engine/Source/Programs/UnrealBuildTool/Configuration/UEBuildTarget.cs
+++ b/Engine/Source/Programs/UnrealBuildTool/Configuration/UEBuildTarget.cs
@@ -182,21 +182,8 @@ namespace UnrealBuildTool
 	 * A target that can be built
 	 */
 	[Serializable]
-	public class UEBuildTarget : IDeserializationCallback
+	public class UEBuildTarget : ISerializable
 	{
-		void IDeserializationCallback.OnDeserialization(object sender)
-		{
-			if (OnlyModules == null)
-			{
-				OnlyModules = new List<OnlyModule>();
-			}
-
-			if (FlatModuleCsData == null)
-			{
-				FlatModuleCsData = new CaselessDictionary<FlatModuleCsDataType>();
-			}
-		}
-
 		public string GetAppName()
 		{
 			return AppName;
@@ -791,8 +778,22 @@ namespace UnrealBuildTool
 		private Dictionary<string, UEBuildModule> Modules = new CaselessDictionary<UEBuildModule>();
 
 		[Serializable]
-		public class FlatModuleCsDataType
+		public class FlatModuleCsDataType : ISerializable
 		{
+			public FlatModuleCsDataType(SerializationInfo Info, StreamingContext Context)
+			{
+				BuildCsFilename    = Info.GetString("bf");
+				ModuleSourceFolder = Info.GetString("mf");
+				UHTHeaderNames     = (List<string>)Info.GetValue("hn", typeof(List<string>));
+			}
+
+			public void GetObjectData(SerializationInfo Info, StreamingContext Context)
+			{
+				Info.AddValue("bf", BuildCsFilename);
+				Info.AddValue("mf", ModuleSourceFolder);
+				Info.AddValue("hn", UHTHeaderNames);
+			}
+
 			public FlatModuleCsDataType(string InBuildCsFilename)
 			{
 				BuildCsFilename = InBuildCsFilename;
@@ -850,6 +851,60 @@ namespace UnrealBuildTool
 			return bCompileMonolithic;	// @todo ubtmake: We need to make sure this function and similar things aren't called in assembler mode
 		}
 
+		public UEBuildTarget(SerializationInfo Info, StreamingContext Context)
+		{
+			TargetType                   = (TargetRules.TargetType)Info.GetInt32("tt");
+			AppName                      = Info.GetString("an");
+			TargetName                   = Info.GetString("tn");
+			bUseSharedBuildEnvironment   = Info.GetBoolean("sb");
+			Platform                     = (UnrealTargetPlatform)Info.GetInt32("pl");
+			Configuration                = (UnrealTargetConfiguration)Info.GetInt32("co");
+			TargetInfo                   = (TargetInfo)Info.GetValue("ti", typeof(TargetInfo));
+			ProjectDirectory             = Info.GetString("pd");
+			ProjectIntermediateDirectory = Info.GetString("pi");
+			EngineIntermediateDirectory  = Info.GetString("ed");
+			OutputPaths                  = (List<string>)Info.GetValue("op", typeof(List<string>));
+			RemoteRoot                   = Info.GetString("rr");
+			bPrecompile                  = Info.GetBoolean("pc");
+			bUsePrecompiled              = Info.GetBoolean("up");
+			bEditorRecompile             = Info.GetBoolean("er");
+			bCompileMonolithic           = Info.GetBoolean("cm");
+			var FlatModuleCsDataKeys   = (string[])Info.GetValue("fk", typeof(string[]));
+			var FlatModuleCsDataValues = (FlatModuleCsDataType[])Info.GetValue("fv", typeof(FlatModuleCsDataType[]));
+			for (int Index = 0; Index != FlatModuleCsDataKeys.Length; ++Index)
+			{
+				FlatModuleCsData.Add(FlatModuleCsDataKeys[Index], FlatModuleCsDataValues[Index]);
+			}
+			Receipt                      = (TargetReceipt)Info.GetValue("re", typeof(TargetReceipt));
+			ReceiptFileName              = Info.GetString("rf");
+			TargetCsFilenameField        = Info.GetString("tc");
+		}
+
+		public void GetObjectData(SerializationInfo Info, StreamingContext Context)
+		{
+			Info.AddValue("tt", (int)TargetType);
+			Info.AddValue("an", AppName);
+			Info.AddValue("tn", TargetName);
+			Info.AddValue("sb", bUseSharedBuildEnvironment);
+			Info.AddValue("pl", (int)Platform);
+			Info.AddValue("co", (int)Configuration);
+			Info.AddValue("ti", TargetInfo);
+			Info.AddValue("pd", ProjectDirectory);
+			Info.AddValue("pi", ProjectIntermediateDirectory);
+			Info.AddValue("ed", EngineIntermediateDirectory);
+			Info.AddValue("op", OutputPaths);
+			Info.AddValue("rr", RemoteRoot);
+			Info.AddValue("pc", bPrecompile);
+			Info.AddValue("up", bUsePrecompiled);
+			Info.AddValue("er", bEditorRecompile);
+			Info.AddValue("cm", bCompileMonolithic);
+			Info.AddValue("fk", FlatModuleCsData.Keys.ToArray());
+			Info.AddValue("fv", FlatModuleCsData.Values.ToArray());
+			Info.AddValue("re", Receipt);
+			Info.AddValue("rf", ReceiptFileName);
+			Info.AddValue("tc", TargetCsFilenameField);
+		}
+
 		/// <summary>
 		/// Constructor.
 		/// </summary>
diff --git a/Engine/Source/Programs/UnrealBuildTool/System/ActionGraph.cs b/Engine/Source/Programs/UnrealBuildTool/System/ActionGraph.cs
--- a/Engine/Source/Programs/UnrealBuildTool/System/ActionGraph.cs
+++ b/Engine/Source/Programs/UnrealBuildTool/System/ActionGraph.cs
@@ -5,6 +5,7 @@ using System.Collections.Generic;
 using System.IO;
 using System.Diagnostics;
 using System.Linq;
+using System.Runtime.Serialization;
 
 namespace UnrealBuildTool
 {
@@ -26,7 +27,7 @@ namespace UnrealBuildTool
 
 	/** A build action. */
 	[Serializable]
-	public class Action
+	public class Action : ISerializable
 	{
 		///
 		/// Preparation and Assembly (serialized)
@@ -94,18 +95,15 @@ namespace UnrealBuildTool
 		///
 
 		/** Unique action identifier.  Used for displaying helpful info about detected cycles in the graph. */
-		[NonSerialized]
         public int UniqueId;
 
         /** Always-incremented unique id */
         private static int NextUniqueId = 0;
 
 		/** Total number of actions depending on this one. */
-		[NonSerialized]
 		public int NumTotalDependentActions = 0;
 		
 		/** Relative cost of producing items for this action. */
-		[NonSerialized]
 		public long RelativeCost = 0;
 
 
@@ -114,11 +112,9 @@ namespace UnrealBuildTool
 		///
 
 		/** Start time of action, optionally set by executor. */
-		[NonSerialized]
 		public DateTimeOffset StartTime = DateTimeOffset.MinValue;
 		
 		/** End time of action, optionally set by executor. */
-		[NonSerialized]
 		public DateTimeOffset EndTime = DateTimeOffset.MinValue;
 		
 
@@ -131,6 +127,44 @@ namespace UnrealBuildTool
             UniqueId = ++NextUniqueId;
 		}
 
+		public Action( SerializationInfo SerializationInfo, StreamingContext StreamingContext )
+		{
+			ActionType                     = (ActionType)SerializationInfo.GetByte("at");
+			WorkingDirectory               = SerializationInfo.GetString("wd");
+			bPrintDebugInfo                = SerializationInfo.GetBoolean("di");
+			CommandPath                    = SerializationInfo.GetString("cp");
+			CommandArguments               = SerializationInfo.GetString("ca");
+			CommandDescription             = SerializationInfo.GetString("cd");
+			StatusDescription              = SerializationInfo.GetString("sd");
+			bCanExecuteRemotely            = SerializationInfo.GetBoolean("ce");
+			bIsGCCCompiler                 = SerializationInfo.GetBoolean("ig");
+			bIsUsingPCH                    = SerializationInfo.GetBoolean("iu");
+			bShouldDeleteProducedItems     = SerializationInfo.GetBoolean("dp");
+			bShouldOutputStatusDescription = SerializationInfo.GetBoolean("os");
+			bProducesImportLibrary         = SerializationInfo.GetBoolean("il");
+			PrerequisiteItems              = (List<FileItem>)SerializationInfo.GetValue("pr", typeof(List<FileItem>));
+			ProducedItems                  = (List<FileItem>)SerializationInfo.GetValue("pd", typeof(List<FileItem>));
+		}
+
+		/** ISerializable: Called when serialized to report additional properties that should be saved */
+		public void GetObjectData( SerializationInfo SerializationInfo, StreamingContext StreamingContext )
+		{
+			SerializationInfo.AddValue("at", (byte)ActionType);
+			SerializationInfo.AddValue("wd", WorkingDirectory);
+			SerializationInfo.AddValue("di", bPrintDebugInfo);
+			SerializationInfo.AddValue("cp", CommandPath);
+			SerializationInfo.AddValue("ca", CommandArguments);
+			SerializationInfo.AddValue("cd", CommandDescription);
+			SerializationInfo.AddValue("sd", StatusDescription);
+			SerializationInfo.AddValue("ce", bCanExecuteRemotely);
+			SerializationInfo.AddValue("ig", bIsGCCCompiler);
+			SerializationInfo.AddValue("iu", bIsUsingPCH);
+			SerializationInfo.AddValue("dp", bShouldDeleteProducedItems);
+			SerializationInfo.AddValue("os", bShouldOutputStatusDescription);
+			SerializationInfo.AddValue("il", bProducesImportLibrary);
+			SerializationInfo.AddValue("pr", PrerequisiteItems);
+			SerializationInfo.AddValue("pd", ProducedItems);
+		}
 
 		/**
 		 * Compares two actions based on total number of dependent items, descending.
diff --git a/Engine/Source/Programs/UnrealBuildTool/System/CPPHeaders.cs b/Engine/Source/Programs/UnrealBuildTool/System/CPPHeaders.cs
--- a/Engine/Source/Programs/UnrealBuildTool/System/CPPHeaders.cs
+++ b/Engine/Source/Programs/UnrealBuildTool/System/CPPHeaders.cs
@@ -6,13 +6,15 @@ using System.Diagnostics;
 using System.Text;
 using System.Text.RegularExpressions;
 using System.IO;
+using System.Runtime.Serialization;
+using System.Linq;
 using Tools.DotNETCommon.FileContentsCacheType;
 
 namespace UnrealBuildTool
 {
 	/** For C++ source file items, this structure is used to cache data that will be used for include dependency scanning */
 	[Serializable]
-	public class CPPIncludeInfo
+	public class CPPIncludeInfo : ISerializable
 	{
 		/** Ordered list of include paths for the module  */
 		public HashSet<string> IncludePaths = new HashSet<string>();
@@ -24,9 +26,24 @@ namespace UnrealBuildTool
 		public HashSet<string> SystemIncludePaths = new HashSet<string>();
 
 		/** Contains a mapping from filename to the full path of the header in this environment.  This is used to optimized include path lookups at runtime for any given single module. */
-		[NonSerialized]
 		public Dictionary<string, FileItem> IncludeFileSearchDictionary = new Dictionary<string, FileItem>();
 
+		public CPPIncludeInfo()
+		{
+		}
+
+		public CPPIncludeInfo(SerializationInfo Info, StreamingContext Context)
+		{
+			IncludePaths       = new HashSet<string>((string[])Info.GetValue("ip", typeof(string[])));
+			SystemIncludePaths = new HashSet<string>((string[])Info.GetValue("sp", typeof(string[])));
+		}
+
+		public void GetObjectData(SerializationInfo Info, StreamingContext Context)
+		{
+			Info.AddValue("ip", IncludePaths.ToArray());
+			Info.AddValue("sp", SystemIncludePaths.ToArray());
+		}
+
 
 		/// <summary>
 		/// Given a C++ source file, returns a list of include paths we should search to resolve #includes for this path
diff --git a/Engine/Source/Programs/UnrealBuildTool/System/ExternalExecution.cs b/Engine/Source/Programs/UnrealBuildTool/System/ExternalExecution.cs
--- a/Engine/Source/Programs/UnrealBuildTool/System/ExternalExecution.cs
+++ b/Engine/Source/Programs/UnrealBuildTool/System/ExternalExecution.cs
@@ -9,6 +9,7 @@ using System.IO;
 using System.Text.RegularExpressions;
 using System.Reflection;
 using System.Runtime.InteropServices;
+using System.Runtime.Serialization;
 
 namespace UnrealBuildTool
 {
@@ -45,7 +46,7 @@ namespace UnrealBuildTool
 
 	/** Information about a module that needs to be passed to UnrealHeaderTool for code generation */
 	[Serializable]
-	public class UHTModuleInfo
+	public class UHTModuleInfo : ISerializable
 	{
 		/** Module name */
 		public string ModuleName;
@@ -74,6 +75,36 @@ namespace UnrealBuildTool
 		/** Version of code generated by UHT */
 		public EGeneratedCodeVersion GeneratedCodeVersion;
 
+		public UHTModuleInfo()
+		{
+		}
+
+		public UHTModuleInfo(SerializationInfo Info, StreamingContext Context)
+		{
+			ModuleName                  = Info.GetString("mn");
+			ModuleDirectory             = Info.GetString("md");
+			ModuleType                  = Info.GetString("mt");
+			PublicUObjectClassesHeaders = (List<FileItem>)Info.GetValue("cl", typeof(List<FileItem>));
+			PublicUObjectHeaders        = (List<FileItem>)Info.GetValue("pu", typeof(List<FileItem>));
+			PrivateUObjectHeaders       = (List<FileItem>)Info.GetValue("pr", typeof(List<FileItem>));
+			PCH                         = Info.GetString("pc");
+			GeneratedCPPFilenameBase    = Info.GetString("ge");
+			GeneratedCodeVersion        = (EGeneratedCodeVersion)Info.GetInt32("gv");
+		}
+
+		public void GetObjectData(SerializationInfo Info, StreamingContext Context)
+		{
+			Info.AddValue("mn", ModuleName);
+			Info.AddValue("md", ModuleDirectory);
+			Info.AddValue("mt", ModuleType);
+			Info.AddValue("cl", PublicUObjectClassesHeaders);
+			Info.AddValue("pu", PublicUObjectHeaders);
+			Info.AddValue("pr", PrivateUObjectHeaders);
+			Info.AddValue("pc", PCH);
+			Info.AddValue("ge", GeneratedCPPFilenameBase);
+			Info.AddValue("gv", (int)GeneratedCodeVersion);
+		}
+
 		public override string ToString()
 		{
 			return ModuleName;
diff --git a/Engine/Source/Programs/UnrealBuildTool/System/RulesCompiler.cs b/Engine/Source/Programs/UnrealBuildTool/System/RulesCompiler.cs
--- a/Engine/Source/Programs/UnrealBuildTool/System/RulesCompiler.cs
+++ b/Engine/Source/Programs/UnrealBuildTool/System/RulesCompiler.cs
@@ -6,6 +6,7 @@ using System.Collections.Generic;
 using System.Linq;
 using System.Reflection;
 using System.Diagnostics;
+using System.Runtime.Serialization;
 
 namespace UnrealBuildTool
 {
@@ -13,7 +14,7 @@ namespace UnrealBuildTool
 	/// Information about a target, passed along when creating a module descriptor
 	/// </summary>
 	[Serializable]
-	public class TargetInfo
+	public class TargetInfo : ISerializable
 	{
 		/// Target platform
 		public readonly UnrealTargetPlatform Platform;
@@ -30,6 +31,38 @@ namespace UnrealBuildTool
 		/// Whether the target is monolithic (if known)
 		public readonly bool? bIsMonolithic;
 
+		public TargetInfo(SerializationInfo Info, StreamingContext Context)
+		{
+			Platform      = (UnrealTargetPlatform)Info.GetInt32("pl");
+			Architecture  = Info.GetString("ar");
+			Configuration = (UnrealTargetConfiguration)Info.GetInt32("co");
+			if (Info.GetBoolean("t?"))
+			{
+				Type = (TargetRules.TargetType)Info.GetInt32("tt");
+			}
+			if (Info.GetBoolean("m?"))
+			{
+				bIsMonolithic = Info.GetBoolean("mo");
+			}
+		}
+
+		public void GetObjectData(SerializationInfo Info, StreamingContext Context)
+		{
+			Info.AddValue("pl", (int)Platform);
+			Info.AddValue("ar", Architecture);
+			Info.AddValue("co", (int)Configuration);
+			Info.AddValue("t?", Type.HasValue);
+			if (Type.HasValue)
+			{
+				Info.AddValue("tt", (int)Type.Value);
+			}
+			Info.AddValue("m?", bIsMonolithic.HasValue);
+			if (bIsMonolithic.HasValue)
+			{
+				Info.AddValue("mo", bIsMonolithic);
+			}
+		}
+
 		/// <summary>
 		/// Constructs a TargetInfo
 		/// </summary>
diff --git a/Engine/Source/Programs/UnrealBuildTool/System/TargetReceipt.cs b/Engine/Source/Programs/UnrealBuildTool/System/TargetReceipt.cs
--- a/Engine/Source/Programs/UnrealBuildTool/System/TargetReceipt.cs
+++ b/Engine/Source/Programs/UnrealBuildTool/System/TargetReceipt.cs
@@ -7,11 +7,12 @@ using System.Diagnostics;
 using System.Linq;
 using System.Xml.Serialization;
 using System.ComponentModel;
+using System.Runtime.Serialization;
 
 namespace UnrealBuildTool
 {
 	[Serializable]
-	public class BuildProperty
+	public class BuildProperty : ISerializable
 	{
 		[XmlAttribute]
 		public string Name;
@@ -23,6 +24,18 @@ namespace UnrealBuildTool
 		{
 		}
 
+		public BuildProperty(SerializationInfo Info, StreamingContext Context)
+		{
+			Name  = Info.GetString("na");
+			Value = Info.GetString("va");
+		}
+
+		public void GetObjectData(SerializationInfo Info, StreamingContext Context)
+		{
+			Info.AddValue("na", Name);
+			Info.AddValue("va", Value);
+		}
+
 		public BuildProperty(string InName, string InValue)
 		{
 			Name = InName;
@@ -41,7 +54,7 @@ namespace UnrealBuildTool
 	}
 
 	[Serializable]
-	public class BuildProduct
+	public class BuildProduct : ISerializable
 	{
 		[XmlAttribute]
 		public string Path;
@@ -56,6 +69,20 @@ namespace UnrealBuildTool
 		{
 		}
 
+		public BuildProduct(SerializationInfo Info, StreamingContext Context)
+		{
+			Path          = Info.GetString("pa");
+			Type          = (BuildProductType)Info.GetInt32("ty");
+			IsPrecompiled = Info.GetBoolean("ip");
+		}
+
+		public void GetObjectData(SerializationInfo Info, StreamingContext Context)
+		{
+			Info.AddValue("pa", Path);
+			Info.AddValue("ty", (int)Type);
+			Info.AddValue("ip", IsPrecompiled);
+		}
+
 		public BuildProduct(string InPath, BuildProductType InType)
 		{
 			Path = InPath;
@@ -76,7 +103,7 @@ namespace UnrealBuildTool
 	}
 
 	[Serializable]
-	public class RuntimeDependency
+	public class RuntimeDependency : ISerializable
 	{
 		[XmlAttribute]
 		public string Path;
@@ -88,6 +115,18 @@ namespace UnrealBuildTool
 		{
 		}
 
+		public RuntimeDependency(SerializationInfo Info, StreamingContext Context)
+		{
+			Path      = Info.GetString("pa");
+			StagePath = Info.GetString("sp");
+		}
+
+		public void GetObjectData(SerializationInfo Info, StreamingContext Context)
+		{
+			Info.AddValue("pa", Path);
+			Info.AddValue("sp", StagePath);
+		}
+
 		public RuntimeDependency(string InPath)
 		{
 			Path = InPath;
@@ -115,7 +154,7 @@ namespace UnrealBuildTool
 	/// Stores a record of a built target, with all metadata that other tools may need to know about the build.
 	/// </summary>
 	[Serializable]
-	public class TargetReceipt
+	public class TargetReceipt : ISerializable
 	{
 		[XmlArrayItem("Property")]
 		public List<BuildProperty> Properties = new List<BuildProperty>();
@@ -136,6 +175,22 @@ namespace UnrealBuildTool
 		{
 		}
 
+		public TargetReceipt(SerializationInfo Info, StreamingContext Context)
+		{
+			Properties                  = (List<BuildProperty>)Info.GetValue("pr", typeof(List<BuildProperty>));
+			BuildProducts               = (List<BuildProduct>)Info.GetValue("bp", typeof(List<BuildProduct>));
+			RuntimeDependencies         = (List<RuntimeDependency>)Info.GetValue("rd", typeof(List<RuntimeDependency>));
+			bRequireDependenciesToExist = Info.GetBoolean("rq");
+		}
+
+		public void GetObjectData(SerializationInfo Info, StreamingContext Context)
+		{
+			Info.AddValue("pr", Properties);
+			Info.AddValue("bp", BuildProducts);
+			Info.AddValue("rd", RuntimeDependencies);
+			Info.AddValue("rq", bRequireDependenciesToExist);
+		}
+
 		/// <summary>
 		/// Copy constructor
 		/// </summary>
diff --git a/Engine/Source/Programs/UnrealBuildTool/System/UnrealBuildTool.cs b/Engine/Source/Programs/UnrealBuildTool/System/UnrealBuildTool.cs
--- a/Engine/Source/Programs/UnrealBuildTool/System/UnrealBuildTool.cs
+++ b/Engine/Source/Programs/UnrealBuildTool/System/UnrealBuildTool.cs
@@ -9,6 +9,7 @@ using System.Threading;
 using System.Reflection;
 using System.Linq;
 using System.Runtime.Serialization.Formatters.Binary;
+using System.Runtime.Serialization;
 using Tools.DotNETCommon;
 
 namespace UnrealBuildTool
@@ -988,7 +989,7 @@ namespace UnrealBuildTool
                     UEBuildTarget.ParsePlatformAndConfiguration(Arguments, out CheckPlatform, out CheckConfiguration, false);
 
 					// @todo ubtmake: remove this when building with RPCUtility works
-					if (CheckPlatform == UnrealTargetPlatform.Mac || CheckPlatform == UnrealTargetPlatform.IOS)
+					if (BuildHostPlatform.Current.Platform != CheckPlatform && (CheckPlatform == UnrealTargetPlatform.Mac || CheckPlatform == UnrealTargetPlatform.IOS))
 					{
 						BuildConfiguration.bUseUBTMakefiles = false;
 					}
@@ -2235,13 +2236,17 @@ namespace UnrealBuildTool
         /// A special Makefile that UBT is able to create in "-gather" mode, then load in "-assemble" mode to accelerate iterative compiling and linking
         /// </summary>
         [Serializable]
-        class UBTMakefile
+        class UBTMakefile : ISerializable
         {
+			public const int CurrentVersion = 1;
+
+			public int Version;
+
             /** Every action in the action graph */
             public List<Action> AllActions;
 
             /** List of the actions that need to be run in order to build the targets' final output items */
-            public Action[] PrerequisiteActions;			
+            public Action[] PrerequisiteActions;
 
             /** Environment variables that we'll need in order to invoke the platform's compiler and linker */
 			// @todo ubtmake: Really we want to allow a different set of environment variables for every Action.  This would allow for targets on multiple platforms to be built in a single assembling phase.  We'd only have unique variables for each platform that has actions, so we'd want to make sure we only store the minimum set.
@@ -2253,6 +2258,37 @@ namespace UnrealBuildTool
 			/** List of targets being built */
 			public List<UEBuildTarget> Targets;
 
+			public UBTMakefile()
+			{
+				Version = CurrentVersion;
+			}
+
+			public UBTMakefile(SerializationInfo Info, StreamingContext Context)
+			{
+				Version = Info.GetInt32("ve");
+				if (Version != CurrentVersion)
+				{
+					throw new Exception(string.Format("Makefile version does not match - found {0}, expected: {1}", Version, CurrentVersion));
+				}
+
+				AllActions                 = (List<Action>)Info.GetValue("ac", typeof(List<Action>));
+				PrerequisiteActions        = (Action[])Info.GetValue("pa", typeof(Action[]));
+				EnvironmentVariables       = ((string[])Info.GetValue("e1", typeof(string[]))).Zip((string[])Info.GetValue("e2", typeof(string[])), (i1, i2) => new Tuple<string, string>(i1, i2)).ToList();
+				TargetNameToUObjectModules = (Dictionary<string,List<UHTModuleInfo>>)Info.GetValue("nu", typeof(Dictionary<string,List<UHTModuleInfo>>));
+				Targets                    = (List<UEBuildTarget>)Info.GetValue("ta", typeof(List<UEBuildTarget>));
+			}
+
+			public void GetObjectData(SerializationInfo Info, StreamingContext Context)
+			{
+				Info.AddValue("ve", Version);
+				Info.AddValue("ac", AllActions);
+				Info.AddValue("pa", PrerequisiteActions);
+				Info.AddValue("e1", EnvironmentVariables.Select(x => x.Item1).ToArray());
+				Info.AddValue("e2", EnvironmentVariables.Select(x => x.Item2).ToArray());
+				Info.AddValue("nu", TargetNameToUObjectModules);
+				Info.AddValue("ta", Targets);
+			}
+
 
             /** @return	 True if this makefile's contents look valid.  Called after loading the file to make sure it is legit. */
             public bool IsValidMakefile()
