UE-19348 - fixed crash on start up for Strategy Game in iOS 9
UE-19367 - fxied crash on start up for Sun Temple in iOS 9
#ios

diff --git a/Engine/Source/Programs/UnrealBuildTool/IOS/UEDeployIOS.cs b/Engine/Source/Programs/UnrealBuildTool/IOS/UEDeployIOS.cs
--- a/Engine/Source/Programs/UnrealBuildTool/IOS/UEDeployIOS.cs
+++ b/Engine/Source/Programs/UnrealBuildTool/IOS/UEDeployIOS.cs
@@ -302,6 +302,8 @@ namespace UnrealBuildTool.IOS
 			Text.AppendLine("\t<true/>");
 			Text.AppendLine("\t<key>UIStatusBarHidden</key>");
 			Text.AppendLine("\t<true/>");
+			Text.AppendLine("\t<key>UIViewControllerBasedStatusBarAppearance</key>");
+			Text.AppendLine("\t<false/>");
 			Text.AppendLine("\t<key>UISupportedInterfaceOrientations</key>");
 			Text.AppendLine("\t<array>");
 			foreach (string Line in SupportedOrientations.Split("\r\n".ToCharArray()))
