UE-17831 - fixed issue with non-existent dSYM files being added to the receipt
Also fixed name issue with the ini file
#ios

diff --git a/Engine/Source/Developer/iOS/IOSTargetPlatform/Private/IOSTargetPlatform.h b/Engine/Source/Developer/iOS/IOSTargetPlatform/Private/IOSTargetPlatform.h
--- a/Engine/Source/Developer/iOS/IOSTargetPlatform/Private/IOSTargetPlatform.h
+++ b/Engine/Source/Developer/iOS/IOSTargetPlatform/Private/IOSTargetPlatform.h
@@ -110,7 +110,7 @@ public:
 		OutSection = TEXT("/Script/IOSRuntimeSettings.IOSRuntimeSettings");
 		InBoolKeys.Add(TEXT("bDevForArmV7")); InBoolKeys.Add(TEXT("bDevForArm64")); InBoolKeys.Add(TEXT("bDevForArmV7S"));
 		InBoolKeys.Add(TEXT("bShipForArmV7")); InBoolKeys.Add(TEXT("bShipForArm64")); InBoolKeys.Add(TEXT("bShipForArmV7S"));
-		InBoolKeys.Add(TEXT("bGenerateSYMFile"));
+		InBoolKeys.Add(TEXT("bGeneratedSYMFile"));
 		InStringKeys.Add(TEXT("MinimumiOSVersion"));
 	}
 
diff --git a/Engine/Source/Programs/UnrealBuildTool/Configuration/UEBuildBinary.cs b/Engine/Source/Programs/UnrealBuildTool/Configuration/UEBuildBinary.cs
--- a/Engine/Source/Programs/UnrealBuildTool/Configuration/UEBuildBinary.cs
+++ b/Engine/Source/Programs/UnrealBuildTool/Configuration/UEBuildBinary.cs
@@ -312,7 +312,7 @@ namespace UnrealBuildTool
 			string DebugExtension = UEBuildPlatform.GetBuildPlatform(Target.Platform).GetDebugInfoExtension(Config.Type);
 			foreach (string OutputFilePath in Config.OutputFilePaths)
 			{
-				AddBuildProductAndDebugFile(OutputFilePath, Type, DebugExtension, Receipt);
+				AddBuildProductAndDebugFile(OutputFilePath, Type, DebugExtension, Receipt, ToolChain);
 			}
 
 			// Add the console app, if there is one
@@ -320,7 +320,7 @@ namespace UnrealBuildTool
 			{
 				foreach (string OutputFilePath in Config.OutputFilePaths)
 				{
-					AddBuildProductAndDebugFile(GetAdditionalConsoleAppPath(OutputFilePath), Type, DebugExtension, Receipt);
+					AddBuildProductAndDebugFile(GetAdditionalConsoleAppPath(OutputFilePath), Type, DebugExtension, Receipt, ToolChain);
 				}
 			}
 
@@ -335,11 +335,11 @@ namespace UnrealBuildTool
 		/// <param name="OutputFile">Build product to add</param>
 		/// <param name="DebugExtension">Extension for the matching debug file (may be null).</param>
 		/// <param name="Receipt">Receipt to add to</param>
-		static void AddBuildProductAndDebugFile(string OutputFile, BuildProductType OutputType, string DebugExtension, TargetReceipt Receipt)
+		static void AddBuildProductAndDebugFile(string OutputFile, BuildProductType OutputType, string DebugExtension, TargetReceipt Receipt, IUEToolChain ToolChain)
 		{
 			Receipt.AddBuildProduct(OutputFile, OutputType);
 
-			if(!String.IsNullOrEmpty(DebugExtension))
+			if(!String.IsNullOrEmpty(DebugExtension) && ToolChain.ShouldAddDebugFileToReceipt(OutputFile, OutputType))
 			{
 				Receipt.AddBuildProduct(Path.ChangeExtension(OutputFile, DebugExtension), BuildProductType.SymbolFile);
 			}
diff --git a/Engine/Source/Programs/UnrealBuildTool/IOS/IOSToolChain.cs b/Engine/Source/Programs/UnrealBuildTool/IOS/IOSToolChain.cs
--- a/Engine/Source/Programs/UnrealBuildTool/IOS/IOSToolChain.cs
+++ b/Engine/Source/Programs/UnrealBuildTool/IOS/IOSToolChain.cs
@@ -103,6 +103,16 @@ namespace UnrealBuildTool
 			}
 		}
 
+		/// <summary>
+		/// Adds a build product and its associated debug file to a receipt.
+		/// </summary>
+		/// <param name="OutputFile">Build product to add</param>
+		/// <param name="DebugExtension">Extension for the matching debug file (may be null).</param>
+		public override bool ShouldAddDebugFileToReceipt(string OutputFile, BuildProductType OutputType)
+		{
+			return OutputType == BuildProductType.Executable;
+		}
+
 		static bool bHasPrinted = false;
 		static string GetArchitectureArgument(CPPTargetConfiguration Configuration, string UBTArchitecture)
 		{
@@ -221,7 +231,7 @@ namespace UnrealBuildTool
 			{
 				if (Framework.OwningModule != null && Framework.FrameworkZipPath != null && Framework.FrameworkZipPath != "")
 				{
-					Result += " -F \"" + GetRemoteIntermediateFrameworkZipPath(Framework) + "\"";
+					Result += " -F\"" + GetRemoteIntermediateFrameworkZipPath(Framework) + "\"";
 				}
 			}
 
@@ -373,7 +383,7 @@ namespace UnrealBuildTool
 				if ( Framework.OwningModule != null && Framework.FrameworkZipPath != null && Framework.FrameworkZipPath != "" )
 				{
 					// If this framework has a zip specified, we'll need to setup the path as well
-					Result += " -F \"" + GetRemoteIntermediateFrameworkZipPath( Framework ) + "\"";
+					Result += " -F\"" + GetRemoteIntermediateFrameworkZipPath( Framework ) + "\"";
 				}
 
 				Result += " -framework " + Framework.FrameworkName;
diff --git a/Engine/Source/Programs/UnrealBuildTool/IOS/UEBuildIOS.cs b/Engine/Source/Programs/UnrealBuildTool/IOS/UEBuildIOS.cs
--- a/Engine/Source/Programs/UnrealBuildTool/IOS/UEBuildIOS.cs
+++ b/Engine/Source/Programs/UnrealBuildTool/IOS/UEBuildIOS.cs
@@ -237,7 +237,7 @@ namespace UnrealBuildTool
 				}
 
 				// determine if we need to generate the dsym
-				Ini.GetBool("/Script/IOSRuntimeSettings.IOSRuntimeSettings", "bGenerateSYMFile", out BuildConfiguration.bGeneratedSYMFile);
+				Ini.GetBool("/Script/IOSRuntimeSettings.IOSRuntimeSettings", "bGeneratedSYMFile", out BuildConfiguration.bGeneratedSYMFile);
 
 				bInitializedProject = true;
 			}
@@ -252,7 +252,7 @@ namespace UnrealBuildTool
 		{
 			string[] BoolKeys = new string[] {
 				"bDevForArmV7", "bDevForArm64", "bDevForArmV7S", "bShipForArmV7", 
-				"bShipForArm64", "bShipForArmV7S", "bGenerateSYMFile",
+				"bShipForArm64", "bShipForArmV7S", "bGeneratedSYMFile",
 			};
 			string[] StringKeys = new string[] {
 				"MinimumiOSVersion", 
diff --git a/Engine/Source/Programs/UnrealBuildTool/ToolChain/UEToolChain.cs b/Engine/Source/Programs/UnrealBuildTool/ToolChain/UEToolChain.cs
--- a/Engine/Source/Programs/UnrealBuildTool/ToolChain/UEToolChain.cs
+++ b/Engine/Source/Programs/UnrealBuildTool/ToolChain/UEToolChain.cs
@@ -40,6 +40,8 @@ namespace UnrealBuildTool
 
 		void AddFilesToReceipt(TargetReceipt Receipt, UEBuildBinary Binary);
 
+		bool ShouldAddDebugFileToReceipt(string OutputFile, BuildProductType OutputType);
+
 		void SetupBundleDependencies(List<UEBuildBinary> Binaries, string GameName);
 
 		void FixBundleBinariesPaths(UEBuildTarget Target, List<UEBuildBinary> Binaries);
@@ -176,6 +178,16 @@ namespace UnrealBuildTool
 		{
 		}
 
+		/// <summary>
+		/// Adds a build product and its associated debug file to a receipt.
+		/// </summary>
+		/// <param name="OutputFile">Build product to add</param>
+		/// <param name="DebugExtension">Extension for the matching debug file (may be null).</param>
+		public virtual bool ShouldAddDebugFileToReceipt(string OutputFile, BuildProductType OutputType)
+		{
+			return true;
+		}
+
 		protected void AddPrerequisiteSourceFile( UEBuildTarget Target, IUEBuildPlatform BuildPlatform, CPPEnvironment CompileEnvironment, FileItem SourceFile, List<FileItem> PrerequisiteItems )
 		{
 			PrerequisiteItems.Add( SourceFile );
