Enabling XGE for Linux in hopes that the crashes are gone.

diff --git a/Engine/Source/Programs/UnrealBuildTool/Linux/UEBuildLinux.cs b/Engine/Source/Programs/UnrealBuildTool/Linux/UEBuildLinux.cs
--- a/Engine/Source/Programs/UnrealBuildTool/Linux/UEBuildLinux.cs
+++ b/Engine/Source/Programs/UnrealBuildTool/Linux/UEBuildLinux.cs
@@ -113,12 +113,6 @@ namespace UnrealBuildTool
             return SDKStatus.Invalid;
         }
 
-		public override bool CanUseXGE()
-		{
-			// disabled until XGE crash is fixed - it is still happening as of 2014-09-30
-			return false;
-		}
-
 		/**
 		 *	Register the platform with the UEBuildPlatform class
 		 */
