Revert the previous change and disable XGE again.

- Modular builds (editor, UT server) seem to be broken by this since FixDeps step probably is not listing its prerequisites correctly.

diff --git a/Engine/Source/Programs/UnrealBuildTool/Linux/UEBuildLinux.cs b/Engine/Source/Programs/UnrealBuildTool/Linux/UEBuildLinux.cs
--- a/Engine/Source/Programs/UnrealBuildTool/Linux/UEBuildLinux.cs
+++ b/Engine/Source/Programs/UnrealBuildTool/Linux/UEBuildLinux.cs
@@ -113,6 +113,12 @@ namespace UnrealBuildTool
             return SDKStatus.Invalid;
         }
 
+		public override bool CanUseXGE()
+		{
+			// [RCL] 2015-08-04 FIXME: modular (cross-)builds (e.g. editor, UT server) fail with XGE as FixDeps step apparently depends on artifacts (object files) which aren't listed among its prerequisites.
+			return false;
+		}
+
 		/**
 		 *	Register the platform with the UEBuildPlatform class
 		 */
