Use gnu-libstdc++ 4.8 for NDK level 19, and allow 4.9 for higher NDK levels
Only use ld.gold for Clang 3.6 if unity build
#ue4
#android

diff --git a/Engine/Source/Programs/UnrealBuildTool/Android/AndroidToolChain.cs b/Engine/Source/Programs/UnrealBuildTool/Android/AndroidToolChain.cs
--- a/Engine/Source/Programs/UnrealBuildTool/Android/AndroidToolChain.cs
+++ b/Engine/Source/Programs/UnrealBuildTool/Android/AndroidToolChain.cs
@@ -143,6 +143,22 @@ namespace UnrealBuildTool
 			return GPUArchitectures;
 		}
 
+		static public int GetNdkApiLevelInt(int MinNdk = 19)
+		{
+			string NDKVersion = AndroidToolChain.GetNdkApiLevel();
+			int NDKVersionInt = MinNdk;
+			if (NDKVersion.Contains("-"))
+			{
+				int Version;
+				if (int.TryParse(NDKVersion.Substring(NDKVersion.LastIndexOf('-') + 1), out Version))
+				{
+					if (Version > NDKVersionInt)
+						NDKVersionInt = Version;
+				}
+			}
+			return NDKVersionInt;
+		}
+
 		static public string GetNdkApiLevel()
 		{
 			// ask the .ini system for what version to use
@@ -550,7 +566,7 @@ namespace UnrealBuildTool
 				Result += " -Wl,--fix-cortex-a8";		// required to route around a CPU bug in some Cortex-A8 implementations
 			}
 
-			if (ClangVersionFloat >= 3.6f)
+			if (BuildConfiguration.bUseUnityBuild && ClangVersionFloat >= 3.6f)
 			{
 				Result += " -fuse-ld=gold";				// ld.gold is available in r10e (clang 3.6)
 			}
diff --git a/Engine/Source/Programs/UnrealBuildTool/Android/UEBuildAndroid.cs b/Engine/Source/Programs/UnrealBuildTool/Android/UEBuildAndroid.cs
--- a/Engine/Source/Programs/UnrealBuildTool/Android/UEBuildAndroid.cs
+++ b/Engine/Source/Programs/UnrealBuildTool/Android/UEBuildAndroid.cs
@@ -362,15 +362,19 @@ namespace UnrealBuildTool
 			NDKPath = NDKPath.Replace("\"", "");
 
 			string GccVersion = "4.6";
-			if (Directory.Exists(Path.Combine(NDKPath, @"sources/cxx-stl/gnu-libstdc++/4.9")))
+			int NDKVersionInt = AndroidToolChain.GetNdkApiLevelInt();
+			if (Directory.Exists(Path.Combine(NDKPath, @"sources/cxx-stl/gnu-libstdc++/4.8")))
 			{
-				GccVersion = "4.9";
+				GccVersion = "4.8";
 			}
-			else if (Directory.Exists(Path.Combine(NDKPath, @"sources/cxx-stl/gnu-libstdc++/4.8")))
+			// only use 4.9 if NDK version > 19
+			if (NDKVersionInt > 19 && Directory.Exists(Path.Combine(NDKPath, @"sources/cxx-stl/gnu-libstdc++/4.9")))
 			{
-				GccVersion = "4.8";
+				GccVersion = "4.9";
 			}
 
+			Log.TraceInformation("NDK version: {0}, GccVersion: {1}", NDKVersionInt.ToString(), GccVersion);
+
 			InBuildTarget.GlobalCompileEnvironment.Config.Definitions.Add("PLATFORM_DESKTOP=0");
 			InBuildTarget.GlobalCompileEnvironment.Config.Definitions.Add("PLATFORM_CAN_SUPPORT_EDITORONLY_DATA=0");
 
diff --git a/Engine/Source/Programs/UnrealBuildTool/Android/UEDeployAndroid.cs b/Engine/Source/Programs/UnrealBuildTool/Android/UEDeployAndroid.cs
--- a/Engine/Source/Programs/UnrealBuildTool/Android/UEDeployAndroid.cs
+++ b/Engine/Source/Programs/UnrealBuildTool/Android/UEDeployAndroid.cs
@@ -530,13 +530,14 @@ namespace UnrealBuildTool.Android
 			string Arch = GetNDKArch(UE4Arch);
 
 			string GccVersion = "4.6";
-			if (Directory.Exists(Environment.ExpandEnvironmentVariables("%NDKROOT%/sources/cxx-stl/gnu-libstdc++/4.9")))
+			if (Directory.Exists(Environment.ExpandEnvironmentVariables("%NDKROOT%/sources/cxx-stl/gnu-libstdc++/4.8")))
 			{
-				GccVersion = "4.9";
+				GccVersion = "4.8";
 			}
-			else if (Directory.Exists(Environment.ExpandEnvironmentVariables("%NDKROOT%/sources/cxx-stl/gnu-libstdc++/4.8")))
+			// only use 4.9 if NDK version > 19
+			if (AndroidToolChain.GetNdkApiLevelInt() > 19 && Directory.Exists(Environment.ExpandEnvironmentVariables("%NDKROOT%/sources/cxx-stl/gnu-libstdc++/4.9")))
 			{
-				GccVersion = "4.8";
+				GccVersion = "4.9";
 			}
 
 			// copy it in!
