Keep the last export of each path from .bash_profile instead of the first one if not set in ini
#ue4
#android

diff --git a/Engine/Source/Programs/UnrealBuildTool/Android/UEBuildAndroid.cs b/Engine/Source/Programs/UnrealBuildTool/Android/UEBuildAndroid.cs
--- a/Engine/Source/Programs/UnrealBuildTool/Android/UEBuildAndroid.cs
+++ b/Engine/Source/Programs/UnrealBuildTool/Android/UEBuildAndroid.cs
@@ -92,22 +92,24 @@ namespace UnrealBuildTool
                     if (File.Exists(BashProfilePath))
                     {
                         string[] BashProfileContents = File.ReadAllLines(BashProfilePath);
-                        foreach (string Line in BashProfileContents)
-                        {
-                            foreach (var kvp in EnvVarNames)
-                            {
-                                if (AndroidEnv.ContainsKey(kvp.Key))
-                                {
-                                    continue;
-                                }
-
-                                if (Line.StartsWith("export " + kvp.Key + "="))
-                                {
-                                    string PathVar = Line.Split('=')[1].Replace("\"", "");
-                                    AndroidEnv.Add(kvp.Key, PathVar);
-                                }
-                            }
-                        }
+
+						// Walk backwards so we keep the last export setting instead of the first
+						for (int LineIndex = BashProfileContents.Length - 1; LineIndex >= 0; --LineIndex)
+						{
+							foreach (var kvp in EnvVarNames)
+							{
+								if (AndroidEnv.ContainsKey(kvp.Key))
+								{
+									continue;
+								}
+
+								if (BashProfileContents[LineIndex].StartsWith("export " + kvp.Key + "="))
+								{
+									string PathVar = BashProfileContents[LineIndex].Split('=')[1].Replace("\"", "");
+									AndroidEnv.Add(kvp.Key, PathVar);
+								}
+							}
+						}
                     }
                 }
 
