Fix to module name comparison when checking files to rebuild.

diff --git a/Engine/Source/Programs/UnrealBuildTool/Configuration/UEBuildTarget.cs b/Engine/Source/Programs/UnrealBuildTool/Configuration/UEBuildTarget.cs
--- a/Engine/Source/Programs/UnrealBuildTool/Configuration/UEBuildTarget.cs
+++ b/Engine/Source/Programs/UnrealBuildTool/Configuration/UEBuildTarget.cs
@@ -3392,7 +3392,7 @@ namespace UnrealBuildTool
 				// Figure out whether we need to build this module
 				// We don't care about actual source files when generating projects, as these are discovered separately
 				bool bDiscoverFiles = !ProjectFileGenerator.bGenerateProjectFiles;
-				bool bBuildFiles    = bDiscoverFiles && (OnlyModules.Count == 0 || OnlyModules.Any(x => x.OnlyModuleName == ModuleName));
+				bool bBuildFiles    = bDiscoverFiles && (OnlyModules.Count == 0 || OnlyModules.Any(x => string.Equals(x.OnlyModuleName, ModuleName, StringComparison.InvariantCultureIgnoreCase)));
 
 				IntelliSenseGatherer IntelliSenseGatherer = null;
 				List<FileItem> FoundSourceFiles = new List<FileItem>();
