Fix handling of spaces in emcc path.
Fix Optimizer path on Mac
fix jira-19532 part 2

diff --git a/Engine/Source/Programs/UnrealBuildTool/HTML5/HTML5SDKInfo.cs b/Engine/Source/Programs/UnrealBuildTool/HTML5/HTML5SDKInfo.cs
--- a/Engine/Source/Programs/UnrealBuildTool/HTML5/HTML5SDKInfo.cs
+++ b/Engine/Source/Programs/UnrealBuildTool/HTML5/HTML5SDKInfo.cs
@@ -73,6 +73,22 @@ namespace UnrealBuildTool
 			}
 		}
 
+		static string OPTIMIZER_NAME
+		{
+			get
+			{
+				switch (BuildHostPlatform.Current.Platform)
+				{
+					case UnrealTargetPlatform.Win64:
+						return "optimizer.exe";
+					case UnrealTargetPlatform.Mac:
+						return "optimizer";
+					default:
+						return "";
+				}
+			}
+		}
+
 		static public string DOT_EMSCRIPTEN
 		{
 			get 
@@ -133,7 +149,7 @@ namespace UnrealBuildTool
 								"LLVM_ROOT='" + LLVM_ROOT +"'",
 								"NODE_JS= '" + NODE_JS +"'",
 								"PYTHON= '" + PYTHON +"'",
-								"EMSCRIPTEN_NATIVE_OPTIMIZER='"+Path.Combine(LLVM_ROOT,"optimizer.exe") +"'",
+								"EMSCRIPTEN_NATIVE_OPTIMIZER='" + Path.Combine(LLVM_ROOT, OPTIMIZER_NAME) + "'",
 								"EMSCRIPTEN_ROOT= '" + EMSCRIPTEN_ROOT + "'",
 								"TEMP_DIR= '" + SetupEmscriptenTemp() + "'",
 								"COMPILER_ENGINE = NODE_JS",
@@ -168,7 +184,7 @@ namespace UnrealBuildTool
 
 		public static string EmscriptenCompiler()
 		{
-			return Path.Combine(EmscriptenBase(), "emcc");
+			return "\"" + Path.Combine(EmscriptenBase(), "emcc") + "\""; 
 		}
 
 		public static bool IsSDKInstalled()
