UBT now rebuilds makefile when hotreloading a different module.

diff --git a/Engine/Source/Programs/UnrealBuildTool/Configuration/UEBuildTarget.cs b/Engine/Source/Programs/UnrealBuildTool/Configuration/UEBuildTarget.cs
--- a/Engine/Source/Programs/UnrealBuildTool/Configuration/UEBuildTarget.cs
+++ b/Engine/Source/Programs/UnrealBuildTool/Configuration/UEBuildTarget.cs
@@ -135,8 +135,21 @@ namespace UnrealBuildTool
 		public readonly List<string> FileNames = new List<string>();
 	}
 
-	public class OnlyModule
+	[Serializable]
+	public class OnlyModule : ISerializable
 	{
+		public OnlyModule(SerializationInfo Info, StreamingContext Context)
+		{
+			OnlyModuleName   = Info.GetString("mn");
+			OnlyModuleSuffix = Info.GetString("ms");
+		}
+
+		public void GetObjectData(SerializationInfo Info, StreamingContext Context)
+		{
+			Info.AddValue("mn", OnlyModuleName);
+			Info.AddValue("ms", OnlyModuleSuffix);
+		}
+
 		public OnlyModule(string InitOnlyModuleName)
 		{
 			OnlyModuleName = InitOnlyModuleName;
@@ -763,8 +776,7 @@ namespace UnrealBuildTool
 		public bool bEditorRecompile;
 
 		/** If building only a specific set of modules, these are the modules to build */
-		[NonSerialized]
-		protected List<OnlyModule> OnlyModules = new List<OnlyModule>();
+		public List<OnlyModule> OnlyModules = new List<OnlyModule>();
 
 		/** Kept to determine the correct module parsing order when filtering modules. */
 		[NonSerialized]
@@ -868,6 +880,7 @@ namespace UnrealBuildTool
 			bPrecompile                  = Info.GetBoolean("pc");
 			bUsePrecompiled              = Info.GetBoolean("up");
 			bEditorRecompile             = Info.GetBoolean("er");
+			OnlyModules                  = (List<OnlyModule>)Info.GetValue("om", typeof(List<OnlyModule>));
 			bCompileMonolithic           = Info.GetBoolean("cm");
 			var FlatModuleCsDataKeys   = (string[])Info.GetValue("fk", typeof(string[]));
 			var FlatModuleCsDataValues = (FlatModuleCsDataType[])Info.GetValue("fv", typeof(FlatModuleCsDataType[]));
@@ -897,6 +910,7 @@ namespace UnrealBuildTool
 			Info.AddValue("pc", bPrecompile);
 			Info.AddValue("up", bUsePrecompiled);
 			Info.AddValue("er", bEditorRecompile);
+			Info.AddValue("om", OnlyModules);
 			Info.AddValue("cm", bCompileMonolithic);
 			Info.AddValue("fk", FlatModuleCsData.Keys.ToArray());
 			Info.AddValue("fv", FlatModuleCsData.Values.ToArray());
diff --git a/Engine/Source/Programs/UnrealBuildTool/System/UnrealBuildTool.cs b/Engine/Source/Programs/UnrealBuildTool/System/UnrealBuildTool.cs
--- a/Engine/Source/Programs/UnrealBuildTool/System/UnrealBuildTool.cs
+++ b/Engine/Source/Programs/UnrealBuildTool/System/UnrealBuildTool.cs
@@ -1720,6 +1720,14 @@ namespace UnrealBuildTool
                         // Try to load the UBTMakefile.  It will only be loaded if it has valid content and is not determined to be out of date.    
 						string ReasonNotLoaded;
                         UBTMakefile = LoadUBTMakefile( UBTMakefilePath, out ReasonNotLoaded );
+
+						// Invalid makefile if only modules have changed
+						if (UBTMakefile != null && !TargetDescs.SelectMany(x => x.OnlyModules).Select(x => x.OnlyModuleName.ToLower()).SequenceEqual(UBTMakefile.Targets.SelectMany(x => x.OnlyModules).Select(x => x.OnlyModuleName.ToLower())))
+						{
+							UBTMakefile     = null;
+							ReasonNotLoaded = "modules to compile have changed";
+						}
+
                         if( UBTMakefile == null )
                         { 
                             // If the Makefile couldn't be loaded, then we're not going to be able to continue in "assembler only" mode.  We'll do both
@@ -2241,7 +2249,7 @@ namespace UnrealBuildTool
         [Serializable]
         class UBTMakefile : ISerializable
         {
-			public const int CurrentVersion = 1;
+			public const int CurrentVersion = 2;
 
 			public int Version;
 
