[INTEGRATE] UnrealBuildTool fixes for small modules

- Fixed small game-specific plugins always using non-unity and getting a unique PCH

- You can now override the default behavior of UBT which compiles small game modules for "fast iteration"
        - Many small game modules have a low frequency of iteration, and never need a unique PCH or non-unity enabled
        - You can use the following module-specific variables to override the default behavior:

                      MinSourceFilesForUnityBuildOverride = 1;
                      MinFilesUsingPrecompiledHeaderOverride = BuildConfiguration.MinFilesUsingPrecompiledHeader;

diff --git a/Engine/Source/Programs/UnrealBuildTool/Configuration/UEBuildModule.cs b/Engine/Source/Programs/UnrealBuildTool/Configuration/UEBuildModule.cs
--- a/Engine/Source/Programs/UnrealBuildTool/Configuration/UEBuildModule.cs
+++ b/Engine/Source/Programs/UnrealBuildTool/Configuration/UEBuildModule.cs
@@ -1080,6 +1080,10 @@ namespace UnrealBuildTool
 		/** If true and unity builds are enabled, this module will build without unity. */
 		public bool bFasterWithoutUnity = false;
 
+		/** The number of source files in this module before unity build will be activated for that module.  If set to
+		    anything besides 0, will override the default setting which is controlled by MinGameModuleSourceFilesForUnityBuild */
+		public int MinSourceFilesForUnityBuildOverride = 0;
+
 		/** Overrides BuildConfiguration.MinFilesUsingPrecompiledHeader if non-zero. */
 		public int MinFilesUsingPrecompiledHeaderOverride = 0;
 
@@ -1178,6 +1182,7 @@ namespace UnrealBuildTool
 			bool InUseRTTI,
 			bool InEnableBufferSecurityChecks,
 			bool InFasterWithoutUnity,
+			int InMinSourceFilesForUnityBuildOverride,
 			int InMinFilesUsingPrecompiledHeaderOverride,
 			bool InBuildLocallyWithSNDBS,
 			bool InEnableExceptions,
@@ -1235,7 +1240,8 @@ namespace UnrealBuildTool
 			bUseAVX                                = bInUseAVX;
 			bEnableBufferSecurityChecks 		   = InEnableBufferSecurityChecks;
 			bFasterWithoutUnity                    = InFasterWithoutUnity;
-			MinFilesUsingPrecompiledHeaderOverride = InMinFilesUsingPrecompiledHeaderOverride;
+			MinSourceFilesForUnityBuildOverride    = InMinSourceFilesForUnityBuildOverride;
+            MinFilesUsingPrecompiledHeaderOverride = InMinFilesUsingPrecompiledHeaderOverride;
 			bBuildLocallyWithSNDBS				   = InBuildLocallyWithSNDBS;
 			bEnableExceptions                      = InEnableExceptions;
 			bEnableShadowVariableWarnings          = InEnableShadowVariableWarnings;
@@ -1315,7 +1321,8 @@ namespace UnrealBuildTool
 
 
 			// Check to see if this is an Engine module (including program or plugin modules).  That is, the module is located under the "Engine" folder
-			var IsGameModule = !Utils.IsFileUnderDirectory( this.ModuleDirectory, Path.Combine( ProjectFileGenerator.EngineRelativePath ) );
+			var IsPluginModule = Utils.IsFileUnderDirectory( this.ModuleDirectory, Path.Combine(Target.ProjectDirectory, "Plugins" ) );
+			var IsGameModule = !IsPluginModule && !Utils.IsFileUnderDirectory(this.ModuleDirectory, Path.Combine(ProjectFileGenerator.EngineRelativePath));
 
 			// Should we force a precompiled header to be generated for this module?  Usually, we only bother with a
 			// precompiled header if there are at least several source files in the module (after combining them for unity
@@ -1334,20 +1341,51 @@ namespace UnrealBuildTool
 			}
 
 
+			// Engine modules will always use unity build mode unless MinSourceFilesForUnityBuildOverride is specified in
+			// the module rules file.  By default, game modules only use unity of they have enough source files for that
+			// to be worthwhile.  If you have a lot of small game modules, consider specifying MinSourceFilesForUnityBuildOverride=0
+			// in the modules that you don't typically iterate on source files in very frequently.
+			int MinSourceFilesForUnityBuild = 0;
+			if (MinSourceFilesForUnityBuildOverride != 0)
+			{
+				MinSourceFilesForUnityBuild = MinSourceFilesForUnityBuildOverride;
+			}
+			else if ( IsGameModule )
+			{
+				// Game modules with only a small number of source files are usually better off having faster iteration times
+				// on single source file changes, so we forcibly disable unity build for those modules
+				MinSourceFilesForUnityBuild = BuildConfiguration.MinGameModuleSourceFilesForUnityBuild;
+			}
+
+
 			// Should we use unity build mode for this module?
-			bool bModuleUsesUnityBuild = BuildConfiguration.bUseUnityBuild || BuildConfiguration.bForceUnityBuild;
-			if (!BuildConfiguration.bForceUnityBuild)
+			bool bModuleUsesUnityBuild = false;
+			if (BuildConfiguration.bUseUnityBuild || BuildConfiguration.bForceUnityBuild)
 			{
-				if (bFasterWithoutUnity)
+				if (BuildConfiguration.bForceUnityBuild)
+				{
+					Log.TraceVerbose("Module '{0}' using unity build mode (bForceUnityBuild enabled for this module)", this.Name);
+					bModuleUsesUnityBuild = true;
+				}
+				else if (bFasterWithoutUnity)
 				{
+					Log.TraceVerbose("Module '{0}' not using unity build mode (bFasterWithoutUnity enabled for this module)", this.Name);
 					bModuleUsesUnityBuild = false;
 				}
-				else if (IsGameModule && SourceFilesToBuild.CPPFiles.Count < BuildConfiguration.MinGameModuleSourceFilesForUnityBuild)
+				else if (SourceFilesToBuild.CPPFiles.Count < MinSourceFilesForUnityBuild)
 				{
-					// Game modules with only a small number of source files are usually better off having faster iteration times
-					// on single source file changes, so we forcibly disable unity build for those modules
+					Log.TraceVerbose("Module '{0}' not using unity build mode (module with fewer than {1} source files)", this.Name, MinSourceFilesForUnityBuild);
 					bModuleUsesUnityBuild = false;
 				}
+				else
+				{
+					Log.TraceVerbose("Module '{0}' using unity build mode (enabled in BuildConfiguration)", this.Name);
+					bModuleUsesUnityBuild = true;
+				}
+			}
+			else
+			{
+				Log.TraceVerbose("Module '{0}' not using unity build mode (disabled in BuildConfiguration)", this.Name);
 			}
 
 			// The environment with which to compile the CPP files
@@ -1895,7 +1933,7 @@ namespace UnrealBuildTool
 			Result.Config.bUseRTTI                               = bUseRTTI;
 			Result.Config.bUseAVX                                = bUseAVX;
 			Result.Config.bEnableBufferSecurityChecks            = bEnableBufferSecurityChecks;
-			Result.Config.bFasterWithoutUnity                    = bFasterWithoutUnity;
+			Result.Config.MinSourceFilesForUnityBuildOverride	 = MinSourceFilesForUnityBuildOverride;
 			Result.Config.MinFilesUsingPrecompiledHeaderOverride = MinFilesUsingPrecompiledHeaderOverride;
 			Result.Config.bBuildLocallyWithSNDBS				 = bBuildLocallyWithSNDBS;
 			Result.Config.bEnableExceptions                      = bEnableExceptions;
@@ -2146,6 +2184,7 @@ namespace UnrealBuildTool
 			bool InUseRTTI,
 			bool InEnableBufferSecurityChecks,
 			bool InFasterWithoutUnity,
+			int InMinSourceFilesForUnityBuildOverride,
 			int InMinFilesUsingPrecompiledHeaderOverride,
 			bool InEnableExceptions,
 			bool InEnableShadowVariableWarnings,
@@ -2158,7 +2197,7 @@ namespace UnrealBuildTool
 			InPublicIncludePathModuleNames,InPublicDependencyModuleNames,InPublicDelayLoadDLLs,InPublicAdditionalLibraries,InPublicFrameworks,InPublicWeakFrameworks,InPublicAdditionalFrameworks,InPublicAdditionalShadowFiles,InPublicAdditionalBundleResources,
 			InPrivateIncludePaths,InPrivateIncludePathModuleNames,InPrivateDependencyModuleNames,
             InCircularlyReferencedDependentModules, InDynamicallyLoadedModuleNames, InPlatformSpecificDynamicallyLoadedModuleNames, InRuntimeDependencies, InOptimizeCode,
-			InAllowSharedPCH, InSharedPCHHeaderFile, InUseRTTI, InEnableBufferSecurityChecks, InFasterWithoutUnity, InMinFilesUsingPrecompiledHeaderOverride, true,
+			InAllowSharedPCH, InSharedPCHHeaderFile, InUseRTTI, InEnableBufferSecurityChecks, InFasterWithoutUnity, InMinSourceFilesForUnityBuildOverride, InMinFilesUsingPrecompiledHeaderOverride, true,
 			InEnableExceptions, InEnableShadowVariableWarnings, bInBuildSourceFiles, InBuildCsFilename, bInUseAVX)
 		{
 			PrivateAssemblyReferences = HashSetFromOptionalEnumerableStringParameter(InPrivateAssemblyReferences);
diff --git a/Engine/Source/Programs/UnrealBuildTool/Configuration/UEBuildTarget.cs b/Engine/Source/Programs/UnrealBuildTool/Configuration/UEBuildTarget.cs
--- a/Engine/Source/Programs/UnrealBuildTool/Configuration/UEBuildTarget.cs
+++ b/Engine/Source/Programs/UnrealBuildTool/Configuration/UEBuildTarget.cs
@@ -2464,6 +2464,7 @@ namespace UnrealBuildTool
 				InUseRTTI: false,
 				InEnableBufferSecurityChecks: true,
 				InFasterWithoutUnity: true,
+				InMinSourceFilesForUnityBuildOverride: 0,
 				InMinFilesUsingPrecompiledHeaderOverride: 0,
 				InBuildLocallyWithSNDBS: false,
 				InEnableExceptions: false,
@@ -3508,6 +3509,7 @@ namespace UnrealBuildTool
 							InUseRTTI: RulesObject.bUseRTTI,
 							InEnableBufferSecurityChecks: RulesObject.bEnableBufferSecurityChecks,
 							InFasterWithoutUnity: RulesObject.bFasterWithoutUnity,
+							InMinSourceFilesForUnityBuildOverride: RulesObject.MinSourceFilesForUnityBuildOverride,
 							InMinFilesUsingPrecompiledHeaderOverride: RulesObject.MinFilesUsingPrecompiledHeaderOverride,
 							InBuildLocallyWithSNDBS: RulesObject.bBuildLocallyWithSNDBS,
 							InEnableExceptions: RulesObject.bEnableExceptions,
@@ -3553,6 +3555,7 @@ namespace UnrealBuildTool
 							InUseRTTI: RulesObject.bUseRTTI,
 							InEnableBufferSecurityChecks: RulesObject.bEnableBufferSecurityChecks,
 							InFasterWithoutUnity: RulesObject.bFasterWithoutUnity,
+							InMinSourceFilesForUnityBuildOverride: RulesObject.MinSourceFilesForUnityBuildOverride,
 							InMinFilesUsingPrecompiledHeaderOverride: RulesObject.MinFilesUsingPrecompiledHeaderOverride,
 							InEnableExceptions: RulesObject.bEnableExceptions,
 							InEnableShadowVariableWarnings: RulesObject.bEnableShadowVariableWarnings,
diff --git a/Engine/Source/Programs/UnrealBuildTool/System/CPPEnvironment.cs b/Engine/Source/Programs/UnrealBuildTool/System/CPPEnvironment.cs
--- a/Engine/Source/Programs/UnrealBuildTool/System/CPPEnvironment.cs
+++ b/Engine/Source/Programs/UnrealBuildTool/System/CPPEnvironment.cs
@@ -209,8 +209,12 @@ namespace UnrealBuildTool
         /** If true and unity builds are enabled, this module will build without unity. */
         public bool bFasterWithoutUnity = false;
 
-        /** Overrides BuildConfiguration.MinFilesUsingPrecompiledHeader if non-zero. */
-        public int MinFilesUsingPrecompiledHeaderOverride = 0;
+		/** The number of source files in this module before unity build will be activated for that module.  If set to
+		    anything besides -1, will override the default setting which is controlled by MinGameModuleSourceFilesForUnityBuild */
+		public int MinSourceFilesForUnityBuildOverride = 0;
+
+		/** Overrides BuildConfiguration.MinFilesUsingPrecompiledHeader if non-zero. */
+		public int MinFilesUsingPrecompiledHeaderOverride = 0;
 
 		/**  Module uses a #import so must be built locally when compiling with SN-DBS */
 		public bool bBuildLocallyWithSNDBS = false;
@@ -278,6 +282,7 @@ namespace UnrealBuildTool
 			bUseRTTI                               = InCopyEnvironment.bUseRTTI;
 			bUseAVX                                = InCopyEnvironment.bUseAVX;
 			bFasterWithoutUnity                    = InCopyEnvironment.bFasterWithoutUnity;
+			MinSourceFilesForUnityBuildOverride	   = InCopyEnvironment.MinSourceFilesForUnityBuildOverride;
 			MinFilesUsingPrecompiledHeaderOverride = InCopyEnvironment.MinFilesUsingPrecompiledHeaderOverride;
 			bBuildLocallyWithSNDBS				   = InCopyEnvironment.bBuildLocallyWithSNDBS;
 			bEnableExceptions                      = InCopyEnvironment.bEnableExceptions;
diff --git a/Engine/Source/Programs/UnrealBuildTool/System/RulesCompiler.cs b/Engine/Source/Programs/UnrealBuildTool/System/RulesCompiler.cs
--- a/Engine/Source/Programs/UnrealBuildTool/System/RulesCompiler.cs
+++ b/Engine/Source/Programs/UnrealBuildTool/System/RulesCompiler.cs
@@ -212,6 +212,10 @@ namespace UnrealBuildTool
 		/** If true and unity builds are enabled, this module will build without unity. */
 		public bool bFasterWithoutUnity = false;
 
+		/** The number of source files in this module before unity build will be activated for that module.  If set to
+		    anything besides -1, will override the default setting which is controlled by MinGameModuleSourceFilesForUnityBuild */
+		public int MinSourceFilesForUnityBuildOverride = 0;
+
 		/** Overrides BuildConfiguration.MinFilesUsingPrecompiledHeader if non-zero. */
 		public int MinFilesUsingPrecompiledHeaderOverride = 0;
 
