[INTEGRATE] Compile times: Fixed game plugin modules not able to use engine shared PCHs
- This change greatly improves compile times for game plugin modules
- Game modules never depend on the game itself, so there was no reason not to use shared PCHs by default

diff --git a/Engine/Source/Programs/UnrealBuildTool/System/RulesCompiler.cs b/Engine/Source/Programs/UnrealBuildTool/System/RulesCompiler.cs
--- a/Engine/Source/Programs/UnrealBuildTool/System/RulesCompiler.cs
+++ b/Engine/Source/Programs/UnrealBuildTool/System/RulesCompiler.cs
@@ -1453,6 +1453,7 @@ namespace UnrealBuildTool
 						// Choose code optimization options based on module type (game/engine) if
 						// default optimization method is selected.
 						bool bIsEngineModule = Utils.IsFileUnderDirectory(ModuleFileName, ProjectFileGenerator.EngineRelativePath);
+						bool IsGamePluginModule = bProjectModule ? Utils.IsFileUnderDirectory(ModuleFileName, Path.Combine(UnrealBuildTool.GetUProjectPath(), "Plugins")) : false;
 						if (RulesObject.OptimizeCode == ModuleRules.CodeOptimization.Default)
 						{
 							// Engine/Source and Engine/Plugins are considered 'Engine' code...
@@ -1468,12 +1469,12 @@ namespace UnrealBuildTool
 							}
 						}
 
-						// Disable shared PCHs for game modules by default
+						// Disable shared PCHs for game modules by default (but not game plugins, since they won't depend on the game's PCH!)
 						if (RulesObject.PCHUsage == ModuleRules.PCHUsageMode.Default)
 						{
 							// Note that bIsEngineModule includes Engine/Plugins, so Engine/Plugins will use shared PCHs.
 							var IsProgramTarget = Target.Type != null && Target.Type == TargetRules.TargetType.Program;
-							if (bIsEngineModule || IsProgramTarget)
+							if (bIsEngineModule || IsProgramTarget || IsGamePluginModule)
 							{
 								// Engine module or plugin module -- allow shared PCHs
 								RulesObject.PCHUsage = ModuleRules.PCHUsageMode.UseSharedPCHs;
