Fix PCH environment changing when a module has an include path dependency, but the decision to compile that module is dependent on something else. Now always generates a consistent compile environment regardless of whether the other module binary is actually built.

diff --git a/Engine/Source/Programs/UnrealBuildTool/Configuration/UEBuildModule.cs b/Engine/Source/Programs/UnrealBuildTool/Configuration/UEBuildModule.cs
--- a/Engine/Source/Programs/UnrealBuildTool/Configuration/UEBuildModule.cs
+++ b/Engine/Source/Programs/UnrealBuildTool/Configuration/UEBuildModule.cs
@@ -588,7 +588,18 @@ namespace UnrealBuildTool
 				Definitions.AddRange(PublicDefinitions);
 				
 				// If this module is being built into a DLL or EXE, set up an IMPORTS or EXPORTS definition for it.
-				if (Binary != null)
+				if (Binary == null)
+				{
+					// If we're referencing include paths for a module that's not being built, we don't actually need to import anything from it, but we need to avoid barfing when
+					// the compiler encounters an _API define. We also want to avoid changing the compile environment in cases where the module happens to be compiled because it's a dependency
+					// of something else, which cause a fall-through to the code below and set up an empty _API define.
+					if(bIncludePathsOnly)
+					{
+						Log.TraceVerbose( "{0}: Include paths only for {1} (no binary)", Path.GetFileNameWithoutExtension( SourceBinary.Config.OutputFilePaths[0] ), Name );
+						Definitions.Add( Name.ToUpperInvariant() + "_API=" );
+					}
+				}
+				else
 				{
 					string BinaryPath = Binary.Config.OutputFilePaths[0];
 					string SourceBinaryPath = SourceBinary.Config.OutputFilePaths[0];
