Fixed UHT compiled with older compiler when targeting Visual Studio 2015

diff --git a/Engine/Source/Programs/UnrealBuildTool/System/ExternalExecution.cs b/Engine/Source/Programs/UnrealBuildTool/System/ExternalExecution.cs
--- a/Engine/Source/Programs/UnrealBuildTool/System/ExternalExecution.cs
+++ b/Engine/Source/Programs/UnrealBuildTool/System/ExternalExecution.cs
@@ -694,6 +694,12 @@ namespace UnrealBuildTool
 						{
 							UBTArguments.Append(" -noxge");
 						}
+
+						// Propagate command-line option to enable the newer compiler on Windows platform
+						if (UnrealBuildTool.CommandLineContains("-2015"))
+						{
+							UBTArguments.Append(" -2015");
+						}
 						
 						if ( RunExternalExecutable( UnrealBuildTool.GetUBTPath(), UBTArguments.ToString() ) != 0 )
 						{ 
