Make it clearer that TargetReceipt.Read can fail if the file doesn't exist or is formatted incorrectly.

diff --git a/Engine/Source/Programs/AutomationTool/Scripts/BuildPluginCommand.Automation.cs b/Engine/Source/Programs/AutomationTool/Scripts/BuildPluginCommand.Automation.cs
--- a/Engine/Source/Programs/AutomationTool/Scripts/BuildPluginCommand.Automation.cs
+++ b/Engine/Source/Programs/AutomationTool/Scripts/BuildPluginCommand.Automation.cs
@@ -115,7 +115,11 @@ class BuildPlugin : BuildCommand
 		List<BuildProduct> BuildProducts = new List<BuildProduct>();
 		foreach(string ReceiptFileName in ReceiptFileNames)
 		{
-			TargetReceipt Receipt = TargetReceipt.Read(ReceiptFileName);
+			TargetReceipt Receipt;
+			if(!TargetReceipt.TryRead(ReceiptFileName, out Receipt))
+			{
+				throw new AutomationException("Missing or invalid target receipt ({0})", ReceiptFileName);
+			}
 			BuildProducts.AddRange(Receipt.BuildProducts);
 		}
 		return BuildProducts;
diff --git a/Engine/Source/Programs/AutomationTool/Scripts/CopyBuildToStagingDirectory.Automation.cs b/Engine/Source/Programs/AutomationTool/Scripts/CopyBuildToStagingDirectory.Automation.cs
--- a/Engine/Source/Programs/AutomationTool/Scripts/CopyBuildToStagingDirectory.Automation.cs
+++ b/Engine/Source/Programs/AutomationTool/Scripts/CopyBuildToStagingDirectory.Automation.cs
@@ -1561,8 +1561,14 @@ public partial class Project : CommandUtils
 
 						}
 
+						// Read the receipt for this target
+						TargetReceipt Receipt;
+						if(!TargetReceipt.TryRead(ReceiptFileName, out Receipt))
+						{
+							throw new AutomationException("Missing or invalid target receipt ({0})", ReceiptFileName);
+						}
+
 						// Convert the paths to absolute
-						TargetReceipt Receipt = TargetReceipt.Read(ReceiptFileName);
 						Receipt.ExpandPathVariables(EngineDir, Path.GetDirectoryName(Params.RawProjectPath));
 						Receipt.SetDependenciesToBeRequired(bRequireStagedFilesToExist);
 						TargetsToStage.Add(Receipt);
diff --git a/Engine/Source/Programs/UnrealBuildTool/Configuration/UEBuildTarget.cs b/Engine/Source/Programs/UnrealBuildTool/Configuration/UEBuildTarget.cs
--- a/Engine/Source/Programs/UnrealBuildTool/Configuration/UEBuildTarget.cs
+++ b/Engine/Source/Programs/UnrealBuildTool/Configuration/UEBuildTarget.cs
@@ -1139,8 +1139,8 @@ namespace UnrealBuildTool
 			// Expand all the paths in the receipt; they'll currently use variables for the engine and project directories
 			IUEBuildPlatform BuildPlatform = UEBuildPlatform.GetBuildPlatform(Platform);
 			// Expand all the paths in the receipt; they'll currently use variables for the engine and project directories
-			TargetReceipt ReceiptWithFullPaths = TargetReceipt.Read(ReceiptFileName);
-			if (ReceiptWithFullPaths == null)
+			TargetReceipt ReceiptWithFullPaths;
+			if(!TargetReceipt.TryRead(ReceiptFileName, out ReceiptWithFullPaths))
 			{
 				ReceiptWithFullPaths = new TargetReceipt(Receipt);
 			}
@@ -1402,8 +1402,8 @@ namespace UnrealBuildTool
 			ReceiptWithFullPaths.ExpandPathVariables(BuildConfiguration.RelativeEnginePath, ProjectDirectory);
 			IUEBuildPlatform BuildPlatform = UEBuildPlatform.GetBuildPlatform(Platform);
 			// Expand all the paths in the receipt; they'll currently use variables for the engine and project directories
-			TargetReceipt BuiltReceiptWithFullPaths = TargetReceipt.Read(ReceiptFileName);
-			if (BuiltReceiptWithFullPaths == null)
+			TargetReceipt BuiltReceiptWithFullPaths;
+			if(!TargetReceipt.TryRead(ReceiptFileName, out BuiltReceiptWithFullPaths))
 			{
 				return;
 			}
diff --git a/Engine/Source/Programs/UnrealBuildTool/System/ExternalExecution.cs b/Engine/Source/Programs/UnrealBuildTool/System/ExternalExecution.cs
--- a/Engine/Source/Programs/UnrealBuildTool/System/ExternalExecution.cs
+++ b/Engine/Source/Programs/UnrealBuildTool/System/ExternalExecution.cs
@@ -269,16 +269,7 @@ namespace UnrealBuildTool
 				}
 
 				TargetReceipt Receipt;
-				try
-				{
-					Receipt = TargetReceipt.Read(ReceiptPath);
-					if(Receipt == null)
-					{
-						Timestamp = DateTime.MaxValue;
-						return false;
-					}
-				}
-				catch(Exception)
+				if(!TargetReceipt.TryRead(ReceiptPath, out Receipt))
 				{
 					Timestamp = DateTime.MaxValue;
 					return false;
diff --git a/Engine/Source/Programs/UnrealBuildTool/System/TargetReceipt.cs b/Engine/Source/Programs/UnrealBuildTool/System/TargetReceipt.cs
--- a/Engine/Source/Programs/UnrealBuildTool/System/TargetReceipt.cs
+++ b/Engine/Source/Programs/UnrealBuildTool/System/TargetReceipt.cs
@@ -402,25 +402,26 @@ namespace UnrealBuildTool
 		/// Read a receipt from disk.
 		/// </summary>
 		/// <param name="FileName">Filename to read from</param>
-		public static TargetReceipt Read(string FileName)
+		public static bool TryRead(string FileName, out TargetReceipt Receipt)
 		{
-			if (File.Exists(FileName))
+			if (!File.Exists(FileName))
 			{
-				try
-				{
-					using (StreamReader Reader = new StreamReader(FileName))
-					{
-						return (TargetReceipt)Serializer.Deserialize(Reader);
-					}
-				}
-				catch(Exception)
+				Receipt = null;
+				return false;
+			}
+
+			try
+			{
+				using (StreamReader Reader = new StreamReader(FileName))
 				{
-					return null;
+					Receipt = (TargetReceipt)Serializer.Deserialize(Reader);
+					return true;
 				}
 			}
-			else
+			catch(Exception)
 			{
-				return null;
+				Receipt = null;
+				return false;
 			}
 		}
 
