Update Distcc/DMUCS integration for new version of DistCode - must provide the correct DistProp & server for DMUCS and acquire the number of build hosts via Xcode on OS X.

diff --git a/Engine/Source/Programs/UnrealBuildTool/Configuration/BuildConfiguration.cs b/Engine/Source/Programs/UnrealBuildTool/Configuration/BuildConfiguration.cs
--- a/Engine/Source/Programs/UnrealBuildTool/Configuration/BuildConfiguration.cs
+++ b/Engine/Source/Programs/UnrealBuildTool/Configuration/BuildConfiguration.cs
@@ -433,12 +433,30 @@ namespace UnrealBuildTool
 		[XmlConfig]
 		public static bool bAllowDistccLocalFallback;
 
+		/// <summary>
+		/// When true enable verbose distcc output to aid debugging.
+		/// </summary>
+		[XmlConfig]
+		public static bool bVerboseDistccOutput;
+
 		/// <summary>
 		/// Path to the Distcc & DMUCS executables.
 		/// </summary>
 		[XmlConfig]
 		public static string DistccExecutablesPath;
 
+		/// <summary>
+		/// DMUCS coordinator hostname or IP address.
+		/// </summary>
+		[XmlConfig]
+		public static string DMUCSCoordinator;
+
+		/// <summary>
+		/// The DMUCS distinguishing property value to request for compile hosts.
+		/// </summary>
+		[XmlConfig]
+		public static string DMUCSDistProp;
+
 		/// <summary>
 		/// Run UnrealCodeAnalyzer instead of regular compilation. Requires setting UCA mode by bUCACheckUObjectThreadSafety or bUCACheckPCHFiles.
 		/// </summary>
@@ -591,10 +609,15 @@ namespace UnrealBuildTool
 			// Distcc requires some setup - so by default disable it so we don't break local or remote building
 			bAllowDistcc = false;
 			bAllowDistccLocalFallback = true;
+			bVerboseDistccOutput = false;
 
 			// The normal Posix place to install distcc/dmucs would be /usr/local so start there
 			DistccExecutablesPath = "/usr/local/bin";
 
+			// The standard coordinator is localhost and the default distprop is empty
+			DMUCSCoordinator = "localhost";
+			DMUCSDistProp = "";
+
 			// By default compile code instead of analyzing it.
 			bRunUnrealCodeAnalyzer = false;
 
@@ -614,6 +637,31 @@ namespace UnrealBuildTool
 				{
 					DistccExecutablesPath = MacDistccExecutablesPath;
 				}
+
+				if (System.IO.File.Exists (UserDir + "/Library/Preferences/com.marksatt.DistCode.plist"))
+				{
+					using (System.Diagnostics.Process DefaultsProcess = new System.Diagnostics.Process())
+					{
+						try
+						{
+							DefaultsProcess.StartInfo.FileName = "/usr/bin/defaults";
+							DefaultsProcess.StartInfo.CreateNoWindow = true;
+							DefaultsProcess.StartInfo.UseShellExecute = false;
+							DefaultsProcess.StartInfo.RedirectStandardOutput = true;
+							DefaultsProcess.StartInfo.Arguments = "read com.marksatt.DistCode DistProp";
+							DefaultsProcess.Start();
+							string Output = DefaultsProcess.StandardOutput.ReadToEnd();
+							DefaultsProcess.WaitForExit();
+							if(DefaultsProcess.ExitCode == 0)
+							{
+								DMUCSDistProp = Output;
+							}
+						}
+						catch (Exception e)
+						{
+						}
+					}
+				}
 			}
 		}
 
diff --git a/Engine/Source/Programs/UnrealBuildTool/System/Distcc.cs b/Engine/Source/Programs/UnrealBuildTool/System/Distcc.cs
--- a/Engine/Source/Programs/UnrealBuildTool/System/Distcc.cs
+++ b/Engine/Source/Programs/UnrealBuildTool/System/Distcc.cs
@@ -22,29 +22,72 @@ namespace UnrealBuildTool
 				const float LoopSleepTime = 0.1f;
 
 				int MaxActionsToExecuteInParallel = 0;
+
 				string UserDir = Environment.GetEnvironmentVariable("HOME");
 				string HostsInfo = UserDir + "/.dmucs/hosts-info";
-				System.IO.StreamReader File = new System.IO.StreamReader(HostsInfo);
-				string Line = null;
-				while((Line = File.ReadLine()) != null)
+				string NumUBTBuildTasks = Environment.GetEnvironmentVariable("NumUBTBuildTasks");
+				Int32 MaxUBTBuildTasks = MaxActionsToExecuteInParallel;
+				if(Int32.TryParse(NumUBTBuildTasks, out MaxUBTBuildTasks))
+				{
+					MaxActionsToExecuteInParallel = MaxUBTBuildTasks;
+				}
+				else if (BuildHostPlatform.Current.Platform == UnrealTargetPlatform.Mac)
+				{
+					using (System.Diagnostics.Process DefaultsProcess = new System.Diagnostics.Process())
+					{
+						try
+						{
+							DefaultsProcess.StartInfo.FileName = "/usr/bin/defaults";
+							DefaultsProcess.StartInfo.CreateNoWindow = true;
+							DefaultsProcess.StartInfo.UseShellExecute = false;
+							DefaultsProcess.StartInfo.RedirectStandardOutput = true;
+							DefaultsProcess.StartInfo.Arguments = "com.apple.dt.Xcode IDEBuildOperationMaxNumberOfConcurrentCompileTasks";
+							DefaultsProcess.Start();
+							string Output = DefaultsProcess.StandardOutput.ReadToEnd();
+							DefaultsProcess.WaitForExit();
+							if(DefaultsProcess.ExitCode == 0 && Int32.TryParse(Output, out MaxUBTBuildTasks))
+							{
+								MaxActionsToExecuteInParallel = MaxUBTBuildTasks;
+							}
+						}
+						catch (Exception e)
+						{
+						}
+					}
+				}
+				else if (System.IO.File.Exists (HostsInfo))
 				{
-					var HostInfo = Line.Split(' ');
-					if (HostInfo.Count () == 3) 
+					System.IO.StreamReader File = new System.IO.StreamReader(HostsInfo);
+					string Line = null;
+					while((Line = File.ReadLine()) != null)
 					{
-						int NumCPUs = 0;
-						if (System.Int32.TryParse (HostInfo [1], out NumCPUs)) 
+						var HostInfo = Line.Split(' ');
+						if (HostInfo.Count () == 3) 
 						{
-							MaxActionsToExecuteInParallel += NumCPUs;
+							int NumCPUs = 0;
+							if (System.Int32.TryParse (HostInfo [1], out NumCPUs)) 
+							{
+								MaxActionsToExecuteInParallel += NumCPUs;
+							}
 						}
 					}
+					File.Close();
+				}
+				else
+				{
+					MaxActionsToExecuteInParallel = System.Environment.ProcessorCount;
 				}
-				File.Close();
 
-				if (BuildConfiguration.bAllowDistccLocalFallback == false) 
+				if (BuildConfiguration.bAllowDistccLocalFallback == false)
 				{
 					Environment.SetEnvironmentVariable("DISTCC_FALLBACK", "0");
 				}
 
+				if (BuildConfiguration.bVerboseDistccOutput == true)
+				{
+					Environment.SetEnvironmentVariable("DISTCC_VERBOSE", "1");
+				}
+
 				string DistccExecutable = BuildConfiguration.DistccExecutablesPath + "/distcc";
 				string GetHostExecutable = BuildConfiguration.DistccExecutablesPath + "/gethost";
 
@@ -144,7 +187,8 @@ namespace UnrealBuildTool
 									{
 										if ((Action.ActionType == ActionType.Compile || Action.ActionType == ActionType.Link) && DistccExecutable != null && GetHostExecutable != null) 
 										{
-											string NewCommandArguments = "--wait -1 \"" + DistccExecutable + "\" " + Action.CommandPath + " " + Action.CommandArguments;
+											string TypeCommand = String.IsNullOrEmpty(BuildConfiguration.DMUCSDistProp) ? "" : (" --type " + BuildConfiguration.DMUCSDistProp);
+											string NewCommandArguments = "--server " + BuildConfiguration.DMUCSCoordinator + TypeCommand + " --wait -1 \"" + DistccExecutable + "\" " + Action.CommandPath + " " + Action.CommandArguments;
 											Action.CommandPath = GetHostExecutable;
 											Action.CommandArguments = NewCommandArguments;
 										}
