Remove unnecessary variable declaration to silence unused warning.

diff --git a/Engine/Source/Programs/UnrealBuildTool/Configuration/BuildConfiguration.cs b/Engine/Source/Programs/UnrealBuildTool/Configuration/BuildConfiguration.cs
--- a/Engine/Source/Programs/UnrealBuildTool/Configuration/BuildConfiguration.cs
+++ b/Engine/Source/Programs/UnrealBuildTool/Configuration/BuildConfiguration.cs
@@ -657,7 +657,7 @@ namespace UnrealBuildTool
 								DMUCSDistProp = Output;
 							}
 						}
-						catch (Exception e)
+						catch (Exception)
 						{
 						}
 					}
diff --git a/Engine/Source/Programs/UnrealBuildTool/System/Distcc.cs b/Engine/Source/Programs/UnrealBuildTool/System/Distcc.cs
--- a/Engine/Source/Programs/UnrealBuildTool/System/Distcc.cs
+++ b/Engine/Source/Programs/UnrealBuildTool/System/Distcc.cs
@@ -50,7 +50,7 @@ namespace UnrealBuildTool
 								MaxActionsToExecuteInParallel = MaxUBTBuildTasks;
 							}
 						}
-						catch (Exception e)
+						catch (Exception)
 						{
 						}
 					}
