Disable some warnings for VS2015

diff --git a/Engine/Source/Programs/UnrealBuildTool/Windows/VCToolChain.cs b/Engine/Source/Programs/UnrealBuildTool/Windows/VCToolChain.cs
--- a/Engine/Source/Programs/UnrealBuildTool/Windows/VCToolChain.cs
+++ b/Engine/Source/Programs/UnrealBuildTool/Windows/VCToolChain.cs
@@ -584,6 +584,20 @@ namespace UnrealBuildTool
 			{
 				AddDefinition( Arguments, "UNREAL_CODE_ANALYZER");
 			}
+
+			//@todo: Disable warnings for VS2015. These should be reenabled as we clear the reasons for them out of the engine source and the VS2015 toolchain evolves.
+			if (WindowsPlatform.Compiler == WindowsCompiler.VisualStudio2015)
+			{
+				// Disable shadow variable warnings
+				Arguments.Append(" /wd4456"); // 4456 - declaration of 'LocalVariable' hides previous local declaration
+				Arguments.Append(" /wd4458"); // 4458 - declaration of 'parameter' hides class member
+				Arguments.Append(" /wd4459"); // 4459 - declaration of 'LocalVariable' hides global declaration
+
+				Arguments.Append(" /wd4463"); // 4463 - overflow; assigning 1 to bit-field that can only hold values from -1 to 0
+
+				Arguments.Append(" /wd4838"); // 4838: conversion from 'type1' to 'type2' requires a narrowing conversion
+			}
+
 		}
 
 		static void AppendCLArguments_CPP( CPPEnvironment CompileEnvironment, StringBuilder Arguments )
