Disable DualShock and Morpheus support for VS2015

diff --git a/Engine/Source/Programs/UnrealBuildTool/Windows/UEBuildWindows.cs b/Engine/Source/Programs/UnrealBuildTool/Windows/UEBuildWindows.cs
--- a/Engine/Source/Programs/UnrealBuildTool/Windows/UEBuildWindows.cs
+++ b/Engine/Source/Programs/UnrealBuildTool/Windows/UEBuildWindows.cs
@@ -577,7 +577,8 @@ namespace UnrealBuildTool
 			InBuildTarget.GlobalCompileEnvironment.Config.Definitions.Add("PLATFORM_WINDOWS=1");
 
             String MorpheusShaderPath = Path.Combine(BuildConfiguration.RelativeEnginePath, "Shaders/PS4/PostProcessHMDMorpheus.usf");
-            if (File.Exists(MorpheusShaderPath))
+            //@todo: VS2015 currently does not have support for Morpheus
+            if (File.Exists(MorpheusShaderPath) && WindowsPlatform.Compiler != WindowsCompiler.VisualStudio2015)
             {
                 InBuildTarget.GlobalCompileEnvironment.Config.Definitions.Add("HAS_MORPHEUS=1");                
 
