Avoid throwing exceptions if the path to .buildmutex files doesn't exist.

diff --git a/Engine/Source/Programs/UnrealBuildTool/System/DependencyCache.cs b/Engine/Source/Programs/UnrealBuildTool/System/DependencyCache.cs
--- a/Engine/Source/Programs/UnrealBuildTool/System/DependencyCache.cs
+++ b/Engine/Source/Programs/UnrealBuildTool/System/DependencyCache.cs
@@ -215,13 +215,16 @@ namespace UnrealBuildTool
 				}
 			}
 
-			try
-			{
-				File.Delete(CachePath + ".buildmutex");
-			}
-			catch
+			if(File.Exists(CachePath + ".buildmutex"))
 			{
-				// We don't care if we couldn't delete this file, as maybe it couldn't have been created in the first place.
+				try
+				{
+					File.Delete(CachePath + ".buildmutex");
+				}
+				catch
+				{
+					// We don't care if we couldn't delete this file, as maybe it couldn't have been created in the first place.
+				}
 			}
 		}
 
diff --git a/Engine/Source/Programs/UnrealBuildTool/System/FlatCPPIncludeDepencencyCache.cs b/Engine/Source/Programs/UnrealBuildTool/System/FlatCPPIncludeDepencencyCache.cs
--- a/Engine/Source/Programs/UnrealBuildTool/System/FlatCPPIncludeDepencencyCache.cs
+++ b/Engine/Source/Programs/UnrealBuildTool/System/FlatCPPIncludeDepencencyCache.cs
@@ -172,13 +172,16 @@ namespace UnrealBuildTool
 				}
 			}
 
-			try
-			{
-				File.Delete(CacheFileItem.AbsolutePath + ".buildmutex");
-			}
-			catch
+			if(File.Exists(CacheFileItem.AbsolutePath + ".buildmutex"))
 			{
-				// We don't care if we couldn't delete this file, as maybe it couldn't have been created in the first place.
+				try
+				{
+					File.Delete(CacheFileItem.AbsolutePath + ".buildmutex");
+				}
+				catch
+				{
+					// We don't care if we couldn't delete this file, as maybe it couldn't have been created in the first place.
+				}
 			}
 		}
 
