Always compile ShaderCompileWorker along with the editor.

diff --git a/Engine/Source/Programs/UnrealBuildTool/System/VCProjectFileGenerator.cs b/Engine/Source/Programs/UnrealBuildTool/System/VCProjectFileGenerator.cs
--- a/Engine/Source/Programs/UnrealBuildTool/System/VCProjectFileGenerator.cs
+++ b/Engine/Source/Programs/UnrealBuildTool/System/VCProjectFileGenerator.cs
@@ -4,6 +4,7 @@ using System;
 using System.Collections.Generic;
 using System.Text;
 using System.IO;
+using System.Linq;
 
 namespace UnrealBuildTool
 {
@@ -361,6 +362,20 @@ namespace UnrealBuildTool
 				throw new BuildException( "Unexpected ProjectFileFormat" );
 			}
 
+			// Find the project for ShaderCompileWorker
+			ProjectFile ShaderCompileWorkerProject = null;
+			foreach(ProjectFile Project in AllProjectFiles)
+			{
+				if(Project.ProjectTargets.Count == 1)
+				{
+					string TargetFilePath = Project.ProjectTargets[0].TargetFilePath;
+					if(TargetFilePath != null && Path.GetFileNameWithoutExtension(Path.GetFileNameWithoutExtension(TargetFilePath)).Equals("ShaderCompileWorker", StringComparison.InvariantCultureIgnoreCase))
+					{
+						ShaderCompileWorkerProject = Project;
+						break;
+					}
+				}
+			}
 
 			// Solution folders, files and project entries
 			{
@@ -435,6 +450,10 @@ namespace UnrealBuildTool
 							Dependencies.Add(UBTProject);
 							Dependencies.AddRange(UBTProject.DependsOnProjects);
 						}
+						if(CurProject.IsGeneratedProject && ShaderCompileWorkerProject != null && CurProject.ProjectTargets.Any(x => x.TargetRules != null && x.TargetRules.Type == TargetRules.TargetType.Editor))
+						{
+							Dependencies.Add(ShaderCompileWorkerProject);
+						}
 						Dependencies.AddRange(CurProject.DependsOnProjects);
 
 						if (Dependencies.Count > 0)
@@ -690,6 +709,11 @@ namespace UnrealBuildTool
 											}
 										}
 
+										// Always allow SCW to build in editor configurations
+										if(MatchingProjectTarget == null && SolutionConfigCombination.TargetConfigurationName == TargetRules.TargetType.Editor.ToString() && CurProject == ShaderCompileWorkerProject)
+										{
+											MatchingProjectTarget = CurProject.ProjectTargets[0];
+										}
 
 										string ProjectConfigName;
 										string ProjectPlatformName;
