Always compile ShaderCompileWorker along with the editor in Xcode

diff --git a/Engine/Source/Programs/UnrealBuildTool/System/XcodeProjectFileGenerator.cs b/Engine/Source/Programs/UnrealBuildTool/System/XcodeProjectFileGenerator.cs
--- a/Engine/Source/Programs/UnrealBuildTool/System/XcodeProjectFileGenerator.cs
+++ b/Engine/Source/Programs/UnrealBuildTool/System/XcodeProjectFileGenerator.cs
@@ -1073,6 +1073,9 @@ namespace UnrealBuildTool
 
 		private void PopulateTargets(List<XcodeProjectTarget> ProjectTargets, List<XcodeContainerItemProxy> ContainerItemProxies, List<XcodeTargetDependency> TargetDependencies, XcodeProjectTarget ProjectTarget, List<XcodeFramework> Frameworks)
 		{
+			XcodeProjectTarget ShaderCompileWorkerTarget = null;
+			List<XcodeProjectTarget> EditorTargets = new List<XcodeProjectTarget>();
+
 			List<ProjectTarget> TargetProjects = new List<ProjectTarget> ();
 			foreach (var Project in GeneratedProjectFiles)
 			{
@@ -1092,110 +1095,132 @@ namespace UnrealBuildTool
 			}
 			foreach (var TargetProject in TargetProjects)
 			{
-					string TargetPath = TargetProject.TargetFilePath;
-					string TargetName = Utils.GetFilenameWithoutAnyExtensions (Path.GetFileName (TargetPath));
-					bool WantProjectFileForTarget = true;
-					bool IsEngineTarget = false;
-					if (bGeneratingGameProjectFiles || bGeneratingRocketProjectFiles)
+				string TargetPath = TargetProject.TargetFilePath;
+				string TargetName = Utils.GetFilenameWithoutAnyExtensions (Path.GetFileName (TargetPath));
+				bool WantProjectFileForTarget = true;
+				bool IsEngineTarget = false;
+				if (bGeneratingGameProjectFiles || bGeneratingRocketProjectFiles)
+				{
+					// Check to see if this is an Engine target.  That is, the target is located under the "Engine" folder
+					string TargetFileRelativeToEngineDirectory = Utils.MakePathRelativeTo(TargetPath, Path.Combine(EngineRelativePath), AlwaysTreatSourceAsDirectory: false);
+					if (!TargetFileRelativeToEngineDirectory.StartsWith("..") && !Path.IsPathRooted(TargetFileRelativeToEngineDirectory))
+					{
+						// This is an engine target
+						IsEngineTarget = true;
+					}
+
+					if (IsEngineTarget)
 					{
-						// Check to see if this is an Engine target.  That is, the target is located under the "Engine" folder
-						string TargetFileRelativeToEngineDirectory = Utils.MakePathRelativeTo(TargetPath, Path.Combine(EngineRelativePath), AlwaysTreatSourceAsDirectory: false);
-						if (!TargetFileRelativeToEngineDirectory.StartsWith("..") && !Path.IsPathRooted(TargetFileRelativeToEngineDirectory))
+						if (!IncludeEngineSource)
+						{
+							// We were asked to exclude engine modules from the generated projects
+							WantProjectFileForTarget = false;
+						}
+						if (bGeneratingGameProjectFiles && this.GameProjectName == TargetName)
 						{
-							// This is an engine target
-							IsEngineTarget = true;
+							WantProjectFileForTarget = true;
 						}
+					}
 
-						if (IsEngineTarget)
+					if (bGeneratingRocketProjectFiles && TargetName.EndsWith("Server"))
+					{
+						WantProjectFileForTarget = false;
+					}
+				}
+
+				if (WantProjectFileForTarget)
+				{
+					string TargetFilePath = TargetProject.TargetFilePath;
+					var Target = new TargetInfo(UnrealTargetPlatform.Mac, UnrealTargetConfiguration.Development);
+					var TargetRulesObject = TargetProject.TargetRules; // RulesCompiler.CreateTargetRules(TargetName, Target, false, out TargetFilePath);
+					List<UnrealTargetPlatform> SupportedPlatforms = new List<UnrealTargetPlatform>();
+					TargetRulesObject.GetSupportedPlatforms(ref SupportedPlatforms);
+					LinkEnvironmentConfiguration LinkConfiguration = new LinkEnvironmentConfiguration();
+					CPPEnvironmentConfiguration CPPConfiguration = new CPPEnvironmentConfiguration();
+					TargetRulesObject.SetupGlobalEnvironment(Target, ref LinkConfiguration, ref CPPConfiguration);
+
+					if (!LinkConfiguration.bIsBuildingConsoleApplication)
+					{
+						TargetsThatNeedApp.Add(TargetName);
+					}
+
+					// if the project is not an engine project check to make sure we have the correct name
+					string DisplayName = TargetName;
+					if (!IsEngineTarget && TargetRulesObject.Type != TargetRules.TargetType.Program && TargetRulesObject.Type != TargetRules.TargetType.Client)
+					{
+						List<UProjectInfo> AllGames = UProjectInfo.FilterGameProjects(true, bGeneratingGameProjectFiles ? GameProjectName : null);
+						UProjectInfo ProjectInfo = FindGameContainingFile(AllGames, TargetFilePath);
+						if (ProjectInfo != null)
 						{
-							if (!IncludeEngineSource)
+							DisplayName = ProjectInfo.GameName;
+							if (TargetName.Contains("Editor"))
 							{
-								// We were asked to exclude engine modules from the generated projects
-								WantProjectFileForTarget = false;
+								DisplayName += "Editor";
 							}
-							if (bGeneratingGameProjectFiles && this.GameProjectName == TargetName)
+							else if (TargetName.Contains("Server"))
 							{
-								WantProjectFileForTarget = true;
+								DisplayName += "Server";
 							}
 						}
+					}
 
-						if (bGeneratingRocketProjectFiles && TargetName.EndsWith("Server"))
-						{
-							WantProjectFileForTarget = false;
-						}
+					// @todo: Remove target platform param and merge Mac and iOS targets. For now BuildTarget knows how to build iOS, but cannot run iOS apps, so we need separate DeployTarget.
+					bool bIsMacOnly = !SupportedPlatforms.Contains(UnrealTargetPlatform.IOS);
+
+					XcodeProjectTarget BuildTarget = new XcodeProjectTarget(DisplayName + " - Mac", TargetName, XcodeTargetType.Legacy, TargetFilePath, "", UnrealTargetPlatform.Mac, bIsMacOnly);
+					if (!bGeneratingRunIOSProject)
+					{
+						ProjectTargets.Add(BuildTarget);
 					}
 
-					if (WantProjectFileForTarget)
+					if (DisplayName.EndsWith("Editor") && !EditorTargets.Contains(BuildTarget))
 					{
-						string TargetFilePath = TargetProject.TargetFilePath;
-						var Target = new TargetInfo(UnrealTargetPlatform.Mac, UnrealTargetConfiguration.Development);
-						var TargetRulesObject = TargetProject.TargetRules; // RulesCompiler.CreateTargetRules(TargetName, Target, false, out TargetFilePath);
-						List<UnrealTargetPlatform> SupportedPlatforms = new List<UnrealTargetPlatform>();
-						TargetRulesObject.GetSupportedPlatforms(ref SupportedPlatforms);
-						LinkEnvironmentConfiguration LinkConfiguration = new LinkEnvironmentConfiguration();
-						CPPEnvironmentConfiguration CPPConfiguration = new CPPEnvironmentConfiguration();
-						TargetRulesObject.SetupGlobalEnvironment(Target, ref LinkConfiguration, ref CPPConfiguration);
-
-						if (!LinkConfiguration.bIsBuildingConsoleApplication)
-						{
-							TargetsThatNeedApp.Add(TargetName);
-						}
+						EditorTargets.Add(BuildTarget);
+					}
+					else if (DisplayName == "ShaderCompileWorker")
+					{
+						ShaderCompileWorkerTarget = BuildTarget;
+					}
 
-						// if the project is not an engine project check to make sure we have the correct name
-						string DisplayName = TargetName;
-						if (!IsEngineTarget && TargetRulesObject.Type != TargetRules.TargetType.Program && TargetRulesObject.Type != TargetRules.TargetType.Client)
+					if (ProjectFilePlatform.HasFlag(XcodeProjectFilePlatform.iOS) && SupportedPlatforms.Contains(UnrealTargetPlatform.IOS))
+					{
+						if ((bGeneratingRocketProjectFiles && TargetName == "UE4Game") || bGeneratingRunIOSProject)
 						{
-							List<UProjectInfo> AllGames = UProjectInfo.FilterGameProjects(true, bGeneratingGameProjectFiles ? GameProjectName : null);
-							UProjectInfo ProjectInfo = FindGameContainingFile(AllGames, TargetFilePath);
-							if (ProjectInfo != null)
+							// Generate Framework references.
+							List<XcodeFrameworkRef> FrameworkRefs = new List<XcodeFrameworkRef>();
+							foreach (XcodeFramework Framework in Frameworks)
 							{
-								DisplayName = ProjectInfo.GameName;
-								if (TargetName.Contains("Editor"))
-								{
-									DisplayName += "Editor";
-								}
-								else if (TargetName.Contains("Server"))
-								{
-									DisplayName += "Server";
-								}
+								FrameworkRefs.Add(new XcodeFrameworkRef(Framework));
 							}
-						}
-
-						// @todo: Remove target platform param and merge Mac and iOS targets. For now BuildTarget knows how to build iOS, but cannot run iOS apps, so we need separate DeployTarget.
-						bool bIsMacOnly = !SupportedPlatforms.Contains(UnrealTargetPlatform.IOS);
 
-						XcodeProjectTarget BuildTarget = new XcodeProjectTarget(DisplayName + " - Mac", TargetName, XcodeTargetType.Legacy, TargetFilePath, "", UnrealTargetPlatform.Mac, bIsMacOnly);
-						if (!bGeneratingRunIOSProject)
-						{
-							ProjectTargets.Add(BuildTarget);
+							XcodeProjectTarget IOSDeployTarget = new XcodeProjectTarget(DisplayName + " - iOS", TargetName, XcodeTargetType.Native, TargetFilePath, TargetName + ".app", UnrealTargetPlatform.IOS, false, null, true, FrameworkRefs);
+							ProjectTargets.Add(IOSDeployTarget);
 						}
-
-						if (ProjectFilePlatform.HasFlag(XcodeProjectFilePlatform.iOS) && SupportedPlatforms.Contains(UnrealTargetPlatform.IOS))
+						else
 						{
-							if ((bGeneratingRocketProjectFiles && TargetName == "UE4Game") || bGeneratingRunIOSProject)
-							{
-								// Generate Framework references.
-								List<XcodeFrameworkRef> FrameworkRefs = new List<XcodeFrameworkRef>();
-								foreach (XcodeFramework Framework in Frameworks)
-								{
-									FrameworkRefs.Add(new XcodeFrameworkRef(Framework));
-								}
-
-								XcodeProjectTarget IOSDeployTarget = new XcodeProjectTarget(DisplayName + " - iOS", TargetName, XcodeTargetType.Native, TargetFilePath, TargetName + ".app", UnrealTargetPlatform.IOS, false, null, true, FrameworkRefs);
-								ProjectTargets.Add(IOSDeployTarget);
-							}
-							else
-							{
-								XcodeContainerItemProxy ContainerProxy = new XcodeContainerItemProxy(ProjectTarget.Guid, BuildTarget.Guid, BuildTarget.DisplayName);
-								XcodeTargetDependency TargetDependency = new XcodeTargetDependency(BuildTarget.DisplayName, BuildTarget.Guid, ContainerProxy.Guid);
-								XcodeProjectTarget IOSDeployTarget = new XcodeProjectTarget(DisplayName + " - iOS", TargetName, XcodeTargetType.Native, TargetFilePath, TargetName + ".app", UnrealTargetPlatform.IOS, false, new List<XcodeTargetDependency>() { TargetDependency }, true);
-								ProjectTargets.Add(IOSDeployTarget);
-								ContainerItemProxies.Add(ContainerProxy);
-								TargetDependencies.Add(TargetDependency);
-							}
+							XcodeContainerItemProxy ContainerProxy = new XcodeContainerItemProxy(ProjectTarget.Guid, BuildTarget.Guid, BuildTarget.DisplayName);
+							XcodeTargetDependency TargetDependency = new XcodeTargetDependency(BuildTarget.DisplayName, BuildTarget.Guid, ContainerProxy.Guid);
+							XcodeProjectTarget IOSDeployTarget = new XcodeProjectTarget(DisplayName + " - iOS", TargetName, XcodeTargetType.Native, TargetFilePath, TargetName + ".app", UnrealTargetPlatform.IOS, false, new List<XcodeTargetDependency>() { TargetDependency }, true);
+							ProjectTargets.Add(IOSDeployTarget);
+							ContainerItemProxies.Add(ContainerProxy);
+							TargetDependencies.Add(TargetDependency);
 						}
 					}
 				}
+			}
+
+			// Always compile ShaderCompileWorker along with the editor
+			if (ShaderCompileWorkerTarget != null)
+			{
+				XcodeContainerItemProxy ContainerProxy = new XcodeContainerItemProxy(ProjectTarget.Guid, ShaderCompileWorkerTarget.Guid, ShaderCompileWorkerTarget.DisplayName);
+				XcodeTargetDependency TargetDependency = new XcodeTargetDependency(ShaderCompileWorkerTarget.DisplayName, ShaderCompileWorkerTarget.Guid, ContainerProxy.Guid);
+				TargetDependencies.Add(TargetDependency);
+
+				foreach (var EditorTarget in EditorTargets)
+				{
+					EditorTarget.Dependencies = new List<XcodeTargetDependency>() { TargetDependency };
+				}
+			}
 		}
 
 		/// <summary>
