Make sure the version manifests and file names are serialized to the UBT makefile. If they're not written out, we won't be able to update the *.modules files correctly when switching between targets, unless we're discarding the makefile.

diff --git a/Engine/Source/Programs/UnrealBuildTool/Configuration/UEBuildTarget.cs b/Engine/Source/Programs/UnrealBuildTool/Configuration/UEBuildTarget.cs
--- a/Engine/Source/Programs/UnrealBuildTool/Configuration/UEBuildTarget.cs
+++ b/Engine/Source/Programs/UnrealBuildTool/Configuration/UEBuildTarget.cs
@@ -826,7 +826,7 @@ namespace UnrealBuildTool
 		private string ReceiptFileName;
 
 		/** Version manifests to be written to each output folder */
-		private Dictionary<string, VersionManifest> FileNameToVersionManifest;
+		private KeyValuePair<string, VersionManifest>[] FileNameToVersionManifestPairs;
 
 		/** Force output of the receipt to an additional filename */
 		[NonSerialized]
@@ -890,6 +890,7 @@ namespace UnrealBuildTool
 			}
 			Receipt                      = (TargetReceipt)Info.GetValue("re", typeof(TargetReceipt));
 			ReceiptFileName              = Info.GetString("rf");
+			FileNameToVersionManifestPairs = (KeyValuePair<string, VersionManifest>[])Info.GetValue("vm", typeof(KeyValuePair<string, VersionManifest>[]));
 			TargetCsFilenameField        = Info.GetString("tc");
 		}
 
@@ -916,6 +917,7 @@ namespace UnrealBuildTool
 			Info.AddValue("fv", FlatModuleCsData.Values.ToArray());
 			Info.AddValue("re", Receipt);
 			Info.AddValue("rf", ReceiptFileName);
+			Info.AddValue("vm", FileNameToVersionManifestPairs);
 			Info.AddValue("tc", TargetCsFilenameField);
 		}
 
@@ -1648,7 +1650,7 @@ namespace UnrealBuildTool
 		/** Prepare all the version manifests for this target. See the VersionManifest class for an explanation of what these files are. */
 		public void PrepareVersionManifests()
 		{
-			FileNameToVersionManifest = new Dictionary<string,VersionManifest>(StringComparer.InvariantCultureIgnoreCase);
+			Dictionary<string, VersionManifest> FileNameToVersionManifest = new Dictionary<string, VersionManifest>(StringComparer.InvariantCultureIgnoreCase);
 			if(!bCompileMonolithic)
 			{
 				// Read the version file
@@ -1699,16 +1701,17 @@ namespace UnrealBuildTool
 					}
 				}
 			}
+			FileNameToVersionManifestPairs = FileNameToVersionManifest.ToArray();
 		}
 
 		/** Writes out the version manifest */
 		public void WriteVersionManifests()
 		{
-			if(FileNameToVersionManifest != null)
+			if(FileNameToVersionManifestPairs != null)
 			{
-				foreach(KeyValuePair<string, VersionManifest> ManifestPair in FileNameToVersionManifest)
+				foreach(KeyValuePair<string, VersionManifest> FileNameToVersionManifest in FileNameToVersionManifestPairs)
 				{
-					ManifestPair.Value.Write(ManifestPair.Key);
+					FileNameToVersionManifest.Value.Write(FileNameToVersionManifest.Key);
 				}
 			}
 		}
