Prevent trying to build ShaderCompileWorker in non-Win64 platforms (Win64 is the only platform you can run the editor on directly).

diff --git a/Engine/Source/Programs/UnrealBuildTool/System/VCProjectFileGenerator.cs b/Engine/Source/Programs/UnrealBuildTool/System/VCProjectFileGenerator.cs
--- a/Engine/Source/Programs/UnrealBuildTool/System/VCProjectFileGenerator.cs
+++ b/Engine/Source/Programs/UnrealBuildTool/System/VCProjectFileGenerator.cs
@@ -710,9 +710,12 @@ namespace UnrealBuildTool
 										}
 
 										// Always allow SCW to build in editor configurations
-										if(MatchingProjectTarget == null && SolutionConfigCombination.TargetConfigurationName == TargetRules.TargetType.Editor.ToString() && CurProject == ShaderCompileWorkerProject)
+										if(MatchingProjectTarget == null && CurProject == ShaderCompileWorkerProject)
 										{
-											MatchingProjectTarget = CurProject.ProjectTargets[0];
+											if(SolutionConfigCombination.TargetConfigurationName == TargetRules.TargetType.Editor.ToString() && SolutionConfigCombination.Platform == UnrealTargetPlatform.Win64)
+											{
+												MatchingProjectTarget = ShaderCompileWorkerProject.ProjectTargets[0];
+											}
 										}
 
 										string ProjectConfigName;
