add log display type as part of each log message
add -verbose=Level argument to be able to specify the verbosity level of logging, defaults to Log

diff --git a/Engine/Source/Programs/AutomationTool/AutomationUtils/Log.cs b/Engine/Source/Programs/AutomationTool/AutomationUtils/Log.cs
--- a/Engine/Source/Programs/AutomationTool/AutomationUtils/Log.cs
+++ b/Engine/Source/Programs/AutomationTool/AutomationUtils/Log.cs
@@ -31,7 +31,7 @@ namespace AutomationTool
 
 			UnrealBuildTool.Log.InitLogging(
                 bLogTimestamps: CommandUtils.ParseParam(CommandLine, "-Timestamps"),
-                bLogVerbose: CommandUtils.ParseParam(CommandLine, "-Verbose"),
+				InLogLevel: (UnrealBuildTool.LogEventType)Enum.Parse(typeof(UnrealBuildTool.LogEventType), CommandUtils.ParseParamValue(CommandLine, "-Verbose=", "Log")),
                 bLogSeverity: true,
                 bLogSources: true,
                 bColorConsoleOutput: true,
diff --git a/Engine/Source/Programs/UnrealBuildTool/Configuration/BuildConfiguration.cs b/Engine/Source/Programs/UnrealBuildTool/Configuration/BuildConfiguration.cs
--- a/Engine/Source/Programs/UnrealBuildTool/Configuration/BuildConfiguration.cs
+++ b/Engine/Source/Programs/UnrealBuildTool/Configuration/BuildConfiguration.cs
@@ -244,7 +244,13 @@ namespace UnrealBuildTool
 		/// </summary>
 		[XmlConfig]
 		public static bool bOmitFramePointers;
-        
+
+		/// <summary>
+		/// What level of logging we wish to show
+		/// </summary>
+		[XmlConfig]
+		public static string LogLevel;
+
 		/// <summary>
 		/// Processor count multiplier for local execution. Can be below 1 to reserve CPU for other tasks.
 		/// </summary>
@@ -533,6 +539,8 @@ namespace UnrealBuildTool
 			bSupportEditAndContinue = false;
 			bUseActionHistory = true;
 
+			LogLevel = "Log";
+
 			// Incremental linking can yield faster iteration times when making small changes
 			// NOTE: We currently don't use incremental linking because it tends to behave a bit buggy on some computers (PDB-related compile errors)
 			bUseIncrementalLinking = false;
diff --git a/Engine/Source/Programs/UnrealBuildTool/System/UnrealBuildTool.cs b/Engine/Source/Programs/UnrealBuildTool/System/UnrealBuildTool.cs
--- a/Engine/Source/Programs/UnrealBuildTool/System/UnrealBuildTool.cs
+++ b/Engine/Source/Programs/UnrealBuildTool/System/UnrealBuildTool.cs
@@ -601,7 +601,7 @@ namespace UnrealBuildTool
         {
             Log.InitLogging(
                 bLogTimestamps: false,
-                bLogVerbose: true,
+                InLogLevel: LogEventType.Log,
                 bLogSeverity: false,
                 bLogSources: false,
                 bColorConsoleOutput: true,
@@ -859,7 +859,7 @@ namespace UnrealBuildTool
             // Initialize the shared log system using UBT specific listeners.
             Log.InitLogging(
                 bLogTimestamps: false,
-                bLogVerbose: BuildConfiguration.bPrintDebugInfo,
+                InLogLevel: (LogEventType)Enum.Parse(typeof(LogEventType), BuildConfiguration.LogLevel),
                 bLogSeverity: false, 
                 bLogSources: false,
                 bColorConsoleOutput: true,
diff --git a/Engine/Source/Programs/UnrealBuildTool/System/Utils.cs b/Engine/Source/Programs/UnrealBuildTool/System/Utils.cs
--- a/Engine/Source/Programs/UnrealBuildTool/System/Utils.cs
+++ b/Engine/Source/Programs/UnrealBuildTool/System/Utils.cs
@@ -973,7 +973,7 @@ namespace UnrealBuildTool
         /// <summary>
         /// When true, verbose loggin is enabled.
         /// </summary>
-        private static bool bLogVerbose = false;
+        private static LogEventType LogLevel = LogEventType.Log;
         /// <summary>
         /// When true, warnings and errors will have a WARNING: or ERROR: prexifx, respectively.
         /// </summary>
@@ -994,7 +994,7 @@ namespace UnrealBuildTool
         /// <summary>
         /// Expose the log level. This is a hack for ProcessResult.LogOutput, which wants to bypass our normal formatting scheme.
         /// </summary>
-        public static bool bIsVerbose { get { return bLogVerbose; } }
+        public static bool bIsVerbose { get { return LogLevel >= LogEventType.Verbose; } }
 
 		/// <summary>
 		/// A collection of strings that have been already written once
@@ -1017,9 +1017,9 @@ namespace UnrealBuildTool
         /// but then read the config and command line later, which could change this value.
         /// </summary>
         /// <param name="bLogVerbose">Whether to log verbose logs.</param>
-        public static void SetVerboseLogging(bool bLogVerbose)
+        public static void SetLoggingLevel(LogEventType InLogLevel)
         {
-            Log.bLogVerbose = bLogVerbose;
+            Log.LogLevel = InLogLevel;
         }
 
         /// <summary>
@@ -1030,11 +1030,11 @@ namespace UnrealBuildTool
         /// <param name="bLogSeverity">If true, warnings and errors will have a WARNING: and ERROR: prefix to them. </param>
         /// <param name="bLogSources">If true, logs will have the originating method name prepended to them.</param>
         /// <param name="TraceListeners">Collection of trace listeners to attach to the Trace.Listeners, in addition to the Default listener. The existing listeners (except the Default listener) are cleared first.</param>
-        public static void InitLogging(bool bLogTimestamps, bool bLogVerbose, bool bLogSeverity, bool bLogSources, bool bColorConsoleOutput, IEnumerable<TraceListener> TraceListeners)
+        public static void InitLogging(bool bLogTimestamps, LogEventType InLogLevel, bool bLogSeverity, bool bLogSources, bool bColorConsoleOutput, IEnumerable<TraceListener> TraceListeners)
         {
             bIsInitialized = true;
             Timer = (bLogTimestamps && Timer == null) ? Stopwatch.StartNew() : null;
-            Log.bLogVerbose = bLogVerbose;
+            Log.LogLevel = InLogLevel;
             Log.bLogSeverity = bLogSeverity;
             Log.bLogSources = bLogSources;
             Log.bColorConsoleOutput = bColorConsoleOutput;
@@ -1078,7 +1078,23 @@ namespace UnrealBuildTool
         /// <returns></returns>
         private static string GetSeverityPrefix(LogEventType Severity)
         {
-            return Severity <= LogEventType.Error ? "ERROR: " : Severity == LogEventType.Warning ? "WARNING: " : "";
+			switch (Severity)
+			{
+				case LogEventType.Fatal:
+					return "FATAL: ";
+				case LogEventType.Error:
+					return "ERROR: ";
+				case LogEventType.Warning:
+					return "WARNING: ";
+				case LogEventType.Console:
+					return "DISPLAY: ";
+				case LogEventType.Verbose:
+					return "VERBOSE: ";
+				case LogEventType.VeryVerbose:
+					return "VVERBOSE: ";
+				default:
+					return " ";
+			}
         }
 
 		/// <summary>
@@ -1108,8 +1124,7 @@ namespace UnrealBuildTool
                     bLogSources ? string.Format("{0}: ", string.IsNullOrEmpty(CustomSource) ? GetSource(StackFramesToSkip) : CustomSource) : "",
                     bLogSeverity ? GetSeverityPrefix(Verbosity) : "",
                     // If there are no extra args, don't try to format the string, in case it has any format control characters in it (our LOCTEXT strings tend to).
-                    Args.Length > 0 ? string.Format(Format, Args) : Format,
-					GetMessageCode(Verbosity).ToString("X3"));
+                    Args.Length > 0 ? string.Format(Format, Args) : Format);
         }
 
         /// <summary>
@@ -1141,7 +1156,7 @@ namespace UnrealBuildTool
 				WriteOnceSet.Add(Formatted);
 			}
 
-			if (Verbosity < LogEventType.Verbose || bLogVerbose)
+			if (Verbosity <= LogLevel)
             {
                 // Do console color highlighting here.
                 ConsoleColor DefaultColor = ConsoleColor.Gray;
@@ -1172,7 +1187,7 @@ namespace UnrealBuildTool
 						{
 						    foreach (TraceListener l in Trace.Listeners) 
 						    {
-								if (Verbosity != LogEventType.Log || (l as ConsoleTraceListener) == null || bLogVerbose)
+								if (Verbosity != LogEventType.Log || (l as ConsoleTraceListener) == null || LogLevel >= LogEventType.Verbose)
 								{
 									l.WriteLine(FormatMessage(StackFramesToSkip + 1, CustomSource, Verbosity, Format, Args));
 								}
