Disable targets added just to give nice error messages by default when determining what should be built by default. Fixes trying to build SCW through its dummy target on PS4.

diff --git a/Engine/Source/Programs/UnrealBuildTool/System/VCProjectFileGenerator.cs b/Engine/Source/Programs/UnrealBuildTool/System/VCProjectFileGenerator.cs
--- a/Engine/Source/Programs/UnrealBuildTool/System/VCProjectFileGenerator.cs
+++ b/Engine/Source/Programs/UnrealBuildTool/System/VCProjectFileGenerator.cs
@@ -733,16 +733,26 @@ namespace UnrealBuildTool
 										// Set whether this project configuration should be built when the user initiates "build solution"
 										if( MatchingProjectTarget != null && CurProject.ShouldBuildByDefaultForSolutionTargets )
 										{
-											VCSolutionFileContent.Append(
-													"		" + CurProjectGUID + "." + SolutionConfigCombination.VCSolutionConfigAndPlatformName + ".Build.0 = " + ProjectConfigAndPlatformPair + ProjectFileGenerator.NewLine );
-
-											var ProjGen = UEPlatformProjectGenerator.GetPlatformProjectGenerator(SolutionConfigCombination.Platform, true);
-											if (MatchingProjectTarget.ProjectDeploys ||
-												((ProjGen != null) && (ProjGen.GetVisualStudioDeploymentEnabled(SolutionPlatform, SolutionConfiguration) == true)))
-											{
-												VCSolutionFileContent.Append(
-														"		" + CurProjectGUID + "." + SolutionConfigCombination.VCSolutionConfigAndPlatformName + ".Deploy.0 = " + ProjectConfigAndPlatformPair + ProjectFileGenerator.NewLine );
-											}
+                                            // Some targets are "dummy targets"; they only exist to show user friendly errors in VS. Weed them out here, and don't set them to build by default.
+                                            List<UnrealTargetPlatform> SupportedPlatforms = null;
+                                            if (MatchingProjectTarget.TargetRules != null)
+                                            {
+                                                SupportedPlatforms = new List<UnrealTargetPlatform>();
+                                                MatchingProjectTarget.TargetRules.GetSupportedPlatforms(ref SupportedPlatforms);
+                                            }
+                                            if(SupportedPlatforms == null || SupportedPlatforms.Contains(SolutionConfigCombination.Platform))
+                                            {
+                                                VCSolutionFileContent.Append(
+                                                        "		" + CurProjectGUID + "." + SolutionConfigCombination.VCSolutionConfigAndPlatformName + ".Build.0 = " + ProjectConfigAndPlatformPair + ProjectFileGenerator.NewLine);
+
+                                                var ProjGen = UEPlatformProjectGenerator.GetPlatformProjectGenerator(SolutionConfigCombination.Platform, true);
+                                                if (MatchingProjectTarget.ProjectDeploys ||
+                                                    ((ProjGen != null) && (ProjGen.GetVisualStudioDeploymentEnabled(SolutionPlatform, SolutionConfiguration) == true)))
+                                                {
+                                                    VCSolutionFileContent.Append(
+                                                            "		" + CurProjectGUID + "." + SolutionConfigCombination.VCSolutionConfigAndPlatformName + ".Deploy.0 = " + ProjectConfigAndPlatformPair + ProjectFileGenerator.NewLine);
+                                                }
+                                            }
 										}
 
 										CombinationsThatWereMatchedToProjects.Add( SolutionConfigCombination );
