Strip debug symbols from libgnustl_shared.so for distribution builds on Android
#ue4
#android

diff --git a/Engine/Source/Programs/UnrealBuildTool/Android/AndroidToolChain.cs b/Engine/Source/Programs/UnrealBuildTool/Android/AndroidToolChain.cs
--- a/Engine/Source/Programs/UnrealBuildTool/Android/AndroidToolChain.cs
+++ b/Engine/Source/Programs/UnrealBuildTool/Android/AndroidToolChain.cs
@@ -778,6 +778,21 @@ namespace UnrealBuildTool
 		static string ArPathx86;
 		static string ArPathx64;
 
+		static public string GetStripExecutablePath(string UE4Arch)
+		{
+			string StripPath;
+
+			switch (UE4Arch)
+			{
+				case "-armv7": StripPath = ArPathArm; break;
+				case "-arm64": StripPath = ArPathArm64; break;
+				case "-x86": StripPath = ArPathx86; break;
+				case "-x64": StripPath = ArPathx64; break;
+				default: StripPath = ArPathArm; break;
+			}
+			return StripPath.Replace("-ar", "-strip");
+		}
+
 		static private bool bHasPrintedApiLevel = false;
 		public override CPPOutput CompileCPPFiles(UEBuildTarget Target, CPPEnvironment CompileEnvironment, List<FileItem> SourceFiles, string ModuleName)
 		{
diff --git a/Engine/Source/Programs/UnrealBuildTool/Android/UEDeployAndroid.cs b/Engine/Source/Programs/UnrealBuildTool/Android/UEDeployAndroid.cs
--- a/Engine/Source/Programs/UnrealBuildTool/Android/UEDeployAndroid.cs
+++ b/Engine/Source/Programs/UnrealBuildTool/Android/UEDeployAndroid.cs
@@ -525,7 +525,25 @@ namespace UnrealBuildTool.Android
 			}
 		}
 
-		private static void CopySTL(string UE4BuildPath, string UE4Arch)
+		private static void StripDebugSymbols(string SourceFileName, string TargetFileName, string UE4Arch)
+		{
+			// Copy the file and remove read-only if necessary
+			File.Copy(SourceFileName, TargetFileName, true);
+			FileAttributes Attribs = File.GetAttributes(TargetFileName);
+			if (Attribs.HasFlag(FileAttributes.ReadOnly))
+			{
+				File.SetAttributes(TargetFileName, Attribs & ~FileAttributes.ReadOnly);
+			}
+
+			ProcessStartInfo StartInfo = new ProcessStartInfo();
+			StartInfo.FileName = AndroidToolChain.GetStripExecutablePath(UE4Arch);
+			StartInfo.Arguments = "--strip-debug " + TargetFileName;
+			StartInfo.UseShellExecute = false;
+			StartInfo.CreateNoWindow = true;
+			Utils.RunLocalProcessAndLogOutput(StartInfo);
+		}
+
+		private static void CopySTL(string UE4BuildPath, string UE4Arch, bool bForDistribution)
 		{
 			string Arch = GetNDKArch(UE4Arch);
 
@@ -544,17 +562,25 @@ namespace UnrealBuildTool.Android
 			string SourceSTLSOName = Environment.ExpandEnvironmentVariables("%NDKROOT%/sources/cxx-stl/gnu-libstdc++/") + GccVersion + "/libs/" + Arch + "/libgnustl_shared.so";
 			string FinalSTLSOName = UE4BuildPath + "/libs/" + Arch + "/libgnustl_shared.so";
 
-			// check to see if libgnustl_shared.so is newer than last time we copied
+			// check to see if libgnustl_shared.so is newer than last time we copied (or needs stripping for distribution)
 			bool bFileExists = File.Exists(FinalSTLSOName);
 			TimeSpan Diff = File.GetLastWriteTimeUtc(FinalSTLSOName) - File.GetLastWriteTimeUtc(SourceSTLSOName);
-			if (!bFileExists || Diff.TotalSeconds < -1 || Diff.TotalSeconds > 1)
+			if (bForDistribution || !bFileExists || Diff.TotalSeconds < -1 || Diff.TotalSeconds > 1)
 			{
 				if (bFileExists)
 				{
 					File.Delete(FinalSTLSOName);
 				}
 				Directory.CreateDirectory(Path.GetDirectoryName(FinalSTLSOName));
-				File.Copy(SourceSTLSOName, FinalSTLSOName, true);
+				if (bForDistribution)
+				{
+					// Strip debug symbols for distribution builds
+					StripDebugSymbols(SourceSTLSOName, FinalSTLSOName, UE4Arch);
+				}
+				else
+				{
+					File.Copy(SourceSTLSOName, FinalSTLSOName, true);
+				}
 			}
 		}
 
@@ -1529,7 +1555,7 @@ namespace UnrealBuildTool.Android
 
                 // after ndk-build is called, we can now copy in the stl .so (ndk-build deletes old files)
                 // copy libgnustl_shared.so to library (use 4.8 if possible, otherwise 4.6)
-                CopySTL(UE4BuildPath, Arch);
+                CopySTL(UE4BuildPath, Arch, bForDistribution);
                 CopyGfxDebugger(UE4BuildPath, Arch);
 				CopyPluginLibs(EngineDirectory, UE4BuildPath, Arch); 
 
