Merging trivial engine fixes from ZebraBranch.
- Uses 10.11 El Capitan secure HTTP exceptions for now.
- Adds icons with correct name in the receipt.
- Fix some pipe functions.

diff --git a/Engine/Source/Programs/UnrealBuildTool/Mac/MacToolChain.cs b/Engine/Source/Programs/UnrealBuildTool/Mac/MacToolChain.cs
--- a/Engine/Source/Programs/UnrealBuildTool/Mac/MacToolChain.cs
+++ b/Engine/Source/Programs/UnrealBuildTool/Mac/MacToolChain.cs
@@ -1018,7 +1018,7 @@ namespace UnrealBuildTool
 							CustomIcon = ConvertPath(CustomIcon);
 						}
 					}
-					AppendMacLine(FinalizeAppBundleScript, "cp -f \"{0}\" \"{2}.app/Contents/Resources/{1}.icns\"", CustomIcon, IconName, ExeName);
+					AppendMacLine(FinalizeAppBundleScript, "cp -f \"{0}\" \"{2}.app/Contents/Resources/{1}.icns\"", CustomIcon, GameName, ExeName);
 
 					if (ExeName.StartsWith("UE4Editor"))
 					{
@@ -1041,7 +1041,7 @@ namespace UnrealBuildTool
 					AppendMacLine(FinalizeAppBundleScript, "sed -i \"\" \"s/\\${0}/{1}/g\" \"{1}.app/Contents/Info.plist\"", "{EXECUTABLE_NAME}", ExeName);
 					AppendMacLine(FinalizeAppBundleScript, "sed -i \"\" \"s/\\${0}/{1}/g\" \"{2}.app/Contents/Info.plist\"", "{APP_NAME}", GameName, ExeName);
 					AppendMacLine(FinalizeAppBundleScript, "sed -i \"\" \"s/\\${0}/{1}/g\" \"{2}.app/Contents/Info.plist\"", "{MACOSX_DEPLOYMENT_TARGET}", MinMacOSVersion, ExeName);
-					AppendMacLine(FinalizeAppBundleScript, "sed -i \"\" \"s/\\${0}/{1}/g\" \"{2}.app/Contents/Info.plist\"", "{ICON_NAME}", IconName, ExeName);
+					AppendMacLine(FinalizeAppBundleScript, "sed -i \"\" \"s/\\${0}/{1}/g\" \"{2}.app/Contents/Info.plist\"", "{ICON_NAME}", GameName, ExeName);
 					AppendMacLine(FinalizeAppBundleScript, "sed -i \"\" \"s/\\${0}/{1}/g\" \"{2}.app/Contents/Info.plist\"", "{BUNDLE_VERSION}", BundleVersion, ExeName);
 
 					// Generate PkgInfo file
@@ -1376,7 +1376,7 @@ namespace UnrealBuildTool
 				// And we also need all the resources
 				Receipt.AddBuildProduct(BundleContentsDirectory + "Info.plist", BuildProductType.RequiredResource);
 				Receipt.AddBuildProduct(BundleContentsDirectory + "PkgInfo", BuildProductType.RequiredResource);
-				Receipt.AddBuildProduct(BundleContentsDirectory + "Resources/UE4.icns", BuildProductType.RequiredResource);
+				Receipt.AddBuildProduct(BundleContentsDirectory + "Resources/" + Binary.Target.AppName + ".icns", BuildProductType.RequiredResource);
 
 				if (Binary.Target.AppName.StartsWith("UE4Editor"))
 				{
diff --git a/Engine/Source/Runtime/Core/Private/Mac/MacPlatformProcess.cpp b/Engine/Source/Runtime/Core/Private/Mac/MacPlatformProcess.cpp
--- a/Engine/Source/Runtime/Core/Private/Mac/MacPlatformProcess.cpp
+++ b/Engine/Source/Runtime/Core/Private/Mac/MacPlatformProcess.cpp
@@ -937,12 +937,15 @@ FString FMacPlatformProcess::ReadPipe( void* ReadPipe )
 
 	if(ReadPipe)
 	{
-		BytesRead = read([(NSFileHandle*)ReadPipe fileDescriptor], Buffer, READ_SIZE);
-		if (BytesRead > 0)
+		do
 		{
-			Buffer[BytesRead] = '\0';
-			Output += StringCast<TCHAR>(Buffer).Get();
-		}
+			BytesRead = read([(NSFileHandle*)ReadPipe fileDescriptor], Buffer, READ_SIZE);
+			if (BytesRead > 0)
+			{
+				Buffer[BytesRead] = '\0';
+				Output += StringCast<TCHAR>(Buffer).Get();
+			}
+		} while (BytesRead > 0);
 	}
 
 	return Output;
@@ -995,7 +998,7 @@ bool FMacPlatformProcess::WritePipe(void* WritePipe, const FString& Message, FSt
 	}
 
 	// Write to pipe
-	uint32 BytesWritten = write(*(int*)WritePipe, Buffer, BytesAvailable);
+	uint32 BytesWritten = write([(NSFileHandle*)WritePipe fileDescriptor], Buffer, BytesAvailable);
 
 	if (OutWritten != nullptr)
 	{
