Fixed Mac icon name in UBT manifest

diff --git a/Engine/Source/Programs/UnrealBuildTool/Mac/MacToolChain.cs b/Engine/Source/Programs/UnrealBuildTool/Mac/MacToolChain.cs
--- a/Engine/Source/Programs/UnrealBuildTool/Mac/MacToolChain.cs
+++ b/Engine/Source/Programs/UnrealBuildTool/Mac/MacToolChain.cs
@@ -1054,7 +1054,7 @@ namespace UnrealBuildTool
 
 					// copy over some needed files
 					// @todo mac: Make a QueueDirectoryForBatchUpload
-					QueueFileForBatchUpload(FileItem.GetItemByFullPath(Path.GetFullPath("../../Engine/Source/Runtime/Launch/Resources/Mac/UE4.icns")));
+					QueueFileForBatchUpload(FileItem.GetItemByFullPath(Path.GetFullPath("../../Engine/Source/Runtime/Launch/Resources/Mac/" + GameName + ".icns")));
 					QueueFileForBatchUpload(FileItem.GetItemByFullPath(Path.GetFullPath("../../Engine/Source/Runtime/Launch/Resources/Mac/UProject.icns")));
 					QueueFileForBatchUpload(FileItem.GetItemByFullPath(Path.GetFullPath("../../Engine/Source/Runtime/Launch/Resources/Mac/Info.plist")));
 					QueueFileForBatchUpload(FileItem.GetItemByFullPath(Path.Combine(LinkEnvironment.Config.IntermediateDirectory, "DylibCopy.sh")));
@@ -1376,9 +1376,9 @@ namespace UnrealBuildTool
 				// And we also need all the resources
 				Receipt.AddBuildProduct(BundleContentsDirectory + "Info.plist", BuildProductType.RequiredResource);
 				Receipt.AddBuildProduct(BundleContentsDirectory + "PkgInfo", BuildProductType.RequiredResource);
-				Receipt.AddBuildProduct(BundleContentsDirectory + "Resources/" + Binary.Target.AppName + ".icns", BuildProductType.RequiredResource);
+				Receipt.AddBuildProduct(BundleContentsDirectory + "Resources/" + Binary.Target.TargetName + ".icns", BuildProductType.RequiredResource);
 
-				if (Binary.Target.AppName.StartsWith("UE4Editor"))
+				if (Binary.Target.TargetName.StartsWith("UE4Editor"))
 				{
 					Receipt.AddBuildProduct(BundleContentsDirectory + "Resources/UProject.icns", BuildProductType.RequiredResource);
 				}
