UE4 - New vk env var

diff --git a/Engine/Source/Programs/UnrealBuildTool/Windows/UEBuildWindows.cs b/Engine/Source/Programs/UnrealBuildTool/Windows/UEBuildWindows.cs
--- a/Engine/Source/Programs/UnrealBuildTool/Windows/UEBuildWindows.cs
+++ b/Engine/Source/Programs/UnrealBuildTool/Windows/UEBuildWindows.cs
@@ -508,7 +508,7 @@ namespace UnrealBuildTool
 
 //#todo-rco: Remove when public
 					{
-						string VulkanSDKPath = Environment.GetEnvironmentVariable("VulkanSDK");
+						string VulkanSDKPath = Environment.GetEnvironmentVariable("VK_SDK_PATH");
 						if (!String.IsNullOrEmpty(VulkanSDKPath))
 						{
 							InModule.AddDynamicallyLoadedModule("VulkanShaderFormat");
diff --git a/Engine/Source/Runtime/RHI/RHI.Build.cs b/Engine/Source/Runtime/RHI/RHI.Build.cs
--- a/Engine/Source/Runtime/RHI/RHI.Build.cs
+++ b/Engine/Source/Runtime/RHI/RHI.Build.cs
@@ -22,7 +22,7 @@ public class RHI : ModuleRules
 					DynamicallyLoadedModuleNames.Add("D3D12RHI");
 
 //#todo-rco: Remove when public
-					string VulkanSDKPath = Environment.GetEnvironmentVariable("VulkanSDK");
+					string VulkanSDKPath = Environment.GetEnvironmentVariable("VK_SDK_PATH");
 					if (!String.IsNullOrEmpty(VulkanSDKPath))
 					{
 						DynamicallyLoadedModuleNames.Add("VulkanRHI");
