Fixed Mac editor icon name in UBT manifest

diff --git a/Engine/Source/Programs/UnrealBuildTool/Mac/MacToolChain.cs b/Engine/Source/Programs/UnrealBuildTool/Mac/MacToolChain.cs
--- a/Engine/Source/Programs/UnrealBuildTool/Mac/MacToolChain.cs
+++ b/Engine/Source/Programs/UnrealBuildTool/Mac/MacToolChain.cs
@@ -1376,12 +1376,16 @@ namespace UnrealBuildTool
 				// And we also need all the resources
 				Receipt.AddBuildProduct(BundleContentsDirectory + "Info.plist", BuildProductType.RequiredResource);
 				Receipt.AddBuildProduct(BundleContentsDirectory + "PkgInfo", BuildProductType.RequiredResource);
-				Receipt.AddBuildProduct(BundleContentsDirectory + "Resources/" + Binary.Target.TargetName + ".icns", BuildProductType.RequiredResource);
 
-				if (Binary.Target.TargetName.StartsWith("UE4Editor"))
+				if (Binary.Target.TargetType == TargetRules.TargetType.Editor)
 				{
+					Receipt.AddBuildProduct(BundleContentsDirectory + "Resources/UE4Editor.icns", BuildProductType.RequiredResource);
 					Receipt.AddBuildProduct(BundleContentsDirectory + "Resources/UProject.icns", BuildProductType.RequiredResource);
 				}
+				else
+				{
+					Receipt.AddBuildProduct(BundleContentsDirectory + "Resources/" + Binary.Target.TargetName + ".icns", BuildProductType.RequiredResource);
+				}
 			}
 		}
 
