Patch newer NDK (r10) cpu-features.c to deal with libhoudini.so SIGSEGV on older versions of Houdini if patch not already applied
#ue4
#android

diff --git a/Engine/Source/Programs/UnrealBuildTool/Android/AndroidToolChain.cs b/Engine/Source/Programs/UnrealBuildTool/Android/AndroidToolChain.cs
--- a/Engine/Source/Programs/UnrealBuildTool/Android/AndroidToolChain.cs
+++ b/Engine/Source/Programs/UnrealBuildTool/Android/AndroidToolChain.cs
@@ -707,7 +707,84 @@ namespace UnrealBuildTool
 			if (ModuleName.Equals("Launch"))
 			{
 				SourceFiles.Add(FileItem.GetItemByPath(Environment.GetEnvironmentVariable("NDKROOT") + "/sources/android/native_app_glue/android_native_app_glue.c"));
-				SourceFiles.Add(FileItem.GetItemByPath(Environment.GetEnvironmentVariable("NDKROOT") + "/sources/android/cpufeatures/cpu-features.c"));
+
+				// Newer NDK cpu_features.c uses getauxval() which causes a SIGSEGV in libhoudini.so (ARM on Intel translator) in older versions of Houdini
+				// so we patch the file to use alternative methods of detecting CPU features if libhoudini.so is detected
+				// The basis for this patch is from here: https://android-review.googlesource.com/#/c/110650/
+				String CpuFeaturesPath = Environment.GetEnvironmentVariable("NDKROOT") + "/sources/android/cpufeatures/";
+				String CpuFeaturesPatchedFile = CpuFeaturesPath + "cpu-features-patched.c";
+				if (!File.Exists(CpuFeaturesPatchedFile))
+				{
+					// Either make a copy or patch it
+					String[] CpuFeaturesLines = File.ReadAllLines(CpuFeaturesPath + "cpu-features.c");
+
+					// Look for get_elf_hwcap_from_getauxval in the file
+					bool NeedsPatch = false;
+					int LineIndex;
+					for (LineIndex = 0; LineIndex < CpuFeaturesLines.Length; ++LineIndex)
+					{
+						if (CpuFeaturesLines[LineIndex].Contains("get_elf_hwcap_from_getauxval"))
+						{
+							NeedsPatch = true;
+
+							// Make sure it doesn't already have the patch (r10c and 10d have it already, but removed in 10e)
+							for (int LineIndex2 = LineIndex; LineIndex2 < CpuFeaturesLines.Length; ++LineIndex2)
+							{
+								if (CpuFeaturesLines[LineIndex2].Contains("has_houdini_binary_translator(void)"))
+								{
+									NeedsPatch = false;
+									break;
+								}
+							}
+							break;
+						}
+					}
+
+					// Apply patch or write unchanged
+					if (NeedsPatch)
+					{
+						List<string> CpuFeaturesList = new List<string>(CpuFeaturesLines);
+
+						// Skip down to section to add Houdini check function for arm
+						while (!CpuFeaturesList[++LineIndex].StartsWith("#if defined(__arm__)")) ;
+						CpuFeaturesList.Insert(++LineIndex, "/* Check Houdini Binary Translator is installed on the system.");
+						CpuFeaturesList.Insert(++LineIndex, " *");
+						CpuFeaturesList.Insert(++LineIndex, " * If this function returns 1, get_elf_hwcap_from_getauxval() function");
+						CpuFeaturesList.Insert(++LineIndex, " * will causes SIGSEGV while calling getauxval() function.");
+						CpuFeaturesList.Insert(++LineIndex, " */");
+						CpuFeaturesList.Insert(++LineIndex, "static int");
+						CpuFeaturesList.Insert(++LineIndex, "has_houdini_binary_translator(void) {");
+						CpuFeaturesList.Insert(++LineIndex, "    int found = 0;");
+						CpuFeaturesList.Insert(++LineIndex, "    if (access(\"/system/lib/libhoudini.so\", F_OK) != -1) {");
+						CpuFeaturesList.Insert(++LineIndex, "        D(\"Found Houdini binary translator\\n\");");
+						CpuFeaturesList.Insert(++LineIndex, "        found = 1;");
+						CpuFeaturesList.Insert(++LineIndex, "    }");
+						CpuFeaturesList.Insert(++LineIndex, "    return found;");
+						CpuFeaturesList.Insert(++LineIndex, "}");
+						CpuFeaturesList.Insert(++LineIndex, "");
+
+						// Add the Houdini check call
+						while (!CpuFeaturesList[++LineIndex].Contains("/* Extract the list of CPU features from ELF hwcaps */")) ;
+						CpuFeaturesList.Insert(LineIndex++, "        /* Check Houdini binary translator is installed */");
+						CpuFeaturesList.Insert(LineIndex++, "        int has_houdini = has_houdini_binary_translator();");
+						CpuFeaturesList.Insert(LineIndex++, "");
+
+						// Make the get_elf_hwcap_from_getauxval() calls conditional
+						while (!CpuFeaturesList[++LineIndex].Contains("hwcaps = get_elf_hwcap_from_getauxval(AT_HWCAP);")) ;
+						CpuFeaturesList.Insert(LineIndex++, "        if (!has_houdini) {");
+						CpuFeaturesList.Insert(++LineIndex, "        }");
+						while (!CpuFeaturesList[++LineIndex].Contains("hwcaps2 = get_elf_hwcap_from_getauxval(AT_HWCAP2);")) ;
+						CpuFeaturesList.Insert(LineIndex++, "        if (!has_houdini) {");
+						CpuFeaturesList.Insert(++LineIndex, "        }");
+
+						File.WriteAllLines(CpuFeaturesPatchedFile, CpuFeaturesList.ToArray());
+					}
+					else
+					{
+						File.WriteAllLines(CpuFeaturesPatchedFile, CpuFeaturesLines);
+					}
+				}
+				SourceFiles.Add(FileItem.GetItemByPath(CpuFeaturesPatchedFile));
 			}
 		}
 
