Handle issue with spaces in path when stripping debug symbols from libgnustl_shared.so distribution builds
#ue4
#android

diff --git a/Engine/Source/Programs/UnrealBuildTool/Android/UEDeployAndroid.cs b/Engine/Source/Programs/UnrealBuildTool/Android/UEDeployAndroid.cs
--- a/Engine/Source/Programs/UnrealBuildTool/Android/UEDeployAndroid.cs
+++ b/Engine/Source/Programs/UnrealBuildTool/Android/UEDeployAndroid.cs
@@ -537,7 +537,7 @@ namespace UnrealBuildTool.Android
 
 			ProcessStartInfo StartInfo = new ProcessStartInfo();
 			StartInfo.FileName = AndroidToolChain.GetStripExecutablePath(UE4Arch);
-			StartInfo.Arguments = "--strip-debug " + TargetFileName;
+			StartInfo.Arguments = "--strip-debug \"" + TargetFileName + "\"";
 			StartInfo.UseShellExecute = false;
 			StartInfo.CreateNoWindow = true;
 			Utils.RunLocalProcessAndLogOutput(StartInfo);
