Fix a crash saving plugin descriptors that have null entries.

diff --git a/Engine/Source/Programs/UnrealBuildTool/Utilities/JsonObject.cs b/Engine/Source/Programs/UnrealBuildTool/Utilities/JsonObject.cs
--- a/Engine/Source/Programs/UnrealBuildTool/Utilities/JsonObject.cs
+++ b/Engine/Source/Programs/UnrealBuildTool/Utilities/JsonObject.cs
@@ -328,7 +328,7 @@ namespace UnrealBuildTool
 
 		public void WriteValue(string Name, string Value)
 		{
-			WriteValueInternal(Name, '"' + Value.Replace("\\", "\\\\") + '"');
+			WriteValueInternal(Name, '"' + (Value ?? "").Replace("\\", "\\\\") + '"');
 		}
 
 		public void WriteValue(string Name, int Value)
