Handle cases where the foreign plugin passed in to UBT is already in the known plugins list. It's not technically foreign, but it's easily handled.

diff --git a/Engine/Source/Programs/UnrealBuildTool/Configuration/UEBuildTarget.cs b/Engine/Source/Programs/UnrealBuildTool/Configuration/UEBuildTarget.cs
--- a/Engine/Source/Programs/UnrealBuildTool/Configuration/UEBuildTarget.cs
+++ b/Engine/Source/Programs/UnrealBuildTool/Configuration/UEBuildTarget.cs
@@ -2968,9 +2968,17 @@ namespace UnrealBuildTool
 			{
 				foreach(string ForeignPlugin in ForeignPlugins)
 				{
-					PluginInfo ForeignPluginInfo = new PluginInfo(ForeignPlugin, PluginLoadedFrom.GameProject);
-					ValidPlugins.Add(ForeignPluginInfo);
-					BuildPlugins.Add(ForeignPluginInfo);
+					PluginInfo ForeignPluginInfo = ValidPlugins.FirstOrDefault(x => String.Compare(x.FileName, ForeignPlugin, true) == 0);
+					if(ForeignPluginInfo == null)
+					{
+						ForeignPluginInfo = new PluginInfo(ForeignPlugin, PluginLoadedFrom.GameProject);
+						ValidPlugins.Add(ForeignPluginInfo);
+						BuildPlugins.Add(ForeignPluginInfo);
+					}
+					else if(!BuildPlugins.Contains(ForeignPluginInfo))
+					{
+						BuildPlugins.Add(ForeignPluginInfo);
+					}
 				}
 			}
 		}
