Remove unused variables from TargetReceipt.

diff --git a/Engine/Source/Programs/UnrealBuildTool/System/TargetReceipt.cs b/Engine/Source/Programs/UnrealBuildTool/System/TargetReceipt.cs
--- a/Engine/Source/Programs/UnrealBuildTool/System/TargetReceipt.cs
+++ b/Engine/Source/Programs/UnrealBuildTool/System/TargetReceipt.cs
@@ -165,9 +165,6 @@ namespace UnrealBuildTool
 		[XmlArrayItem("RuntimeDependency")]
 		public List<RuntimeDependency> RuntimeDependencies = new List<RuntimeDependency>();
 
-		// if packaging in a mode where some files aren't required, set this to false
-		public bool bRequireDependenciesToExist = true;
-
 		/// <summary>
 		/// Default constructor
 		/// </summary>
@@ -180,7 +177,6 @@ namespace UnrealBuildTool
 			Properties                  = (List<BuildProperty>)Info.GetValue("pr", typeof(List<BuildProperty>));
 			BuildProducts               = (List<BuildProduct>)Info.GetValue("bp", typeof(List<BuildProduct>));
 			RuntimeDependencies         = (List<RuntimeDependency>)Info.GetValue("rd", typeof(List<RuntimeDependency>));
-			bRequireDependenciesToExist = Info.GetBoolean("rq");
 		}
 
 		public void GetObjectData(SerializationInfo Info, StreamingContext Context)
@@ -188,7 +184,6 @@ namespace UnrealBuildTool
 			Info.AddValue("pr", Properties);
 			Info.AddValue("bp", BuildProducts);
 			Info.AddValue("rd", RuntimeDependencies);
-			Info.AddValue("rq", bRequireDependenciesToExist);
 		}
 
 		/// <summary>
@@ -298,15 +293,6 @@ namespace UnrealBuildTool
 			ExpandPathVariables(EngineDir, ProjectDir, new Dictionary<string, string>());
 		}
 
-		/// <summary>
-		/// Control whether the dependencies are required while staging them
-		/// </summary>
-		/// <param name="InDependenciesAreRequired"></param>
-		public void SetDependenciesToBeRequired(bool InDependenciesAreRequired)
-		{
-			bRequireDependenciesToExist = InDependenciesAreRequired;
-		}
-
 		/// <summary>
 		/// Expand all the path variables in the manifest, including a list of supplied variable values.
 		/// </summary>
