At least on El Capitan debug symbols can't always be found when using the .dsym file extension only the .dSYM extension will work.

diff --git a/Engine/Source/Programs/UnrealBuildTool/Mac/MacToolChain.cs b/Engine/Source/Programs/UnrealBuildTool/Mac/MacToolChain.cs
--- a/Engine/Source/Programs/UnrealBuildTool/Mac/MacToolChain.cs
+++ b/Engine/Source/Programs/UnrealBuildTool/Mac/MacToolChain.cs
@@ -1301,7 +1301,7 @@ namespace UnrealBuildTool
 			// The cross-platform code adds .dSYMs for static libraries, which is just wrong, so
 			// eliminate them here for now.
 			string DebugExtension = UEBuildPlatform.GetBuildPlatform(Binary.Target.Platform).GetDebugInfoExtension(Binary.Config.Type);
-			if(DebugExtension == ".dsym")
+			if(DebugExtension == ".dSYM")
 			{
 				for (int i = 0; i < Receipt.BuildProducts.Count; i++)
 				{
diff --git a/Engine/Source/Programs/UnrealBuildTool/Mac/UEBuildMac.cs b/Engine/Source/Programs/UnrealBuildTool/Mac/UEBuildMac.cs
--- a/Engine/Source/Programs/UnrealBuildTool/Mac/UEBuildMac.cs
+++ b/Engine/Source/Programs/UnrealBuildTool/Mac/UEBuildMac.cs
@@ -89,7 +89,7 @@ namespace UnrealBuildTool
 		 */
 		public override string GetDebugInfoExtension(UEBuildBinaryType InBinaryType)
 		{
-			return BuildConfiguration.bGeneratedSYMFile || BuildConfiguration.bUsePDBFiles ? ".dsym" : "";
+			return BuildConfiguration.bGeneratedSYMFile || BuildConfiguration.bUsePDBFiles ? ".dSYM" : "";
 		}
 
 		public override void ModifyNewlyLoadedModule(UEBuildModule InModule, TargetInfo Target)
