PR-1125 - fix folder for 3rd party IOS frameworks
#ios

diff --git a/Engine/Source/Programs/UnrealBuildTool/IOS/IOSToolChain.cs b/Engine/Source/Programs/UnrealBuildTool/IOS/IOSToolChain.cs
--- a/Engine/Source/Programs/UnrealBuildTool/IOS/IOSToolChain.cs
+++ b/Engine/Source/Programs/UnrealBuildTool/IOS/IOSToolChain.cs
@@ -1053,12 +1053,10 @@ namespace UnrealBuildTool
 							continue;		// This framework item doesn't belong to this target, skip it
 						}
 
-						string LocalZipPath = GetLocalFrameworkZipPath( Framework );
-
-						LocalZipPath = LocalZipPath.Replace( ".zip", "" );
+						string UnpackedZipPath = GetRemoteIntermediateFrameworkZipPath( Framework );
 
 						// For now, this is hard coded, but we need to loop over all modules, and copy bundled assets that need it
-						string LocalSource	= LocalZipPath + "/" + Framework.CopyBundledAssets;
+						string LocalSource	= UnpackedZipPath + "/" + Framework.CopyBundledAssets;
 						string BundleName	= Framework.CopyBundledAssets.Substring( Framework.CopyBundledAssets.LastIndexOf( '/' ) + 1 );
 						string LocalDest	= LocalFrameworkAssets + "/" + BundleName;
 
