Fix hot reload issues on Mac. Don't write BOM to response files.

diff --git a/Engine/Source/Programs/UnrealBuildTool/System/FileItem.cs b/Engine/Source/Programs/UnrealBuildTool/System/FileItem.cs
--- a/Engine/Source/Programs/UnrealBuildTool/System/FileItem.cs
+++ b/Engine/Source/Programs/UnrealBuildTool/System/FileItem.cs
@@ -272,7 +272,8 @@ namespace UnrealBuildTool
 		private static Encoding GetEncodingForString(string Str)
 		{
 			// If the string length is equivalent to the encoded length, then no non-ASCII characters were present in the string.
-			return (Encoding.UTF8.GetByteCount(Str) == Str.Length) ? Encoding.ASCII : Encoding.UTF8;
+			// Don't write BOM as it messes with clang when loading response files.
+			return (Encoding.UTF8.GetByteCount(Str) == Str.Length) ? Encoding.ASCII : new UTF8Encoding(false);
 		}
 
 		/**
diff --git a/Engine/Source/Programs/UnrealBuildTool/System/UnrealBuildTool.cs b/Engine/Source/Programs/UnrealBuildTool/System/UnrealBuildTool.cs
--- a/Engine/Source/Programs/UnrealBuildTool/System/UnrealBuildTool.cs
+++ b/Engine/Source/Programs/UnrealBuildTool/System/UnrealBuildTool.cs
@@ -2834,7 +2834,7 @@ namespace UnrealBuildTool
 
 					// Overwrite the original file
 					var FileContentsString = FileContents.ToString();
-					File.WriteAllText( ResponseFilePath, FileContentsString, System.Text.Encoding.UTF8 );
+					File.WriteAllText(ResponseFilePath, FileContentsString, new System.Text.UTF8Encoding(false));
 				}
 			}
 		}
