PR-1282 - added support for specifying additional linker flags in iOS project settings
based on changes from derekvanvliet
#ios

diff --git a/Engine/Source/Programs/UnrealBuildTool/IOS/IOSToolChain.cs b/Engine/Source/Programs/UnrealBuildTool/IOS/IOSToolChain.cs
--- a/Engine/Source/Programs/UnrealBuildTool/IOS/IOSToolChain.cs
+++ b/Engine/Source/Programs/UnrealBuildTool/IOS/IOSToolChain.cs
@@ -203,6 +203,8 @@ namespace UnrealBuildTool
 
 			Result += " -miphoneos-version-min=" + BuildPlat.GetRunTimeVersion();
 
+			Result += " " + BuildPlat.GetAdditionalLinkerFlags(CompileEnvironment.Config.Target.Configuration);
+
 			// Optimize non- debug builds.
 			if (CompileEnvironment.Config.Target.Configuration != CPPTargetConfiguration.Debug)
 			{
diff --git a/Engine/Source/Programs/UnrealBuildTool/IOS/UEBuildIOS.cs b/Engine/Source/Programs/UnrealBuildTool/IOS/UEBuildIOS.cs
--- a/Engine/Source/Programs/UnrealBuildTool/IOS/UEBuildIOS.cs
+++ b/Engine/Source/Programs/UnrealBuildTool/IOS/UEBuildIOS.cs
@@ -27,6 +27,12 @@ namespace UnrealBuildTool
 		[XmlConfig]
 		public static string ShippingArchitectures = "armv7,arm64";
 
+		/** additional linker flags for shipping */
+		public static string AdditionalShippingLinkerFlags = "";
+
+		/** additional linker flags for non-shipping */
+		public static string AdditionalLinkerFlags = "";
+
 		private bool bInitializedProject = false;
 
 		public string GetRunTimeVersion()
@@ -39,6 +45,18 @@ namespace UnrealBuildTool
 			return RunTimeIOSDevices;
 		}
 
+		public string GetAdditionalLinkerFlags(CPPTargetConfiguration InConfiguration)
+		{
+			if (InConfiguration != CPPTargetConfiguration.Shipping)
+			{
+				return AdditionalLinkerFlags;
+			}
+			else
+			{
+				return AdditionalShippingLinkerFlags;
+			}
+
+		}
 		// The current architecture - affects everything about how UBT operates on IOS
         public override string GetActiveArchitecture()
         {
@@ -239,6 +257,9 @@ namespace UnrealBuildTool
 				// determine if we need to generate the dsym
 				Ini.GetBool("/Script/IOSRuntimeSettings.IOSRuntimeSettings", "bGeneratedSYMFile", out BuildConfiguration.bGeneratedSYMFile);
 
+				Ini.GetString("/Script/IOSRuntimeSettings.IOSRuntimeSettings", "AdditionalLinkerFlags", out AdditionalLinkerFlags);
+				Ini.GetString("/Script/IOSRuntimeSettings.IOSRuntimeSettings", "AdditionalShippingLinkerFlags", out AdditionalShippingLinkerFlags);
+
 				bInitializedProject = true;
 			}
 		}
@@ -256,6 +277,8 @@ namespace UnrealBuildTool
 			};
 			string[] StringKeys = new string[] {
 				"MinimumiOSVersion", 
+				"AdditionalLinkerFlags",
+				"AdditionalShippingLinkerFlags"
 			};
 
 			// look up iOS specific settings
diff --git a/Engine/Source/Runtime/IOS/IOSRuntimeSettings/Classes/IOSRuntimeSettings.h b/Engine/Source/Runtime/IOS/IOSRuntimeSettings/Classes/IOSRuntimeSettings.h
--- a/Engine/Source/Runtime/IOS/IOSRuntimeSettings/Classes/IOSRuntimeSettings.h
+++ b/Engine/Source/Runtime/IOS/IOSRuntimeSettings/Classes/IOSRuntimeSettings.h
@@ -194,6 +194,14 @@ public:
 	UPROPERTY(GlobalConfig, EditAnywhere, Category = Build, meta = (DisplayName = "Support armv7s in Shipping"))
 	bool bShipForArmV7S;
 	
+	// Any additional linker flags to pass to the linker in non-shipping builds
+	UPROPERTY(GlobalConfig, EditAnywhere, Category = Build, meta = (DispalyName = "Additional Non-Shipping Linker Flags", ConfigHierarchyEditable))
+	FString AdditionalLinkerFlags;
+
+	// Any additional linker flags to pass to the linker in shipping builds
+	UPROPERTY(GlobalConfig, EditAnywhere, Category = Build, meta = (DispalyName = "Additional Shipping Linker Flags", ConfigHierarchyEditable))
+	FString AdditionalShippingLinkerFlags;
+
 	// The name or ip address of the remote mac which will be used to build IOS
 	UPROPERTY(GlobalConfig, EditAnywhere, Category = "Build", meta = (ConfigHierarchyEditable))
 	FString RemoteServerName;
diff --git a/Engine/Source/Runtime/IOS/IOSRuntimeSettings/Private/IOSRuntimeSettings.cpp b/Engine/Source/Runtime/IOS/IOSRuntimeSettings/Private/IOSRuntimeSettings.cpp
--- a/Engine/Source/Runtime/IOS/IOSRuntimeSettings/Private/IOSRuntimeSettings.cpp
+++ b/Engine/Source/Runtime/IOS/IOSRuntimeSettings/Private/IOSRuntimeSettings.cpp
@@ -26,6 +26,8 @@ UIOSRuntimeSettings::UIOSRuntimeSettings(const FObjectInitializer& ObjectInitial
 	bShipForArmV7S = false;
 	bUseRSync = true;
 	AdditionalPlistData = TEXT("");
+	AdditionalLinkerFlags = TEXT("");
+	AdditionalShippingLinkerFlags = TEXT("");
 }
 
 #if WITH_EDITOR
