Passing -verbose to UBT now actually sets the verbose log level and not just bPrintDebugInfo.

diff --git a/Engine/Source/Programs/UnrealBuildTool/System/UnrealBuildTool.cs b/Engine/Source/Programs/UnrealBuildTool/System/UnrealBuildTool.cs
--- a/Engine/Source/Programs/UnrealBuildTool/System/UnrealBuildTool.cs
+++ b/Engine/Source/Programs/UnrealBuildTool/System/UnrealBuildTool.cs
@@ -677,6 +677,7 @@ namespace UnrealBuildTool
                 if (LowercaseArg == "-verbose")
                 {
                     BuildConfiguration.bPrintDebugInfo = true;
+                    BuildConfiguration.LogLevel = "Verbose";
                 }
                 else if (LowercaseArg.StartsWith("-log="))
                 {
