CleanDirectory now ignores DirectoryNotFoundExceptions to harden against outside sources deleting that folder before it gets to it. I ran into this when manually deleting a subset to Intermediates to iterate on the UBT bug.

diff --git a/Engine/Source/Programs/UnrealBuildTool/Configuration/UEBuildTarget.cs b/Engine/Source/Programs/UnrealBuildTool/Configuration/UEBuildTarget.cs
--- a/Engine/Source/Programs/UnrealBuildTool/Configuration/UEBuildTarget.cs
+++ b/Engine/Source/Programs/UnrealBuildTool/Configuration/UEBuildTarget.cs
@@ -1091,23 +1091,28 @@ namespace UnrealBuildTool
 			bool bResult = false;
 			do
 			{
-				try
-				{
-					Directory.Delete(DirectoryPath, true);
-					bResult = true;
-				}
-				catch (Exception Ex)
-				{
-					// This happens mostly because some other stale process is still locking this file
-					Log.TraceVerbose(Ex.Message);
-					if (--RetryCount < 0)
-					{
-						throw Ex;
-					}
-					System.Threading.Thread.Sleep(RetryDelay);
-					// Try with a slightly longer delay next time
-					RetryDelay += RetryDelayStep;
-				}
+                try
+                {
+                    Directory.Delete(DirectoryPath, true);
+                    bResult = true;
+                }
+                catch (DirectoryNotFoundException)
+                {
+                    // this is ok, someone else may have killed it for us.
+                    bResult = true;
+                }
+                catch (Exception Ex)
+                {
+                    // This happens mostly because some other stale process is still locking this file
+                    Log.TraceVerbose(Ex.Message);
+                    if (--RetryCount < 0)
+                    {
+                        throw Ex;
+                    }
+                    System.Threading.Thread.Sleep(RetryDelay);
+                    // Try with a slightly longer delay next time
+                    RetryDelay += RetryDelayStep;
+                }
 			}
 			while (!bResult);
 		}
