[INTEGRATE] Unreal Build Tool: Make Unity builds more deterministic
- Before creating unity blobs, we now sort all full file paths alphabetically
- This allows for better build system determinism across computers
- Important:  It's possible this checkin can introduce compile errors for projects that have outstanding non-unity issues!

diff --git a/Engine/Source/Programs/UnrealBuildTool/System/Unity.cs b/Engine/Source/Programs/UnrealBuildTool/System/Unity.cs
--- a/Engine/Source/Programs/UnrealBuildTool/System/Unity.cs
+++ b/Engine/Source/Programs/UnrealBuildTool/System/Unity.cs
@@ -131,8 +131,24 @@ namespace UnrealBuildTool
 			// Build the list of unity files.
 			List<FileCollection> AllUnityFiles;
 			{
+				// Sort the incoming file paths alphabetically, so there will be consistency in unity blobs across multiple machines.
+				// Note that we're relying on this not only sorting files within each directory, but also the directories
+				// themselves, so the whole list of file paths is the same across computers.
+				var SortedCPPFiles = CPPFiles.GetRange( 0, CPPFiles.Count );
+				{
+					// Case-insensitive file path compare, because you never know what is going on with local file systems
+					Comparison<FileItem> FileItemComparer = (FileA, FileB) => { return FileA.AbsolutePath.ToLowerInvariant().CompareTo( FileB.AbsolutePath.ToLowerInvariant() ); };
+					SortedCPPFiles.Sort( FileItemComparer );
+				}
+				Console.WriteLine("==== SORTED ====");
+				foreach( var File in SortedCPPFiles )
+				{
+					Console.WriteLine(File.AbsolutePath);
+				}
+				Console.WriteLine("==== SORTED ====");
+
 				var CPPUnityFileBuilder = new UnityFileBuilder(bForceIntoSingleUnityFile ? -1 : BuildConfiguration.NumIncludedBytesPerUnityCPP);
-				foreach( var CPPFile in CPPFiles )
+				foreach( var CPPFile in SortedCPPFiles )
 				{
 					if (!bForceIntoSingleUnityFile && CPPFile.AbsolutePath.Contains(".GeneratedWrapper."))
 					{
diff --git a/Engine/Source/Programs/UnrealBuildTool/System/VCProject.cs b/Engine/Source/Programs/UnrealBuildTool/System/VCProject.cs
--- a/Engine/Source/Programs/UnrealBuildTool/System/VCProject.cs
+++ b/Engine/Source/Programs/UnrealBuildTool/System/VCProject.cs
@@ -309,17 +309,6 @@ namespace UnrealBuildTool
 
 			bool bSuccess = true;
 
-			// NOTE: We intentionally do not SORT the source file list, as we want the order they're written to disk to be consistent
-			//       with how they are stored in memory.  This makes for more consistent Unity compiles when alternating between
-			//       using "auto" projects and on-disk projects for builds.
-			var ShouldSortSourceFiles = false;
-			if( ShouldSortSourceFiles )
-			{
-				// Source our list of source files
-				Comparison<SourceFile> SourceFileComparer = ( FileA, FileB ) => { return FileA.FilePath.CompareTo( FileB.FilePath ); };
-				SourceFiles.Sort( SourceFileComparer );
-			}
-
 			// Build up the new include search path string
 			var VCIncludeSearchPaths = new StringBuilder();
 			{
