[INTEGRATE] Fixed debug console spew enabled for Unity builds

diff --git a/Engine/Source/Programs/UnrealBuildTool/System/Unity.cs b/Engine/Source/Programs/UnrealBuildTool/System/Unity.cs
--- a/Engine/Source/Programs/UnrealBuildTool/System/Unity.cs
+++ b/Engine/Source/Programs/UnrealBuildTool/System/Unity.cs
@@ -140,12 +140,6 @@ namespace UnrealBuildTool
 					Comparison<FileItem> FileItemComparer = (FileA, FileB) => { return FileA.AbsolutePath.ToLowerInvariant().CompareTo( FileB.AbsolutePath.ToLowerInvariant() ); };
 					SortedCPPFiles.Sort( FileItemComparer );
 				}
-				Console.WriteLine("==== SORTED ====");
-				foreach( var File in SortedCPPFiles )
-				{
-					Console.WriteLine(File.AbsolutePath);
-				}
-				Console.WriteLine("==== SORTED ====");
 
 				var CPPUnityFileBuilder = new UnityFileBuilder(bForceIntoSingleUnityFile ? -1 : BuildConfiguration.NumIncludedBytesPerUnityCPP);
 				foreach( var CPPFile in SortedCPPFiles )
