UBT: Exposed the ability to disable /largeaddressaware for Win32 compiles
[INTEGRATE] Change 2668376 by [REDACTED_PERSON][REDACTED_HANDLE] on 2015/08/25 16:56:36

diff --git a/Engine/Source/Programs/UnrealBuildTool/Windows/UEBuildWindows.cs b/Engine/Source/Programs/UnrealBuildTool/Windows/UEBuildWindows.cs
--- a/Engine/Source/Programs/UnrealBuildTool/Windows/UEBuildWindows.cs
+++ b/Engine/Source/Programs/UnrealBuildTool/Windows/UEBuildWindows.cs
@@ -108,6 +108,9 @@ namespace UnrealBuildTool
 			return SupportWindowsXP;
 		}
 
+		/// True if we allow using addresses larger than 2GB on 32 bit builds
+		public static bool bBuildLargeAddressAwareBinary = true;
+
         /** True if VS EnvDTE is available (false when building using Visual Studio Express) */
         public static bool bHasVisualStudioDTE
         {
diff --git a/Engine/Source/Programs/UnrealBuildTool/Windows/VCToolChain.cs b/Engine/Source/Programs/UnrealBuildTool/Windows/VCToolChain.cs
--- a/Engine/Source/Programs/UnrealBuildTool/Windows/VCToolChain.cs
+++ b/Engine/Source/Programs/UnrealBuildTool/Windows/VCToolChain.cs
@@ -828,7 +828,7 @@ namespace UnrealBuildTool
 				Arguments.Append(" /FIXED:No");
 
 				// Option is only relevant with 32 bit toolchain.
-				if (LinkEnvironment.Config.Target.Platform == CPPTargetPlatform.Win32)
+				if ((LinkEnvironment.Config.Target.Platform == CPPTargetPlatform.Win32) && WindowsPlatform.bBuildLargeAddressAwareBinary)
 				{
 					// Disables the 2GB address space limit on 64-bit Windows and 32-bit Windows with /3GB specified in boot.ini
 					Arguments.Append(" /LARGEADDRESSAWARE");
