Convert target receipts to JSON, so they can be read early at startup like version manifests, and unify the way that target receipts and version manifests are loaded and saved.

diff --git a/Engine/Source/Programs/AutomationTool/AutomationUtils/UE4Build.cs b/Engine/Source/Programs/AutomationTool/AutomationUtils/UE4Build.cs
--- a/Engine/Source/Programs/AutomationTool/AutomationUtils/UE4Build.cs
+++ b/Engine/Source/Programs/AutomationTool/AutomationUtils/UE4Build.cs
@@ -97,7 +97,7 @@ namespace AutomationTool
 
 		static bool IsBuildReceipt(string FileName)
 		{
-			return Path.GetDirectoryName(FileName).Replace('\\', '/').EndsWith("/Receipts", StringComparison.InvariantCultureIgnoreCase);
+			return FileName.EndsWith(".target", StringComparison.InvariantCultureIgnoreCase);
 		}
 
 		void AddBuildProductsFromManifest(string ManifestName)
diff --git a/Engine/Source/Programs/AutomationTool/Mac/MacPlatform.Automation.cs b/Engine/Source/Programs/AutomationTool/Mac/MacPlatform.Automation.cs
--- a/Engine/Source/Programs/AutomationTool/Mac/MacPlatform.Automation.cs
+++ b/Engine/Source/Programs/AutomationTool/Mac/MacPlatform.Automation.cs
@@ -195,7 +195,7 @@ public class MacPlatform : Platform
 						}
 						else if(Params.IsCodeBasedProject)
 						{
-							BootstrapExeName = Target.Receipt.GetProperty("TargetName", SC.ShortProjectName) + ".app";
+							BootstrapExeName = Target.Receipt.TargetName + ".app";
 						}
 						else
 						{
diff --git a/Engine/Source/Programs/AutomationTool/Win/WinPlatform.Automation.cs b/Engine/Source/Programs/AutomationTool/Win/WinPlatform.Automation.cs
--- a/Engine/Source/Programs/AutomationTool/Win/WinPlatform.Automation.cs
+++ b/Engine/Source/Programs/AutomationTool/Win/WinPlatform.Automation.cs
@@ -68,7 +68,7 @@ public abstract class BaseWinPlatform : Platform
 						}
 						else if(Params.IsCodeBasedProject)
 						{
-							BootstrapExeName = Target.Receipt.GetProperty("TargetName", SC.ShortProjectName) + ".exe";
+							BootstrapExeName = Target.Receipt.TargetName + ".exe";
 						}
 						else
 						{
diff --git a/Engine/Source/Programs/UnrealBuildTool/Configuration/UEBuildTarget.cs b/Engine/Source/Programs/UnrealBuildTool/Configuration/UEBuildTarget.cs
--- a/Engine/Source/Programs/UnrealBuildTool/Configuration/UEBuildTarget.cs
+++ b/Engine/Source/Programs/UnrealBuildTool/Configuration/UEBuildTarget.cs
@@ -1652,24 +1652,38 @@ namespace UnrealBuildTool
 			}
 		}
 
-		/** Prepare all the version manifests for this target. See the VersionManifest class for an explanation of what these files are. */
-		public void PrepareVersionManifests()
+		/** Prepare all the receipts this target (all the .target and .modules files). See the VersionManifest class for an explanation of what these files are. */
+		public void PrepareReceipts(IUEToolChain ToolChain)
 		{
-			Dictionary<string, VersionManifest> FileNameToVersionManifest = new Dictionary<string, VersionManifest>(StringComparer.InvariantCultureIgnoreCase);
-			if(!bCompileMonolithic)
+			// Read the version file
+			BuildVersion Version;
+			if(!BuildVersion.TryRead(Path.Combine(BuildConfiguration.RelativeEnginePath, "Build", "Build.version"), out Version))
+			{
+				Version = new BuildVersion();
+			}
+
+			// Create a unique identifier for this build, which can be used to identify modules when the changelist is constant. It's fine to share this between runs with the same makefile; 
+			// the output won't change. By default we leave it blank when compiling a subset of modules (for hot reload, etc...), otherwise it won't match anything else. When writing to a directory
+			// that already contains a manifest, we'll reuse the build id that's already in there (see below).
+			string BuildId = (OnlyModules.Count == 0)? Guid.NewGuid().ToString() : "";
+
+			// Create the receipt, and merge all the receipts from the individual binaries
+			Receipt = new TargetReceipt(TargetName, Platform, Configuration, BuildId, Version);
+			foreach(UEBuildBinary Binary in AppBinaries)
 			{
-				// Read the version file
-				BuildVersion Version;
-				if(!Utils.TryReadBuildVersion(Path.Combine(BuildConfiguration.RelativeEnginePath, "Build", "Build.version"), out Version))
+				TargetReceipt BinaryReceipt = Binary.MakeReceipt(ToolChain);
+				if(Binary.Config.Type == UEBuildBinaryType.StaticLibrary)
 				{
-					Version = new BuildVersion();
+					BinaryReceipt.RuntimeDependencies.Clear();
 				}
+				Receipt.Merge(BinaryReceipt);
+			}
+			Receipt.InsertStandardPathVariables(BuildConfiguration.RelativeEnginePath, ProjectDirectory);
 
-				// Create a unique identifier for this build, which can be used to identify modules when the changelist is constant. It's fine to share this between runs with the same makefile; 
-				// the output won't change. By default we leave it blank when compiling a subset of modules (for hot reload, etc...), otherwise it won't match anything else. When writing to a directory
-				// that already contains a manifest, we'll reuse the build id that's already in there (see below).
-				string BuildId = (OnlyModules.Count == 0)? Guid.NewGuid().ToString() : "";
-
+			// Prepare all the version manifests
+			Dictionary<string, VersionManifest> FileNameToVersionManifest = new Dictionary<string, VersionManifest>(StringComparer.InvariantCultureIgnoreCase);
+			if(!bCompileMonolithic)
+			{
 				// Create the receipts for each folder
 				foreach(UEBuildBinary Binary in AppBinaries)
 				{
@@ -1710,45 +1724,7 @@ namespace UnrealBuildTool
 		}
 
 		/** Writes out the version manifest */
-		public void WriteVersionManifests()
-		{
-			if(FileNameToVersionManifestPairs != null)
-			{
-				foreach(KeyValuePair<string, VersionManifest> FileNameToVersionManifest in FileNameToVersionManifestPairs)
-				{
-					FileNameToVersionManifest.Value.Write(FileNameToVersionManifest.Key);
-				}
-			}
-		}
-
-		/** Creates the receipt for the target */
-		private void PrepareReceipt(IUEToolChain ToolChain)
-		{
-			Receipt = new TargetReceipt();
-
-			// Add the target properties
-			Receipt.SetProperty("TargetName", TargetName);
-			Receipt.SetProperty("Platform", Platform.ToString());
-			Receipt.SetProperty("Configuration", Configuration.ToString());
-			Receipt.SetProperty("Precompile", bPrecompile.ToString());
-
-			// Merge all the binary receipts into this
-			foreach(UEBuildBinary Binary in AppBinaries)
-			{
-				TargetReceipt BinaryReceipt = Binary.MakeReceipt(ToolChain);
-				if(Binary.Config.Type == UEBuildBinaryType.StaticLibrary)
-				{
-					BinaryReceipt.RuntimeDependencies.Clear();
-				}
-				Receipt.Merge(BinaryReceipt);
-			}
-
-			// Convert all the paths to use variables for the engine and project root
-			Receipt.InsertStandardPathVariables(BuildConfiguration.RelativeEnginePath, ProjectDirectory);
-		}
-
-		/** Writes the receipt for this target to disk */
-		public void WriteReceipt()
+		public void WriteReceipts()
 		{
 			if(Receipt != null)
 			{
@@ -1764,6 +1740,13 @@ namespace UnrealBuildTool
 					Receipt.Write(ForceReceiptFileName);
 				}
 			}
+			if(FileNameToVersionManifestPairs != null)
+			{
+				foreach(KeyValuePair<string, VersionManifest> FileNameToVersionManifest in FileNameToVersionManifestPairs)
+				{
+					FileNameToVersionManifest.Value.Write(FileNameToVersionManifest.Key);
+				}
+			}
 		}
 
 		/// <summary>
@@ -1868,11 +1851,8 @@ namespace UnrealBuildTool
 				return ECompilationResult.Succeeded;
 			}
 
-			// Create all the version manifests
-			PrepareVersionManifests();
-
 			// Create a receipt for the target
-			PrepareReceipt(TargetToolChain);
+			PrepareReceipts(TargetToolChain);
 
 			// If we're only generating the manifest, return now
 			if (UEBuildConfiguration.bGenerateManifest || UEBuildConfiguration.bCleanProject)
diff --git a/Engine/Source/Programs/UnrealBuildTool/Mac/MacToolChain.cs b/Engine/Source/Programs/UnrealBuildTool/Mac/MacToolChain.cs
--- a/Engine/Source/Programs/UnrealBuildTool/Mac/MacToolChain.cs
+++ b/Engine/Source/Programs/UnrealBuildTool/Mac/MacToolChain.cs
@@ -425,7 +425,7 @@ namespace UnrealBuildTool
 		private int LoadEngineCL()
 		{
 			BuildVersion Version;
-			if(Utils.TryReadBuildVersion("../Build/Build.version", out Version))
+			if(BuildVersion.TryRead("../Build/Build.version", out Version))
 			{
 				return Version.Changelist;
 			}
diff --git a/Engine/Source/Programs/UnrealBuildTool/System/TargetReceipt.cs b/Engine/Source/Programs/UnrealBuildTool/System/TargetReceipt.cs
--- a/Engine/Source/Programs/UnrealBuildTool/System/TargetReceipt.cs
+++ b/Engine/Source/Programs/UnrealBuildTool/System/TargetReceipt.cs
@@ -11,38 +11,6 @@ using System.Runtime.Serialization;
 
 namespace UnrealBuildTool
 {
-	[Serializable]
-	public class BuildProperty : ISerializable
-	{
-		[XmlAttribute]
-		public string Name;
-
-		[XmlAttribute]
-		public string Value;
-
-		private BuildProperty()
-		{
-		}
-
-		public BuildProperty(SerializationInfo Info, StreamingContext Context)
-		{
-			Name  = Info.GetString("na");
-			Value = Info.GetString("va");
-		}
-
-		public void GetObjectData(SerializationInfo Info, StreamingContext Context)
-		{
-			Info.AddValue("na", Name);
-			Info.AddValue("va", Value);
-		}
-
-		public BuildProperty(string InName, string InValue)
-		{
-			Name = InName;
-			Value = InValue;
-		}
-	}
-
 	public enum BuildProductType
 	{
 		Executable,
@@ -54,35 +22,16 @@ namespace UnrealBuildTool
 	}
 
 	[Serializable]
-	public class BuildProduct : ISerializable
+	public class BuildProduct
 	{
-		[XmlAttribute]
 		public string Path;
-
-		[XmlAttribute]
 		public BuildProductType Type;
-
-		[XmlAttribute, DefaultValue(false)]
 		public bool IsPrecompiled;
 
 		private BuildProduct()
 		{
 		}
 
-		public BuildProduct(SerializationInfo Info, StreamingContext Context)
-		{
-			Path          = Info.GetString("pa");
-			Type          = (BuildProductType)Info.GetInt32("ty");
-			IsPrecompiled = Info.GetBoolean("ip");
-		}
-
-		public void GetObjectData(SerializationInfo Info, StreamingContext Context)
-		{
-			Info.AddValue("pa", Path);
-			Info.AddValue("ty", (int)Type);
-			Info.AddValue("ip", IsPrecompiled);
-		}
-
 		public BuildProduct(string InPath, BuildProductType InType)
 		{
 			Path = InPath;
@@ -103,30 +52,22 @@ namespace UnrealBuildTool
 	}
 
 	[Serializable]
-	public class RuntimeDependency : ISerializable
+	public class RuntimeDependency
 	{
-		[XmlAttribute]
+		/// <summary>
+		/// The file that should be staged. Should use $(EngineDir) and $(ProjectDir) variables as a root, so that the target can be relocated to different machines.
+		/// </summary>
 		public string Path;
 
-		[XmlAttribute, DefaultValue(null)]
+		/// <summary>
+		/// The path that the file should be staged to, if different. Leave as null if the file should be staged to the same relative path.
+		/// </summary>
 		public string StagePath;
 
 		private RuntimeDependency()
 		{
 		}
 
-		public RuntimeDependency(SerializationInfo Info, StreamingContext Context)
-		{
-			Path      = Info.GetString("pa");
-			StagePath = Info.GetString("sp");
-		}
-
-		public void GetObjectData(SerializationInfo Info, StreamingContext Context)
-		{
-			Info.AddValue("pa", Path);
-			Info.AddValue("sp", StagePath);
-		}
-
 		public RuntimeDependency(string InPath)
 		{
 			Path = InPath;
@@ -154,15 +95,41 @@ namespace UnrealBuildTool
 	/// Stores a record of a built target, with all metadata that other tools may need to know about the build.
 	/// </summary>
 	[Serializable]
-	public class TargetReceipt : ISerializable
+	public class TargetReceipt
 	{
-		[XmlArrayItem("Property")]
-		public List<BuildProperty> Properties = new List<BuildProperty>();
+		/// <summary>
+		/// The name of this target
+		/// </summary>
+		public string TargetName;
+
+		/// <summary>
+		/// Which platform the target is compiled for
+		/// </summary>
+		public UnrealTargetPlatform Platform;
+
+		/// <summary>
+		/// Which configuration this target is compiled in
+		/// </summary>
+		public UnrealTargetConfiguration Configuration;
 
-		[XmlArrayItem("BuildProduct")]
+		/// <summary>
+		/// The unique ID for this build.
+		/// </summary>
+		public string BuildId;
+
+		/// <summary>
+		/// The changelist that this target was compiled with.
+		/// </summary>
+		public BuildVersion Version;
+
+		/// <summary>
+		/// The build products which are part of this target
+		/// </summary>
 		public List<BuildProduct> BuildProducts = new List<BuildProduct>();
 
-		[XmlArrayItem("RuntimeDependency")]
+		/// <summary>
+		/// All the runtime dependencies that this target relies on
+		/// </summary>
 		public List<RuntimeDependency> RuntimeDependencies = new List<RuntimeDependency>();
 
 		/// <summary>
@@ -172,18 +139,19 @@ namespace UnrealBuildTool
 		{
 		}
 
-		public TargetReceipt(SerializationInfo Info, StreamingContext Context)
-		{
-			Properties                  = (List<BuildProperty>)Info.GetValue("pr", typeof(List<BuildProperty>));
-			BuildProducts               = (List<BuildProduct>)Info.GetValue("bp", typeof(List<BuildProduct>));
-			RuntimeDependencies         = (List<RuntimeDependency>)Info.GetValue("rd", typeof(List<RuntimeDependency>));
-		}
-
-		public void GetObjectData(SerializationInfo Info, StreamingContext Context)
+		/// <summary>
+		/// Constructor
+		/// </summary>
+		/// <param name="InTargetName">The name of the target being compiled</param>
+		/// <param name="InPlatform">Platform for the target being compiled</param>
+		/// <param name="InConfiguration">Configuration of the target being compiled</param>
+		public TargetReceipt(string InTargetName, UnrealTargetPlatform InPlatform, UnrealTargetConfiguration InConfiguration, string InBuildId, BuildVersion InVersion)
 		{
-			Info.AddValue("pr", Properties);
-			Info.AddValue("bp", BuildProducts);
-			Info.AddValue("rd", RuntimeDependencies);
+			TargetName = InTargetName;
+			Platform = InPlatform;
+			Configuration = InConfiguration;
+			BuildId = InBuildId;
+			Version = InVersion;
 		}
 
 		/// <summary>
@@ -202,42 +170,6 @@ namespace UnrealBuildTool
 			}
 		}
 
-		/// <summary>
-		/// Sets a property with the given name
-		/// </summary>
-		/// <param name="Name">Name of the property; case sensitive.</param>
-		/// <param name="Value">Value for the property</param>
-		public void SetProperty(string Name, string Value)
-		{
-			BuildProperty Property = Properties.FirstOrDefault(x => x.Name == Name);
-			if(Property == null)
-			{
-				Properties.Add(new BuildProperty(Name, Value));
-			}
-			else
-			{
-				Property.Value = Value;
-			}
-		}
-
-		/// <summary>
-		/// Gets the value associated with a property
-		/// </summary>
-		/// <param name="Name">Name of the property; case sensitive.</param>
-		/// <param name="DefaultValue">Default value for the property if it's not found</param>
-		public string GetProperty(string Name, string DefaultValue)
-		{
-			BuildProperty Property = Properties.FirstOrDefault(x => x.Name == Name);
-			if(Property == null)
-			{
-				return DefaultValue;
-			}
-			else
-			{
-				return Property.Value;
-			}
-		}
-
 		/// <summary>
 		/// Adds a build product to the receipt. Does not check whether it already exists.
 		/// </summary>
@@ -379,18 +311,93 @@ namespace UnrealBuildTool
 		/// <returns>Path to the receipt for this target</returns>
 		public static string GetDefaultPath(string BaseDir, string TargetName, UnrealTargetPlatform Platform, UnrealTargetConfiguration Configuration, string BuildArchitecture)
 		{
-			return Path.Combine(BaseDir, "Build", "Receipts", String.Format("{0}-{1}-{2}{3}.target.xml", TargetName, Platform.ToString(), Configuration.ToString(), BuildArchitecture));
+			if(String.IsNullOrEmpty(BuildArchitecture) && Configuration == UnrealTargetConfiguration.Development)
+			{
+				return Path.Combine(BaseDir, "Binaries", Platform.ToString(), String.Format("{0}.target", TargetName));
+			}
+			else
+			{
+				return Path.Combine(BaseDir, "Binaries", Platform.ToString(), String.Format("{0}-{1}-{2}{3}.target", TargetName, Platform.ToString(), Configuration.ToString(), BuildArchitecture));
+			}
 		}
 
-		static XmlSerializer Serializer = XmlSerializer.FromTypes(new Type[]{ typeof(TargetReceipt) })[0];
-
 		/// <summary>
 		/// Read a receipt from disk.
 		/// </summary>
 		/// <param name="FileName">Filename to read from</param>
+		public static TargetReceipt Read(string FileName)
+		{
+			JsonObject RawObject = JsonObject.Read(FileName);
+
+			// Read the initial fields
+			string TargetName = RawObject.GetStringField("TargetName");
+			UnrealTargetPlatform Platform = RawObject.GetEnumField<UnrealTargetPlatform>("Platform");
+			UnrealTargetConfiguration Configuration = RawObject.GetEnumField<UnrealTargetConfiguration>("Configuration");
+			string BuildId = RawObject.GetStringField("BuildId");
+
+			// Try to read the build version
+			BuildVersion Version;
+			if(!BuildVersion.TryParse(RawObject.GetObjectField("Version"), out Version))
+			{
+				throw new JsonParseException("Invalid 'Version' field");
+			}
+
+			// Create the receipt
+			TargetReceipt Receipt = new TargetReceipt(TargetName, Platform, Configuration, BuildId, Version);
+
+			// Read the build products
+			JsonObject[] BuildProductObjects;
+			if(RawObject.TryGetObjectArrayField("BuildProducts", out BuildProductObjects))
+			{
+				foreach(JsonObject BuildProductObject in BuildProductObjects)
+				{
+					string Path;
+					BuildProductType Type;
+					if(BuildProductObject.TryGetStringField("Path", out Path) && BuildProductObject.TryGetEnumField("Type", out Type))
+					{
+						string Module;
+						BuildProductObject.TryGetStringField("Module", out Module);
+
+						BuildProduct NewBuildProduct = Receipt.AddBuildProduct(Path, Type);
+
+						bool IsPrecompiled;
+						if(BuildProductObject.TryGetBoolField("IsPrecompiled", out IsPrecompiled))
+						{
+							NewBuildProduct.IsPrecompiled = IsPrecompiled;
+						}
+					}
+				}
+			}
+
+			// Read the runtime dependencies
+			JsonObject[] RuntimeDependencyObjects;
+			if(RawObject.TryGetObjectArrayField("RuntimeDependencies", out RuntimeDependencyObjects))
+			{
+				foreach(JsonObject RuntimeDependencyObject in RuntimeDependencyObjects)
+				{
+					string Path;
+					if(RuntimeDependencyObject.TryGetStringField("Path", out Path))
+					{
+						string StagePath;
+						if(!RuntimeDependencyObject.TryGetStringField("StagePath", out StagePath))
+						{
+							StagePath = null;
+						}
+						Receipt.AddRuntimeDependency(Path, StagePath);
+					}
+				}
+			}
+
+			return Receipt;
+		}
+
+		/// <summary>
+		/// Try to read a receipt from disk, failing gracefully if it can't be read.
+		/// </summary>
+		/// <param name="FileName">Filename to read from</param>
 		public static bool TryRead(string FileName, out TargetReceipt Receipt)
 		{
-			if (!File.Exists(FileName))
+			if(File.Exists(FileName))
 			{
 				Receipt = null;
 				return false;
@@ -398,12 +405,9 @@ namespace UnrealBuildTool
 
 			try
 			{
-				using (StreamReader Reader = new StreamReader(FileName))
-				{
-					Receipt = (TargetReceipt)Serializer.Deserialize(Reader);
+				Receipt = Read(FileName);
 					return true;
 				}
-			}
 			catch(Exception)
 			{
 				Receipt = null;
@@ -417,9 +421,46 @@ namespace UnrealBuildTool
 		/// <param name="FileName">Output filename</param>
 		public void Write(string FileName)
 		{
-			using(StreamWriter Writer = new StreamWriter(FileName))
+			using(JsonWriter Writer = new JsonWriter(FileName))
 			{
-				Serializer.Serialize(Writer, this);
+				Writer.WriteObjectStart();
+				Writer.WriteValue("TargetName", TargetName);
+				Writer.WriteValue("Platform", Platform.ToString());
+				Writer.WriteValue("Configuration", Configuration.ToString());
+				Writer.WriteValue("BuildId", BuildId);
+
+				Writer.WriteObjectStart("Version");
+				Version.Write(Writer);
+				Writer.WriteObjectEnd();
+
+				Writer.WriteArrayStart("BuildProducts");
+				foreach(BuildProduct BuildProduct in BuildProducts)
+				{
+					Writer.WriteObjectStart();
+					Writer.WriteValue("Path", BuildProduct.Path);
+					Writer.WriteValue("Type", BuildProduct.Type.ToString());
+					if(BuildProduct.IsPrecompiled)
+					{
+						Writer.WriteValue("IsPrecompiled", BuildProduct.IsPrecompiled);
+					}
+					Writer.WriteObjectEnd();
+				}
+				Writer.WriteArrayEnd();
+
+				Writer.WriteArrayStart("RuntimeDependencies");
+				foreach(RuntimeDependency RuntimeDependency in RuntimeDependencies)
+				{
+					Writer.WriteObjectStart();
+					Writer.WriteValue("Path", RuntimeDependency.Path);
+					if(RuntimeDependency.StagePath != null)
+					{
+						Writer.WriteValue("StagePath", RuntimeDependency.StagePath);
+					}
+					Writer.WriteObjectEnd();
+				}
+				Writer.WriteArrayEnd();
+
+				Writer.WriteObjectEnd();
 			}
 		}
 	}
diff --git a/Engine/Source/Programs/UnrealBuildTool/System/UnrealBuildTool.cs b/Engine/Source/Programs/UnrealBuildTool/System/UnrealBuildTool.cs
--- a/Engine/Source/Programs/UnrealBuildTool/System/UnrealBuildTool.cs
+++ b/Engine/Source/Programs/UnrealBuildTool/System/UnrealBuildTool.cs
@@ -1179,7 +1179,7 @@ namespace UnrealBuildTool
 
 					// Get the build version
 					BuildVersion Version;
-					if(!Utils.TryReadBuildVersion("../Build/Build.version", out Version))
+					if(!BuildVersion.TryRead("../Build/Build.version", out Version))
 					{
 						Log.TraceWarning("Could not read build.version");
 					}
@@ -1989,8 +1989,7 @@ namespace UnrealBuildTool
                             {
                                 foreach (UEBuildTarget Target in Targets)
                                 {
-									Target.WriteVersionManifests();
-									Target.WriteReceipt();
+									Target.WriteReceipts();
                                     ToolChain.PostBuildSync(Target);
                                 }
 								if (ActionsToExecute.Count == 0 && UEBuildConfiguration.bSkipLinkingWhenNothingToCompile)
diff --git a/Engine/Source/Programs/UnrealBuildTool/System/Utils.cs b/Engine/Source/Programs/UnrealBuildTool/System/Utils.cs
--- a/Engine/Source/Programs/UnrealBuildTool/System/Utils.cs
+++ b/Engine/Source/Programs/UnrealBuildTool/System/Utils.cs
@@ -20,6 +20,7 @@ namespace UnrealBuildTool
 	/// <summary>
 	/// Holds information about the current engine version
 	/// </summary>
+	[Serializable]
 	public class BuildVersion
 	{
 		public int MajorVersion;
@@ -28,6 +29,61 @@ namespace UnrealBuildTool
 		public int Changelist;
 		public int IsLicenseeVersion;
 		public string BranchName;
+
+		/// <summary>
+		/// Try to read the Build/Build.version file from disk
+		/// </summary>
+		/// <param name="Version">The version information</param>
+		/// <returns>True if the version was read sucessfully, false otherwise</returns>
+		public static bool TryRead(string FileName, out BuildVersion Version)
+		{
+			JsonObject Object;
+			if(!JsonObject.TryRead(FileName, out Object))
+			{
+				Version = null;
+				return false;
+			}
+			return TryParse(Object, out Version);
+		}
+
+		/// <summary>
+		/// Parses a build version from a JsonObject
+		/// </summary>
+		/// <param name="Object">The object to read from</param>
+		/// <param name="Version">The resulting version field</param>
+		/// <returns>True if the build version could be read, false otherwise</returns>
+		public static bool TryParse(JsonObject Object, out BuildVersion Version)
+		{
+			BuildVersion NewVersion = new BuildVersion();
+			if(!Object.TryGetIntegerField("MajorVersion", out NewVersion.MajorVersion) || !Object.TryGetIntegerField("MinorVersion", out NewVersion.MinorVersion) || !Object.TryGetIntegerField("PatchVersion", out NewVersion.PatchVersion))
+			{
+				Version = null;
+				return false;
+			}
+
+			Object.TryGetIntegerField("Changelist", out NewVersion.Changelist);
+			Object.TryGetIntegerField("IsLicenseeVersion", out NewVersion.IsLicenseeVersion);
+			Object.TryGetStringField("BranchName", out NewVersion.BranchName);
+
+			Version = NewVersion;
+			return true;
+		}
+
+		/// <summary>
+		/// Exports this object as Json
+		/// </summary>
+		/// <param name="Object">The object to read from</param>
+		/// <param name="Version">The resulting version field</param>
+		/// <returns>True if the build version could be read, false otherwise</returns>
+		public void Write(JsonWriter Writer)
+		{
+			Writer.WriteValue("MajorVersion", MajorVersion);
+			Writer.WriteValue("MinorVersion", MinorVersion);
+			Writer.WriteValue("PatchVersion", PatchVersion);
+			Writer.WriteValue("Changelist", Changelist);
+			Writer.WriteValue("IsLicenseeVersion", IsLicenseeVersion);
+			Writer.WriteValue("BranchName", BranchName);
+		}
 	}
 
 	/// 
@@ -821,35 +877,6 @@ namespace UnrealBuildTool
 			return true;
 		}
 
-		/// <summary>
-		/// Try to read the Build/Build.version file from disk
-		/// </summary>
-		/// <param name="Version">The version information</param>
-		/// <returns>True if the version was read sucessfully, false otherwise</returns>
-		public static bool TryReadBuildVersion(string FileName, out BuildVersion Version)
-		{
-			JsonObject Object;
-			if(!JsonObject.TryRead(FileName, out Object))
-			{
-				Version = null;
-				return false;
-			}
-
-			BuildVersion NewVersion = new BuildVersion();
-			if(!Object.TryGetIntegerField("MajorVersion", out NewVersion.MajorVersion) || !Object.TryGetIntegerField("MinorVersion", out NewVersion.MinorVersion) || !Object.TryGetIntegerField("PatchVersion", out NewVersion.PatchVersion))
-			{
-				Version = null;
-				return false;
-			}
-
-			Object.TryGetIntegerField("Changelist", out NewVersion.Changelist);
-			Object.TryGetIntegerField("IsLicenseeVersion", out NewVersion.IsLicenseeVersion);
-			Object.TryGetStringField("BranchName", out NewVersion.BranchName);
-
-			Version = NewVersion;
-			return true;
-		}
-
 		/// <summary>
 		/// Checks if given type implements given interface.
 		/// </summary>
