Linux: do not compile ICU for server targets.

- Saves memory/code size, and servers shouldn't need that (by default at least).

(Merging CL 2665433 to main)

diff --git a/Engine/Source/Programs/UnrealBuildTool/Linux/UEBuildLinux.cs b/Engine/Source/Programs/UnrealBuildTool/Linux/UEBuildLinux.cs
--- a/Engine/Source/Programs/UnrealBuildTool/Linux/UEBuildLinux.cs
+++ b/Engine/Source/Programs/UnrealBuildTool/Linux/UEBuildLinux.cs
@@ -363,6 +363,12 @@ namespace UnrealBuildTool
             {
                 UEBuildConfiguration.bCompileSimplygon = false;
             }
+
+			if (InBuildTarget.TargetType == TargetRules.TargetType.Server)
+			{
+				// Localization shouldn't be needed on servers by default, and ICU is pretty heavy
+				UEBuildConfiguration.bCompileICU = false;
+			}
         }
 
         /**
