Fix ability to compile foreign plugins against a target (required by BuildPlugin command in UAT).

diff --git a/Engine/Source/Programs/UnrealBuildTool/System/RulesCompiler.cs b/Engine/Source/Programs/UnrealBuildTool/System/RulesCompiler.cs
--- a/Engine/Source/Programs/UnrealBuildTool/System/RulesCompiler.cs
+++ b/Engine/Source/Programs/UnrealBuildTool/System/RulesCompiler.cs
@@ -2025,7 +2025,7 @@ namespace UnrealBuildTool
 		{
 			var CreateTargetStartTime = DateTime.UtcNow;
 
-			RulesCompiler.SetAssemblyNameAndGameFolders(Desc.AssemblyName, Desc.GameFolders);
+			RulesCompiler.SetAssemblyNameAndGameFolders(Desc.AssemblyName, Desc.GameFolders, Desc.ForeignPlugins);
 
 			string TargetFileName;
 			TargetRules RulesObject = CreateTargetRules(Desc.TargetName, new TargetInfo(Desc.Platform, Desc.Configuration), Desc.bIsEditorRecompile, out TargetFileName);
