Fix incorrect code to check whether a target receipt exists.

diff --git a/Engine/Source/Programs/UnrealBuildTool/System/TargetReceipt.cs b/Engine/Source/Programs/UnrealBuildTool/System/TargetReceipt.cs
--- a/Engine/Source/Programs/UnrealBuildTool/System/TargetReceipt.cs
+++ b/Engine/Source/Programs/UnrealBuildTool/System/TargetReceipt.cs
@@ -397,7 +397,7 @@ namespace UnrealBuildTool
 		/// <param name="FileName">Filename to read from</param>
 		public static bool TryRead(string FileName, out TargetReceipt Receipt)
 		{
-			if(File.Exists(FileName))
+			if(!File.Exists(FileName))
 			{
 				Receipt = null;
 				return false;
@@ -406,8 +406,8 @@ namespace UnrealBuildTool
 			try
 			{
 				Receipt = Read(FileName);
-					return true;
-				}
+				return true;
+			}
 			catch(Exception)
 			{
 				Receipt = null;
