Escape characters in filenames which aren't valid in XML attributes.

diff --git a/Engine/Source/Programs/UnrealBuildTool/System/VCProject.cs b/Engine/Source/Programs/UnrealBuildTool/System/VCProject.cs
--- a/Engine/Source/Programs/UnrealBuildTool/System/VCProject.cs
+++ b/Engine/Source/Programs/UnrealBuildTool/System/VCProject.cs
@@ -9,6 +9,7 @@ using System.Xml.Linq;
 using System.Linq;
 using System.Text;
 using System.Diagnostics;
+using System.Security;
 
 namespace UnrealBuildTool
 {
@@ -245,6 +246,16 @@ namespace UnrealBuildTool
 			return true;
 		}
 
+		/// <summary>
+		/// Escapes characters in a filename so they can be stored in an XML attribute
+		/// </summary>
+		/// <param name="FileName">The filename to escape</param>
+		/// <returns>The escaped filename</returns>
+		public static string EscapeFileName(string FileName)
+		{
+			return SecurityElement.Escape(FileName);
+		}
+
 		/// <summary>
 		/// GUID for this Visual C++ project file
 		/// </summary>
@@ -301,7 +312,6 @@ namespace UnrealBuildTool
 			}
 		}
 
-
 		/// Implements Project interface
 		public override bool WriteProjectFile(List<UnrealTargetPlatform> InPlatforms, List<UnrealTargetConfiguration> InConfigurations)
 		{
@@ -566,12 +576,12 @@ namespace UnrealBuildTool
 					var VCFileType = GetVCFileType(AliasedFile.FileSystemPath);
 
 					VCProjectFileContent.Append(
-						"		<" + VCFileType + " Include=\"" + AliasedFile.FileSystemPath + "\" />" + ProjectFileGenerator.NewLine);
+						"		<" + VCFileType + " Include=\"" + EscapeFileName(AliasedFile.FileSystemPath) + "\" />" + ProjectFileGenerator.NewLine);
 
 					if (!String.IsNullOrWhiteSpace(AliasedFile.ProjectPath))
 					{
 						VCFiltersFileContent.Append(
-							"		<" + VCFileType + " Include=\"" + AliasedFile.FileSystemPath + "\">" + ProjectFileGenerator.NewLine +
+							"		<" + VCFileType + " Include=\"" + EscapeFileName(AliasedFile.FileSystemPath) + "\">" + ProjectFileGenerator.NewLine +
 							"			<Filter>" + Utils.CleanDirectorySeparators(AliasedFile.ProjectPath) + "</Filter>" + ProjectFileGenerator.NewLine +
 							"		</" + VCFileType + " >" + ProjectFileGenerator.NewLine);
 
@@ -581,7 +591,7 @@ namespace UnrealBuildTool
 					{
 						// No need to specify the root directory relative to the project (it would just be an empty string!)
 						VCFiltersFileContent.Append(
-							"		<" + VCFileType + " Include=\"" + AliasedFile.FileSystemPath + "\" />" + ProjectFileGenerator.NewLine);
+							"		<" + VCFileType + " Include=\"" + EscapeFileName(AliasedFile.FileSystemPath) + "\" />" + ProjectFileGenerator.NewLine);
 					}
 				}
 
@@ -1222,7 +1232,7 @@ namespace UnrealBuildTool
 				ProjectFileContent.Append("<ItemGroup>" + ProjectFileGenerator.NewLine);
 				foreach (var CurReference in DotNetAssemblyReferences)
 				{
-					ProjectFileContent.Append("\t<Reference Include=\"" + CurReference + "\" />" + ProjectFileGenerator.NewLine);
+					ProjectFileContent.Append("\t<Reference Include=\"" + EscapeFileName(CurReference) + "\" />" + ProjectFileGenerator.NewLine);
 				}
 				ProjectFileContent.Append("</ItemGroup>" + ProjectFileGenerator.NewLine);
 			}
@@ -1232,7 +1242,7 @@ namespace UnrealBuildTool
 				ProjectFileContent.Append( "<ItemGroup>" + ProjectFileGenerator.NewLine );
 				foreach( var CurReference in AssemblyReferences )
 				{
-					ProjectFileContent.Append( "\t<Reference Include=\"" + Path.GetFileNameWithoutExtension( CurReference ) + "\" >" + ProjectFileGenerator.NewLine );
+					ProjectFileContent.Append( "\t<Reference Include=\"" + EscapeFileName(Path.GetFileNameWithoutExtension( CurReference )) + "\" >" + ProjectFileGenerator.NewLine );
 					ProjectFileContent.Append( "\t\t<HintPath>" + Utils.MakePathRelativeTo( CurReference, Path.GetDirectoryName(ProjectFilePath) ) + "</HintPath>" + ProjectFileGenerator.NewLine );
 					ProjectFileContent.Append( "\t</Reference>" + ProjectFileGenerator.NewLine );
 				}
@@ -1246,7 +1256,7 @@ namespace UnrealBuildTool
 				RelativePath = Path.Combine(RelativePath, Path.GetFileName(Project.ProjectFilePath));
 				ProjectFileContent.Append(
 					"<ItemGroup>" + ProjectFileGenerator.NewLine +
-						"<ProjectReference Include=\"" + RelativePath + "\">" + ProjectFileGenerator.NewLine +
+						"<ProjectReference Include=\"" + EscapeFileName(RelativePath) + "\">" + ProjectFileGenerator.NewLine +
 							"<Project>" + ((MSBuildProjectFile)Project).ProjectGUID.ToString("B").ToUpperInvariant() + "</Project>" + ProjectFileGenerator.NewLine +
 							"<Name>" + Path.GetFileNameWithoutExtension(RelativePath) + "</Name>" + ProjectFileGenerator.NewLine +
 						"</ProjectReference>" + ProjectFileGenerator.NewLine +
@@ -1261,7 +1271,7 @@ namespace UnrealBuildTool
 			{
 				var ProjectRelativeSourceFile = Utils.MakePathRelativeTo( CurFile.FilePath, Path.GetDirectoryName( ProjectFilePath ) );
 				ProjectFileContent.Append(
-					"		<Compile Include=\"" + ProjectRelativeSourceFile + "\" />" + ProjectFileGenerator.NewLine);
+					"		<Compile Include=\"" + EscapeFileName(ProjectRelativeSourceFile) + "\" />" + ProjectFileGenerator.NewLine);
 			}
 			ProjectFileContent.Append(
 				"	</ItemGroup>" + ProjectFileGenerator.NewLine );
