[INTEGRATE] Unreal Build Tool: Adaptive unity builds (experimental!)
- New UBT support for "adaptive unity builds", which can greatly increase iterative compile times
- This feature allows UBT to keep track of which source files you're currently editing between runs
- UBT will exclude your "working set" source files from all unity blob files, allowing for fast iterative compiles!
- Currently, adaptive unity builds relies on files being writable vs. read-only to determine what you're working on (P4 edit)
- Enable this new feature in your BuildConfiguration.xml file.

    <BuildConfiguration>
           <bUseAdaptiveUnityBuild>true</bUseAdaptiveUnityBuild>
    </BuildConfiguration>

diff --git a/Engine/Source/Programs/UnrealBuildTool/Configuration/BuildConfiguration.cs b/Engine/Source/Programs/UnrealBuildTool/Configuration/BuildConfiguration.cs
--- a/Engine/Source/Programs/UnrealBuildTool/Configuration/BuildConfiguration.cs
+++ b/Engine/Source/Programs/UnrealBuildTool/Configuration/BuildConfiguration.cs
@@ -29,6 +29,15 @@ namespace UnrealBuildTool
 		[XmlConfig]
         public static bool bForceUnityBuild;
 
+		/// <summary>
+		/// An experimental new feature that, when enabled, will try to determine source files that you are actively iteratively changing, 
+		/// and break those files out of their unity blobs so that you can compile them as individual translation units, much faster than 
+		/// recompiling the entire unity blob each time.
+		/// IMPORTANT: This feature is not complete yet!  Currently, it requires source files to be read-only that you aren't actively working with!
+		/// </summary>
+		[XmlConfig]
+		public static bool bUseAdaptiveUnityBuild;
+
 		/// <summary>
 		/// The number of source files in a game module before unity build will be activated for that module.  This
 		/// allows small game modules to have faster iterative compile times for single files, at the expense of slower full
@@ -363,7 +372,7 @@ namespace UnrealBuildTool
 		public static bool bEnableCodeAnalysis;
 
 		/// <summary>
-		/// Enables "Shared PCHs", an experimental feature which may significantly speed up compile times by attempting to
+		/// Enables "Shared PCHs", a feature which significantly speeds up compile times by attempting to
 		/// share certain PCH files between modules that UBT detects is including those PCH's header files
 		/// </summary>
 		[XmlConfig]
@@ -559,6 +568,10 @@ namespace UnrealBuildTool
 			// the synthesized file names for unity code files may change between builds.
 			bUseUnityBuild = true;
 
+			// This experimental feature may help to improve iterative compile times when working on code, but is
+			// disabled by default until we test it sufficiently
+			bUseAdaptiveUnityBuild = false;
+
 			bForceUnityBuild = false;
 
 			// For programmers working on gameplay code, it's convenient to have fast iteration times when changing only single source
diff --git a/Engine/Source/Programs/UnrealBuildTool/System/Unity.cs b/Engine/Source/Programs/UnrealBuildTool/System/Unity.cs
--- a/Engine/Source/Programs/UnrealBuildTool/System/Unity.cs
+++ b/Engine/Source/Programs/UnrealBuildTool/System/Unity.cs
@@ -11,6 +11,11 @@ namespace UnrealBuildTool
 {
 	public class Unity
 	{
+		/** The set of source files that UnrealBuildTool determined to be part of the programmer's "working set".
+		    This feature is only used when bUseAdaptiveUnityBuild is enabled */
+		public static HashSet<FileItem> SourceFileWorkingSet = new HashSet<FileItem>();
+
+
 		/// <summary>
 		/// A class which represents a list of files and the sum of their lengths.
 		/// </summary>
@@ -114,6 +119,8 @@ namespace UnrealBuildTool
 			string BaseName
 			)
 		{
+			var NewCPPFiles = new List<FileItem>();
+
 			var ToolChain = UEToolChain.GetPlatformToolChain(CompileEnvironment.Config.Target.Platform);
 
 			var BuildPlatform = UEBuildPlatform.GetBuildPlatformForCPPTargetPlatform(CompileEnvironment.Config.Target.Platform);
@@ -124,7 +131,13 @@ namespace UnrealBuildTool
 			// We have an increased threshold for unity file size if, and only if, all files fit into the same unity file. This
 			// is beneficial when dealing with PCH files. The default PCH creation limit is X unity files so if we generate < X 
 			// this could be fairly slow and we'd rather bump the limit a bit to group them all into the same unity file.
-			//
+
+
+			// When enabled, UnrealBuildTool will try to determine source files that you are actively iteratively changing, and break those files
+			// out of their unity blobs so that you can compile them as individual translation units, much faster than recompiling the entire
+			// unity blob each time.
+			bool bUseAdaptiveUnityBuild = BuildConfiguration.bUseAdaptiveUnityBuild && !BuildConfiguration.bStressTestUnity;
+
 			// Optimization only makes sense if PCH files are enabled.
 			bool bForceIntoSingleUnityFile = BuildConfiguration.bStressTestUnity || (TotalBytesInCPPFiles < BuildConfiguration.NumIncludedBytesPerUnityCPP * 2 && CompileEnvironment.ShouldUsePCHs());
 
@@ -141,26 +154,85 @@ namespace UnrealBuildTool
 					SortedCPPFiles.Sort( FileItemComparer );
 				}
 
+
+				// Figure out whether we REALLY want to use adaptive unity for this module.  If nearly every file in the module appears in the working
+				// set, we'll just go ahead and let unity build do its thing.
+				if( bUseAdaptiveUnityBuild )
+				{
+					int CandidateWorkingSetSourceFileCount = 0;
+					int WorkingSetSourceFileCount = 0;
+				    foreach( var CPPFile in SortedCPPFiles )
+				    {
+					    // Don't include writable source files into unity blobs
+						if( !Utils.IsFileUnderDirectory( CPPFile.AbsolutePath, Target.EngineIntermediateDirectory ) &&
+							!Utils.IsFileUnderDirectory( CPPFile.AbsolutePath, Target.ProjectIntermediateDirectory ) )
+						{
+							++CandidateWorkingSetSourceFileCount;
+
+							if( UnrealBuildTool.ShouldSourceFileBePartOfWorkingSet( CPPFile.AbsolutePath ) )
+							{
+								++WorkingSetSourceFileCount;
+
+								// Mark this file as part of the working set.  This will be saved into the UBT Makefile so that
+								// the assembler can automatically invalidate the Makefile when the working set changes (allowing this
+								// code to run again, to build up new unity blobs.)
+								SourceFileWorkingSet.Add( CPPFile );
+							}
+						}
+					}					
+
+					if( WorkingSetSourceFileCount >= CandidateWorkingSetSourceFileCount )
+					{
+						// Every single file in the module appears in the working set, so don't bother using adaptive unity for this
+						// module.  Otherwise it would make full builds really slow.
+						bUseAdaptiveUnityBuild = false;
+					}
+				}
+
 				var CPPUnityFileBuilder = new UnityFileBuilder(bForceIntoSingleUnityFile ? -1 : BuildConfiguration.NumIncludedBytesPerUnityCPP);
+				var AdaptiveUnityBuildInfoString = new StringBuilder();
 				foreach( var CPPFile in SortedCPPFiles )
 				{
-					if (!bForceIntoSingleUnityFile && CPPFile.AbsolutePath.Contains(".GeneratedWrapper."))
+					if (!bForceIntoSingleUnityFile && CPPFile.AbsolutePath.IndexOf(".GeneratedWrapper.", StringComparison.InvariantCultureIgnoreCase) != -1)
 					{
-						CPPUnityFileBuilder.EndCurrentUnityFile();
-						CPPUnityFileBuilder.AddFile(CPPFile);
-						CPPUnityFileBuilder.EndCurrentUnityFile();
+						NewCPPFiles.Add(CPPFile);
+					}
+
+					// When adaptive unity is enabled, go ahead and exclude any source files that we're actively working with
+					if( bUseAdaptiveUnityBuild && SourceFileWorkingSet.Contains( CPPFile ) )
+					{
+						// Just compile this file normally, not as part of the unity blob
+						NewCPPFiles.Add( CPPFile );
+
+						var CPPFileName = Path.GetFileName( CPPFile.AbsolutePath );
+						if( AdaptiveUnityBuildInfoString.Length == 0 )
+						{
+							AdaptiveUnityBuildInfoString.Append( String.Format( "[Adaptive unity build] Excluded from {0} unity file: {1}", BaseName, CPPFileName ) );
+						}
+						else
+						{
+							AdaptiveUnityBuildInfoString.Append( ", " + CPPFileName );
+						}
 					}
+
 					else
 					{
+						// Compile this file as part of the unity blob
 						CPPUnityFileBuilder.AddFile(CPPFile);
+
+						// Now that the CPPFile is part of this unity file, we will no longer need to treat it like a root level prerequisite for our
+						// dependency cache, as it is now an "indirect include" from the unity file.  We'll clear out the compile environment
+						// attached to this file.  This prevents us from having to cache all of the indirect includes from these files inside our
+						// dependency cache, which speeds up iterative builds a lot!
+						CPPFile.CachedCPPIncludeInfo = null;
 					}
+				}
 
-					// Now that the CPPFile is part of this unity file, we will no longer need to treat it like a root level prerequisite for our
-					// dependency cache, as it is now an "indirect include" from the unity file.  We'll clear out the compile environment
-					// attached to this file.  This prevents us from having to cache all of the indirect includes from these files inside our
-					// dependency cache, which speeds up iterative builds a lot!
-					CPPFile.CachedCPPIncludeInfo = null;
+				if( AdaptiveUnityBuildInfoString.Length > 0 )
+				{
+					Log.TraceInformation( AdaptiveUnityBuildInfoString.ToString() );
 				}
+
 				AllUnityFiles = CPPUnityFileBuilder.GetUnityFiles();
 			}
 
@@ -180,7 +252,6 @@ namespace UnrealBuildTool
 			// Create a set of CPP files that combine smaller CPP files into larger compilation units, along with the corresponding 
 			// actions to compile them.
 			int CurrentUnityFileCount = 0;
-			var UnityCPPFiles         = new List<FileItem>();
 			foreach( var UnityFile in AllUnityFiles )
 			{
 				++CurrentUnityFileCount;
@@ -238,14 +309,14 @@ namespace UnrealBuildTool
 
 				UnityCPPFile.RelativeCost        = UnityFile.TotalLength;
 				UnityCPPFile.PCHHeaderNameInCode = PCHHeaderNameInCode;
-                UnityCPPFiles.Add(UnityCPPFile);
+                NewCPPFiles.Add(UnityCPPFile);
 
 				// Cache information about the unity .cpp dependencies
 				// @todo ubtmake urgent: Fails when building remotely for Mac because unity .cpp has an include for a PCH on the REMOTE machine
 				UEBuildModuleCPP.CachePCHUsageForModuleSourceFile( Target, CompileEnvironment, UnityCPPFile );
 			}
 
-			return UnityCPPFiles;
+			return NewCPPFiles;
 		}
 	}
 }
diff --git a/Engine/Source/Programs/UnrealBuildTool/System/UnrealBuildTool.cs b/Engine/Source/Programs/UnrealBuildTool/System/UnrealBuildTool.cs
--- a/Engine/Source/Programs/UnrealBuildTool/System/UnrealBuildTool.cs
+++ b/Engine/Source/Programs/UnrealBuildTool/System/UnrealBuildTool.cs
@@ -1890,6 +1890,7 @@ namespace UnrealBuildTool
                         }
                         UBTMakefile.TargetNameToUObjectModules = TargetNameToUObjectModules;
 						UBTMakefile.Targets = Targets;
+						UBTMakefile.SourceFileWorkingSet = Unity.SourceFileWorkingSet;
 
 						if( BuildConfiguration.bUseUBTMakefiles )
 						{ 
@@ -2249,7 +2250,7 @@ namespace UnrealBuildTool
         [Serializable]
         class UBTMakefile : ISerializable
         {
-			public const int CurrentVersion = 2;
+			public const int CurrentVersion = 3;
 
 			public int Version;
 
@@ -2269,6 +2270,10 @@ namespace UnrealBuildTool
 			/** List of targets being built */
 			public List<UEBuildTarget> Targets;
 
+
+			/** Current working set of source files, for when bUseAdaptiveUnityBuild is enabled */
+			public HashSet<FileItem> SourceFileWorkingSet = new HashSet<FileItem>();
+
 			public UBTMakefile()
 			{
 				Version = CurrentVersion;
@@ -2287,6 +2292,7 @@ namespace UnrealBuildTool
 				EnvironmentVariables       = ((string[])Info.GetValue("e1", typeof(string[]))).Zip((string[])Info.GetValue("e2", typeof(string[])), (i1, i2) => new Tuple<string, string>(i1, i2)).ToList();
 				TargetNameToUObjectModules = (Dictionary<string,List<UHTModuleInfo>>)Info.GetValue("nu", typeof(Dictionary<string,List<UHTModuleInfo>>));
 				Targets                    = (List<UEBuildTarget>)Info.GetValue("ta", typeof(List<UEBuildTarget>));
+				SourceFileWorkingSet       = (HashSet<FileItem>)Info.GetValue("ws", typeof(HashSet<FileItem>));
 			}
 
 			public void GetObjectData(SerializationInfo Info, StreamingContext Context)
@@ -2298,6 +2304,7 @@ namespace UnrealBuildTool
 				Info.AddValue("e2", EnvironmentVariables.Select(x => x.Item2).ToArray());
 				Info.AddValue("nu", TargetNameToUObjectModules);
 				Info.AddValue("ta", Targets);
+				Info.AddValue("ws", SourceFileWorkingSet);
 			}
 
 
@@ -2309,7 +2316,8 @@ namespace UnrealBuildTool
                     PrerequisiteActions != null && PrerequisiteActions.Length > 0 &&
                     EnvironmentVariables != null &&
                     TargetNameToUObjectModules != null && TargetNameToUObjectModules.Count > 0 &&
-					Targets != null && Targets.Count > 0;
+					Targets != null && Targets.Count > 0 &&
+					SourceFileWorkingSet != null;
             }
         }
 
@@ -2506,7 +2514,7 @@ namespace UnrealBuildTool
 					new HashSet<string>(
 						Target.FlatModuleCsData
 						.SelectMany(x => x.Value.ModuleSourceFolder != null ? Directory.EnumerateFiles(x.Value.ModuleSourceFolder, "*.h", SearchOption.AllDirectories) : Enumerable.Empty<string>())
-						.Where(y => Directory.GetLastWriteTime(y) > UBTMakefileInfo.LastWriteTime)
+						.Where(y => Directory.GetLastWriteTimeUtc(y) > UBTMakefileInfo.LastWriteTimeUtc)
 						.OrderBy(z => z).Distinct()
 					);
 
@@ -2523,12 +2531,54 @@ namespace UnrealBuildTool
 					if (bContainsUHTData != bWasProcessed)
 					{
 						Log.TraceVerbose("{0} {1} contain UHT types and now {2} , ignoring it ({3})", Filename, bWasProcessed ? "used to" : "didn't", bWasProcessed ? "doesn't" : "does", UBTMakefileInfo.FullName);
-						ReasonNotLoaded = string.Format("List of files containing UHT types has changed");
+						ReasonNotLoaded = string.Format("new files with reflected types");
 						return null;
 					}
 				}
 			}
 
+			// If adaptive unity build is enabled, do a check to see if there are any source files that became part of the
+			// working set since the Makefile was created (or, source files were removed from the working set.)  If anything
+			// changed, then we'll force a new Makefile to be created so that we have fresh unity build blobs.  We always
+			// want to make sure that source files in the working set are excluded from those unity blobs (for fastest possible
+			// iteration times.)
+			if( BuildConfiguration.bUseAdaptiveUnityBuild )
+			{
+				// Get all .cpp files in processed modules newer than the makefile itself
+				foreach (var Target in LoadedUBTMakefile.Targets)
+				{
+					foreach( var FlatModule in Target.FlatModuleCsData.Values )
+					{
+						if( FlatModule != null && !String.IsNullOrEmpty( FlatModule.ModuleSourceFolder ) )
+						{
+							// Has the directory been touched?  Don't bother checking folders that haven't been written to since the
+							// Makefile was created. 
+							//if( Directory.GetLastWriteTimeUtc( FlatModule.ModuleSourceFolder ) > UBTMakefileInfo.LastWriteTimeUtc )
+							{
+								foreach( var SourceFilePath in Directory.EnumerateFiles( FlatModule.ModuleSourceFolder, "*.cpp", SearchOption.AllDirectories ) )
+								{
+									// Was the file written to since the Makefile was created?  We're not able to do a fully exhaustive
+									// check for working sets when in assembler mode, so we might get it wrong occasionally.  This is
+									// why we only want to check newly-modified or created files.  Otherwise, we could potentially
+									// invalidate the Makefile every single run.
+									if( File.GetLastWriteTimeUtc( SourceFilePath ) > UBTMakefileInfo.LastWriteTimeUtc )
+									{
+										bool bShouldBePartOfWorkingSet = UnrealBuildTool.ShouldSourceFileBePartOfWorkingSet( SourceFilePath );
+										bool bIsAlreadyPartOfWorkingSet = LoadedUBTMakefile.SourceFileWorkingSet.Contains( FileItem.GetItemByFullPath( SourceFilePath ) );
+										if( bShouldBePartOfWorkingSet != bIsAlreadyPartOfWorkingSet )
+										{
+											Log.TraceVerbose("{0} {1} part of source working set and now {2}; invalidating makefile ({3})", SourceFilePath, bIsAlreadyPartOfWorkingSet ? "was" : "was not", bShouldBePartOfWorkingSet ? "is" : "is not", UBTMakefileInfo.FullName);
+											ReasonNotLoaded = string.Format("working set of source files changed");
+											return null;
+										}
+									}
+								}
+							}
+						}
+					}
+				}
+			}
+
 			ReasonNotLoaded = null;
 			return LoadedUBTMakefile;
 		}
@@ -2838,5 +2888,21 @@ namespace UnrealBuildTool
 				}
 			}
 		}
+
+		/// <summary>
+		/// Given a path to a source file, determines whether it should be considered part of the programmer's working
+		/// file set.  That is, a file that is currently being iterated on and is likely to be recompiled soon.  This
+		/// is only used when bUseAdaptiveUnityBuild is enabled.
+		/// </summary>
+		/// <param name="SourceFileAbsolutePath">Path to the source file on disk</param>
+		/// <returns>True if file is part of the working set</returns>
+		public static bool ShouldSourceFileBePartOfWorkingSet( string SourceFileAbsolutePath )
+		{
+			// @todo: This logic needs work.  Currently its designed for users working with Perforce, or a similar
+			// source control system that keeps source files read-only when not actively checked out for editing.
+			bool bShouldBePartOfWorkingSourceFileSet =
+				!System.IO.File.GetAttributes( SourceFileAbsolutePath ).HasFlag( FileAttributes.ReadOnly );
+			return bShouldBePartOfWorkingSourceFileSet;
+		}
     }
 }
