Change IUEBuildPlatform.ModifyNewlyLoadedModule to IUEBuildPlatform.ModifyModuleRules(). Only modifying the rules object rather than the implementation itself makes it easier to optimize caching module references.

diff --git a/Engine/Source/Programs/UnrealBuildTool/Android/UEBuildAndroid.cs b/Engine/Source/Programs/UnrealBuildTool/Android/UEBuildAndroid.cs
--- a/Engine/Source/Programs/UnrealBuildTool/Android/UEBuildAndroid.cs
+++ b/Engine/Source/Programs/UnrealBuildTool/Android/UEBuildAndroid.cs
@@ -298,36 +298,36 @@ namespace UnrealBuildTool
 		{
 		}
 
-		public override void ModifyNewlyLoadedModule(UEBuildModule InModule, TargetInfo Target)
+		public override void ModifyModuleRules(string ModuleName, ModuleRules Rules, TargetInfo Target)
 		{
 			if ((Target.Platform == UnrealTargetPlatform.Win32) || (Target.Platform == UnrealTargetPlatform.Win64) || (Target.Platform == UnrealTargetPlatform.Mac))
 			{
 				bool bBuildShaderFormats = UEBuildConfiguration.bForceBuildShaderFormats;
 				if (!UEBuildConfiguration.bBuildRequiresCookedData)
 				{
-					if (InModule.ToString() == "Engine")
+					if (ModuleName == "Engine")
 					{
 						if (UEBuildConfiguration.bBuildDeveloperTools)
 						{
-							InModule.AddPlatformSpecificDynamicallyLoadedModule("AndroidTargetPlatform");
-							InModule.AddPlatformSpecificDynamicallyLoadedModule("Android_PVRTCTargetPlatform");
-							InModule.AddPlatformSpecificDynamicallyLoadedModule("Android_ATCTargetPlatform");
-							InModule.AddPlatformSpecificDynamicallyLoadedModule("Android_DXTTargetPlatform");
-							InModule.AddPlatformSpecificDynamicallyLoadedModule("Android_ETC1TargetPlatform");
-							InModule.AddPlatformSpecificDynamicallyLoadedModule("Android_ETC2TargetPlatform");
-							InModule.AddPlatformSpecificDynamicallyLoadedModule("Android_ASTCTargetPlatform");
+							Rules.PlatformSpecificDynamicallyLoadedModuleNames.Add("AndroidTargetPlatform");
+							Rules.PlatformSpecificDynamicallyLoadedModuleNames.Add("Android_PVRTCTargetPlatform");
+							Rules.PlatformSpecificDynamicallyLoadedModuleNames.Add("Android_ATCTargetPlatform");
+							Rules.PlatformSpecificDynamicallyLoadedModuleNames.Add("Android_DXTTargetPlatform");
+							Rules.PlatformSpecificDynamicallyLoadedModuleNames.Add("Android_ETC1TargetPlatform");
+							Rules.PlatformSpecificDynamicallyLoadedModuleNames.Add("Android_ETC2TargetPlatform");
+							Rules.PlatformSpecificDynamicallyLoadedModuleNames.Add("Android_ASTCTargetPlatform");
                         }
 					}
-					else if (InModule.ToString() == "TargetPlatform")
+					else if (ModuleName == "TargetPlatform")
 					{
 						bBuildShaderFormats = true;
-						InModule.AddDynamicallyLoadedModule("TextureFormatPVR");
-						InModule.AddDynamicallyLoadedModule("TextureFormatDXT");
-						InModule.AddDynamicallyLoadedModule("TextureFormatASTC");
-						InModule.AddPlatformSpecificDynamicallyLoadedModule("TextureFormatAndroid");    // ATITC, ETC1 and ETC2
+						Rules.DynamicallyLoadedModuleNames.Add("TextureFormatPVR");
+						Rules.DynamicallyLoadedModuleNames.Add("TextureFormatDXT");
+						Rules.DynamicallyLoadedModuleNames.Add("TextureFormatASTC");
+						Rules.PlatformSpecificDynamicallyLoadedModuleNames.Add("TextureFormatAndroid");    // ATITC, ETC1 and ETC2
 						if (UEBuildConfiguration.bBuildDeveloperTools)
 						{
-							//InModule.AddDynamicallyLoadedModule("AudioFormatADPCM");	//@todo android: android audio
+							//Rules.DynamicallyLoadedModuleNames.Add("AudioFormatADPCM");	//@todo android: android audio
 						}
 					}
 				}
@@ -335,18 +335,18 @@ namespace UnrealBuildTool
 				// allow standalone tools to use targetplatform modules, without needing Engine
 				if (UEBuildConfiguration.bForceBuildTargetPlatforms)
 				{
-					InModule.AddPlatformSpecificDynamicallyLoadedModule("AndroidTargetPlatform");
-					InModule.AddPlatformSpecificDynamicallyLoadedModule("Android_PVRTCTargetPlatform");
-					InModule.AddPlatformSpecificDynamicallyLoadedModule("Android_ATCTargetPlatform");
-					InModule.AddPlatformSpecificDynamicallyLoadedModule("Android_DXTTargetPlatform");
-					InModule.AddPlatformSpecificDynamicallyLoadedModule("Android_ETC1TargetPlatform");
-					InModule.AddPlatformSpecificDynamicallyLoadedModule("Android_ETC2TargetPlatform");
-					InModule.AddPlatformSpecificDynamicallyLoadedModule("Android_ASTCTargetPlatform");
+					Rules.PlatformSpecificDynamicallyLoadedModuleNames.Add("AndroidTargetPlatform");
+					Rules.PlatformSpecificDynamicallyLoadedModuleNames.Add("Android_PVRTCTargetPlatform");
+					Rules.PlatformSpecificDynamicallyLoadedModuleNames.Add("Android_ATCTargetPlatform");
+					Rules.PlatformSpecificDynamicallyLoadedModuleNames.Add("Android_DXTTargetPlatform");
+					Rules.PlatformSpecificDynamicallyLoadedModuleNames.Add("Android_ETC1TargetPlatform");
+					Rules.PlatformSpecificDynamicallyLoadedModuleNames.Add("Android_ETC2TargetPlatform");
+					Rules.PlatformSpecificDynamicallyLoadedModuleNames.Add("Android_ASTCTargetPlatform");
                 }
 
 				if (bBuildShaderFormats)
 				{
-					//InModule.AddDynamicallyLoadedModule("ShaderFormatAndroid");		//@todo android: ShaderFormatAndroid
+					//Rules.DynamicallyLoadedModuleNames.Add("ShaderFormatAndroid");		//@todo android: ShaderFormatAndroid
 				}
 			}
 		}
diff --git a/Engine/Source/Programs/UnrealBuildTool/Configuration/UEBuildModule.cs b/Engine/Source/Programs/UnrealBuildTool/Configuration/UEBuildModule.cs
--- a/Engine/Source/Programs/UnrealBuildTool/Configuration/UEBuildModule.cs
+++ b/Engine/Source/Programs/UnrealBuildTool/Configuration/UEBuildModule.cs
@@ -306,266 +306,6 @@ namespace UnrealBuildTool
 			Target.RegisterModule(this);
 		}
 
-		/** Adds a public definition */
-		public virtual void AddPublicDefinition(string Definition)
-		{
-			PublicDefinitions.Add(Definition);
-		}
-
-		/** Adds a public include path. */
-		public virtual void AddPublicIncludePath(string PathName, bool bCheckForDuplicates = true)
-		{
-			if (bCheckForDuplicates == true)
-			{
-				if (PublicIncludePaths.Contains(PathName))
-				{
-					return;
-				}
-			}
-			PublicIncludePaths.Add(PathName);
-		}
-
-		/** Adds a public library path */
-		public virtual void AddPublicLibraryPath(string PathName, bool bCheckForDuplicates = true)
-		{
-			if (bCheckForDuplicates == true)
-			{
-				if (PublicLibraryPaths.Contains(PathName))
-				{
-					return;
-				}
-			}
-			PublicLibraryPaths.Add(PathName);
-		}
-
-		/** Adds a public additional library */
-		public virtual void AddPublicAdditionalLibrary(string LibraryName, bool bCheckForDuplicates = true)
-		{
-			if (bCheckForDuplicates == true)
-			{
-				if (PublicAdditionalLibraries.Contains(LibraryName))
-				{
-					return;
-				}
-			}
-			PublicAdditionalLibraries.Add(LibraryName);
-		}
-
-		/** Adds a public framework */
-		public virtual void AddPublicFramework(string LibraryName, bool bCheckForDuplicates = true)
-		{
-			if (bCheckForDuplicates == true)
-			{
-				if (PublicFrameworks.Contains(LibraryName))
-				{
-					return;
-				}
-			}
-			PublicFrameworks.Add(LibraryName);
-		}
-
-		/** Adds a public weak framework */
-		public virtual void AddPublicWeakFramework(string LibraryName, bool bCheckForDuplicates = true)
-		{
-			if (bCheckForDuplicates == true)
-			{
-				if (PublicWeakFrameworks.Contains(LibraryName))
-				{
-					return;
-				}
-			}
-			PublicWeakFrameworks.Add(LibraryName);
-		}
-
-		/** Adds a public additional framework */
-		public virtual void AddPublicAdditionalFramework(UEBuildFramework Framework, bool bCheckForDuplicates = true)
-		{
-			if (bCheckForDuplicates == true)
-			{
-				if ( PublicAdditionalFrameworks.Contains( Framework ) )
-				{
-					return;
-				}
-			}
-			PublicAdditionalFrameworks.Add( Framework );
-		}
-
-		/** Adds a file that will be synced to a remote machine if a remote build is being executed */
-		public virtual void AddPublicAdditionalShadowFile(string FileName)
-		{
-			PublicAdditionalShadowFiles.Add(FileName);
-		}
-
-		/** Adds a file or folder that will be copied to Mac or iOS app bundle */
-		public virtual void AddAdditionalBundleResource(UEBuildBundleResource Resource, bool bCheckForDuplicates = true)
-		{
-			if (bCheckForDuplicates == true)
-			{
-				if ( PublicAdditionalBundleResources.Contains( Resource ) )
-				{
-					return;
-				}
-			}
-			PublicAdditionalBundleResources.Add( Resource );
-		}
-
-		/** Adds a public include path module. */
-		public virtual void AddPublicIncludePathModule(string ModuleName, bool bCheckForDuplicates = true)
-		{
-			if (bCheckForDuplicates == true)
-			{
-				if (PublicIncludePathModuleNames.Contains(ModuleName))
-				{
-					return;
-				}
-			}
-			PublicIncludePathModuleNames.Add(ModuleName);
-		}
-
-		/** Adds a public dependency module. */
-		public virtual void AddPublicDependencyModule(string ModuleName, bool bCheckForDuplicates = true)
-		{
-			if (bCheckForDuplicates == true)
-			{
-				if (PublicDependencyModuleNames.Contains(ModuleName))
-				{
-					return;
-				}
-			}
-			PublicDependencyModuleNames.Add(ModuleName);
-		}
-		
-		/** Adds a dynamically loaded module. */
-		public virtual void AddDynamicallyLoadedModule(string ModuleName, bool bCheckForDuplicates = true)
-		{
-			if (bCheckForDuplicates == true)
-			{
-				if (DynamicallyLoadedModuleNames.Contains(ModuleName))
-				{
-					return;
-				}
-			}
-			DynamicallyLoadedModuleNames.Add(ModuleName);
-		}
-        /** Adds a dynamically loaded module platform specifc. */
-		public virtual void AddPlatformSpecificDynamicallyLoadedModule(string ModuleName, bool bCheckForDuplicates = true)
-        {
-            if (bCheckForDuplicates == true)
-            {
-                if (PlatformSpecificDynamicallyLoadedModuleNames.Contains(ModuleName))
-                {
-                    return;
-                }
-            }
-            PlatformSpecificDynamicallyLoadedModuleNames.Add(ModuleName);
-        }
-		/** Adds a public delay load DLL */
-		public virtual void AddPublicDelayLoadDLLs(string DelayLoadDLLName, bool bCheckForDuplicates = true)
-		{
-			if (bCheckForDuplicates == true)
-			{
-				if (PublicDelayLoadDLLs.Contains(DelayLoadDLLName))
-				{
-					return;
-				}
-			}
-			PublicDelayLoadDLLs.Add(DelayLoadDLLName);
-		}
-
-		/** Adds a private include path. */
-		public virtual void AddPrivateIncludePath(string PathName, bool bCheckForDuplicates = true)
-		{
-			if (bCheckForDuplicates == true)
-			{
-				if (PrivateIncludePaths.Contains(PathName))
-				{
-					return;
-				}
-			}
-			PrivateIncludePaths.Add(PathName);
-		}
-
-		/** Adds a private include path module. */
-		public virtual void AddPrivateIncludePathModule(string ModuleName, bool bCheckForDuplicates = true)
-		{
-			if (bCheckForDuplicates == true)
-			{
-				if (PrivateIncludePathModuleNames.Contains(ModuleName))
-				{
-					return;
-				}
-			}
-			PrivateIncludePathModuleNames.Add(ModuleName);
-		}
-
-		/** Adds a dependency module. */
-		public virtual void AddPrivateDependencyModule(string ModuleName, bool bCheckForDuplicates = true)
-		{
-			if (bCheckForDuplicates == true)
-			{
-				if (PrivateDependencyModuleNames.Contains(ModuleName))
-				{
-					return;
-				}
-			}
-			PrivateDependencyModuleNames.Add(ModuleName);
-		}
-
-		/** Tells the module that a specific dependency is circularly reference */
-		public virtual void AddCircularlyReferencedDependentModule(string ModuleName)
-		{
-			CircularlyReferencedDependentModules.Add( ModuleName );
-		}
-
-		/** Adds an extra module name */
-		public virtual void AddExtraModuleName(string ModuleName)
-		{
-			DynamicallyLoadedModuleNames.Add( ModuleName );
-		}
-
-		/**
-		 * Clears all public include paths.
-		 */
-		internal void ClearPublicIncludePaths()
-		{
-			PublicIncludePaths.Clear();
-		}
-
-		/** 
-		 * Clears all public library paths.
-		 * This is to allow a given platform the opportunity to remove set libraries.
-		 */
-		internal void ClearPublicLibraryPaths()
-		{
-			PublicLibraryPaths.Clear();
-		}
-
-		/** 
-		 * Clears all public additional libraries from the module.
-		 * This is to allow a given platform the opportunity to remove set libraries.
-		 */
-		internal void ClearPublicAdditionalLibraries()
-		{
-			PublicAdditionalLibraries.Clear();
-		}
-
-		/// <summary>
-		/// If present, remove the given library from the PublicAdditionalLibraries array
-		/// </summary>
-		/// <param name="InLibrary">The library to remove from the list</param>
-		internal void RemovePublicAdditionalLibrary(string InLibrary)
-		{
-			PublicAdditionalLibraries.Remove(InLibrary);
-		}
-
-		/**
-		 * If found, remove the given entry from the definitions of the module.
-		 */
-		internal void RemovePublicDefinition(string InDefintion)
-		{
-			PublicDefinitions.Remove(InDefintion);
-		}
-
 		/** Sets up the environment for compiling any module that includes the public interface of this module. */
 		protected virtual void SetupPublicCompileEnvironment(
 			UEBuildBinary SourceBinary,
diff --git a/Engine/Source/Programs/UnrealBuildTool/Configuration/UEBuildPlatform.cs b/Engine/Source/Programs/UnrealBuildTool/Configuration/UEBuildPlatform.cs
--- a/Engine/Source/Programs/UnrealBuildTool/Configuration/UEBuildPlatform.cs
+++ b/Engine/Source/Programs/UnrealBuildTool/Configuration/UEBuildPlatform.cs
@@ -150,10 +150,11 @@ namespace UnrealBuildTool
 		 * This is not required - but allows for hiding details of a
 		 * particular platform.
 		 * 
-		 * @param InModule The newly loaded module
-		 * @param Target The target being build
+		 *  @param  ModuleName		The name of the module
+		 *	@param	Rules			The module rules
+		 *	@param	Target			The target being build
 		 */
-		void ModifyNewlyLoadedModule(UEBuildModule InModule, TargetInfo Target);
+		void ModifyModuleRules(string ModuleName, ModuleRules Rules, TargetInfo Target);
 
 		/**
 		 * Setup the target environment for building
@@ -411,17 +412,18 @@ namespace UnrealBuildTool
 		 *	passed in for the given platform.
 		 *	This is not required - but allows for hiding details of a particular platform.
 		 *	
-		 *	@param	InModule		The newly loaded module
+		 *  @param  Name			The name of the module
+		 *	@param	Module			The module rules
 		 *	@param	Target			The target being build
 		 *	@param	Only			If this is not unknown, then only run that platform
 		 */
-		public static void PlatformModifyNewlyLoadedModule(UEBuildModule InModule, TargetInfo Target, UnrealTargetPlatform Only = UnrealTargetPlatform.Unknown)
+		public static void PlatformModifyModuleRules(string ModuleName, ModuleRules Rules, TargetInfo Target, UnrealTargetPlatform Only = UnrealTargetPlatform.Unknown)
 		{
 			foreach (var PlatformEntry in BuildPlatformDictionary)
 			{
 				if (Only == UnrealTargetPlatform.Unknown || PlatformEntry.Key == Only || PlatformEntry.Key == Target.Platform)
 				{
-					PlatformEntry.Value.ModifyNewlyLoadedModule(InModule, Target);
+					PlatformEntry.Value.ModifyModuleRules(ModuleName, Rules, Target);
 				}
 			}
 		}
@@ -635,10 +637,11 @@ namespace UnrealBuildTool
 		 *	This is not required - but allows for hiding details of a
 		 *	particular platform.
 		 *	
-		 *	@param	InModule		The newly loaded module
+		 *  @param  ModuleName		The name of the module
+		 *	@param	Rules			The module rules
 		 *	@param	Target			The target being build
 		 */
-		public virtual void ModifyNewlyLoadedModule(UEBuildModule InModule, TargetInfo Target)
+		public virtual void ModifyModuleRules(string ModuleName, ModuleRules Rules, TargetInfo Target)
 		{
 		}
 
diff --git a/Engine/Source/Programs/UnrealBuildTool/Configuration/UEBuildTarget.cs b/Engine/Source/Programs/UnrealBuildTool/Configuration/UEBuildTarget.cs
--- a/Engine/Source/Programs/UnrealBuildTool/Configuration/UEBuildTarget.cs
+++ b/Engine/Source/Programs/UnrealBuildTool/Configuration/UEBuildTarget.cs
@@ -3438,20 +3438,19 @@ namespace UnrealBuildTool
 					}
 				}
 
-				// Now, go ahead and create the module builder instance
-				Module = InstantiateModule(RulesObject, ModuleName, ModuleType, ModuleDirectory, GeneratedCodeDirectory, IntelliSenseGatherer, FoundSourceFiles, bBuildFiles, ModuleFileName);
-
+				// Allow all build platforms to 'adjust' the module setting. 
+				// This will allow undisclosed platforms to make changes without 
+				// exposing information about the platform in publicly accessible 
+				// locations.
 				UnrealTargetPlatform Only = UnrealBuildTool.GetOnlyPlatformSpecificFor();
-
 				if (Only == UnrealTargetPlatform.Unknown && UnrealBuildTool.SkipNonHostPlatforms())
 				{
 					Only = Platform;
 				}
-				// Allow all build platforms to 'adjust' the module setting. 
-				// This will allow undisclosed platforms to make changes without 
-				// exposing information about the platform in publicly accessible 
-				// locations.
-				UEBuildPlatform.PlatformModifyNewlyLoadedModule(Module, TargetInfo, Only);
+				UEBuildPlatform.PlatformModifyModuleRules(ModuleName, RulesObject, TargetInfo, Only);
+
+				// Now, go ahead and create the module builder instance
+				Module = InstantiateModule(RulesObject, ModuleName, ModuleType, ModuleDirectory, GeneratedCodeDirectory, IntelliSenseGatherer, FoundSourceFiles, bBuildFiles, ModuleFileName);
 			}
 			return Module;
 		}
diff --git a/Engine/Source/Programs/UnrealBuildTool/HTML5/UEBuildHTML5.cs b/Engine/Source/Programs/UnrealBuildTool/HTML5/UEBuildHTML5.cs
--- a/Engine/Source/Programs/UnrealBuildTool/HTML5/UEBuildHTML5.cs
+++ b/Engine/Source/Programs/UnrealBuildTool/HTML5/UEBuildHTML5.cs
@@ -289,35 +289,36 @@ namespace UnrealBuildTool
          *	This is not required - but allows for hiding details of a
          *	particular platform.
          *	
-         *	@param	InModule		The newly loaded module
-         *	@param	Target			The target being build
-         */
-        public override void ModifyNewlyLoadedModule(UEBuildModule InModule, TargetInfo Target)
+		 *  @param  Name			The name of the module
+		 *	@param	Rules			The module rules
+		 *	@param	Target			The target being build
+		 */
+		public override void ModifyModuleRules(string ModuleName, ModuleRules Rules, TargetInfo Target)
         {
 			if (Target.Platform == UnrealTargetPlatform.HTML5)
             {
-                if (InModule.ToString() == "Core")
+                if (ModuleName == "Core")
                 {
-                    InModule.AddPublicIncludePath("Runtime/Core/Public/HTML5");
-                    InModule.AddPublicDependencyModule("zlib");
+                    Rules.PublicIncludePaths.Add("Runtime/Core/Public/HTML5");
+                    Rules.PublicDependencyModuleNames.Add("zlib");
                 }
-                else if (InModule.ToString() == "Engine")
+                else if (ModuleName == "Engine")
                 {
-                    InModule.AddPrivateDependencyModule("zlib");
-                    InModule.AddPrivateDependencyModule("UElibPNG");
-                    InModule.AddPublicDependencyModule("UEOgg");
-                    InModule.AddPublicDependencyModule("Vorbis");
+                    Rules.PrivateDependencyModuleNames.Add("zlib");
+                    Rules.PrivateDependencyModuleNames.Add("UElibPNG");
+                    Rules.PublicDependencyModuleNames.Add("UEOgg");
+                    Rules.PublicDependencyModuleNames.Add("Vorbis");
                 }
             }
             else if (Target.Platform == UnrealTargetPlatform.Win32 || Target.Platform == UnrealTargetPlatform.Win64 || Target.Platform == UnrealTargetPlatform.Mac )
             {
 				// allow standalone tools to use targetplatform modules, without needing Engine
 				if ((!UEBuildConfiguration.bBuildRequiresCookedData
-                    && InModule.ToString() == "Engine"
+                    && ModuleName == "Engine"
                     && UEBuildConfiguration.bBuildDeveloperTools)
                     || UEBuildConfiguration.bForceBuildTargetPlatforms)
                 {
-                    InModule.AddPlatformSpecificDynamicallyLoadedModule("HTML5TargetPlatform");
+                    Rules.PlatformSpecificDynamicallyLoadedModuleNames.Add("HTML5TargetPlatform");
                 }
             }
         }
diff --git a/Engine/Source/Programs/UnrealBuildTool/IOS/UEBuildIOS.cs b/Engine/Source/Programs/UnrealBuildTool/IOS/UEBuildIOS.cs
--- a/Engine/Source/Programs/UnrealBuildTool/IOS/UEBuildIOS.cs
+++ b/Engine/Source/Programs/UnrealBuildTool/IOS/UEBuildIOS.cs
@@ -411,32 +411,32 @@ namespace UnrealBuildTool
          *     This is not required - but allows for hiding details of a
          *     particular platform.
          *     
-         *     @param InModule             The newly loaded module
-         *     @param GameName             The game being build
-         *     @param Target               The target being build
-         */
-        public override void ModifyNewlyLoadedModule(UEBuildModule InModule, TargetInfo Target)
+		 *  @param  Name			The name of the module
+		 *	@param	Rules			The module rules
+		 *	@param	Target			The target being build
+		 */
+		public override void ModifyModuleRules(string ModuleName, ModuleRules Rules, TargetInfo Target)
         {
             if ((Target.Platform == UnrealTargetPlatform.Win32) || (Target.Platform == UnrealTargetPlatform.Win64) || (Target.Platform == UnrealTargetPlatform.Mac))
             {
                 bool bBuildShaderFormats = UEBuildConfiguration.bForceBuildShaderFormats;
                 if (!UEBuildConfiguration.bBuildRequiresCookedData)
                 {
-                    if (InModule.ToString() == "Engine")
+                    if (ModuleName == "Engine")
                     {
                         if (UEBuildConfiguration.bBuildDeveloperTools)
                         {
-                            InModule.AddPlatformSpecificDynamicallyLoadedModule("IOSTargetPlatform");
+                            Rules.PlatformSpecificDynamicallyLoadedModuleNames.Add("IOSTargetPlatform");
                         }
                     }
-                    else if (InModule.ToString() == "TargetPlatform")
+                    else if (ModuleName == "TargetPlatform")
                     {
                         bBuildShaderFormats = true;
-						InModule.AddDynamicallyLoadedModule("TextureFormatPVR");
-						InModule.AddDynamicallyLoadedModule("TextureFormatASTC");
+						Rules.DynamicallyLoadedModuleNames.Add("TextureFormatPVR");
+						Rules.DynamicallyLoadedModuleNames.Add("TextureFormatASTC");
                         if (UEBuildConfiguration.bBuildDeveloperTools)
                         {
-                            InModule.AddPlatformSpecificDynamicallyLoadedModule("AudioFormatADPCM");
+                            Rules.PlatformSpecificDynamicallyLoadedModuleNames.Add("AudioFormatADPCM");
                         }
                     }
                 }
@@ -444,12 +444,12 @@ namespace UnrealBuildTool
                 // allow standalone tools to use targetplatform modules, without needing Engine
                 if (UEBuildConfiguration.bForceBuildTargetPlatforms)
                 {
-                    InModule.AddPlatformSpecificDynamicallyLoadedModule("IOSTargetPlatform");
+                    Rules.PlatformSpecificDynamicallyLoadedModuleNames.Add("IOSTargetPlatform");
                 }
 
                 if (bBuildShaderFormats)
                 {
-                    InModule.AddPlatformSpecificDynamicallyLoadedModule("MetalShaderFormat");
+                    Rules.PlatformSpecificDynamicallyLoadedModuleNames.Add("MetalShaderFormat");
                 }
             }
         }
diff --git a/Engine/Source/Programs/UnrealBuildTool/Linux/UEBuildLinux.cs b/Engine/Source/Programs/UnrealBuildTool/Linux/UEBuildLinux.cs
--- a/Engine/Source/Programs/UnrealBuildTool/Linux/UEBuildLinux.cs
+++ b/Engine/Source/Programs/UnrealBuildTool/Linux/UEBuildLinux.cs
@@ -281,22 +281,23 @@ namespace UnrealBuildTool
          *	This is not required - but allows for hiding details of a
          *	particular platform.
          *	
-         *	@param	InModule		The newly loaded module
-         *	@param	Target			The target being build
-         */
-        public override void ModifyNewlyLoadedModule(UEBuildModule InModule, TargetInfo Target)
+		 *  @param  Name			The name of the module
+		 *	@param	Rules			The module rules
+		 *	@param	Target			The target being build
+		 */
+		public override void ModifyModuleRules(string ModuleName, ModuleRules Rules, TargetInfo Target)
         {
             if ((Target.Platform == UnrealTargetPlatform.Win32) || (Target.Platform == UnrealTargetPlatform.Win64))
             {
                 if (!UEBuildConfiguration.bBuildRequiresCookedData)
                 {
-                    if (InModule.ToString() == "Engine")
+                    if (ModuleName == "Engine")
                     {
                         if (UEBuildConfiguration.bBuildDeveloperTools)
                         {
-                            InModule.AddPlatformSpecificDynamicallyLoadedModule("LinuxTargetPlatform");
-                            InModule.AddPlatformSpecificDynamicallyLoadedModule("LinuxNoEditorTargetPlatform");
-                            InModule.AddPlatformSpecificDynamicallyLoadedModule("LinuxServerTargetPlatform");
+                            Rules.PlatformSpecificDynamicallyLoadedModuleNames.Add("LinuxTargetPlatform");
+                            Rules.PlatformSpecificDynamicallyLoadedModuleNames.Add("LinuxNoEditorTargetPlatform");
+                            Rules.PlatformSpecificDynamicallyLoadedModuleNames.Add("LinuxServerTargetPlatform");
                         }
                     }
                 }
@@ -304,9 +305,9 @@ namespace UnrealBuildTool
                 // allow standalone tools to use targetplatform modules, without needing Engine
                 if (UEBuildConfiguration.bForceBuildTargetPlatforms)
                 {
-                    InModule.AddPlatformSpecificDynamicallyLoadedModule("LinuxTargetPlatform");
-                    InModule.AddPlatformSpecificDynamicallyLoadedModule("LinuxNoEditorTargetPlatform");
-                    InModule.AddPlatformSpecificDynamicallyLoadedModule("LinuxServerTargetPlatform");
+                    Rules.PlatformSpecificDynamicallyLoadedModuleNames.Add("LinuxTargetPlatform");
+                    Rules.PlatformSpecificDynamicallyLoadedModuleNames.Add("LinuxNoEditorTargetPlatform");
+                    Rules.PlatformSpecificDynamicallyLoadedModuleNames.Add("LinuxServerTargetPlatform");
                 }
             }
             else if (Target.Platform == UnrealTargetPlatform.Linux)
@@ -315,7 +316,7 @@ namespace UnrealBuildTool
 
                 if (!UEBuildConfiguration.bBuildRequiresCookedData)
                 {
-                    if (InModule.ToString() == "TargetPlatform")
+                    if (ModuleName == "TargetPlatform")
                     {
                         bBuildShaderFormats = true;
                     }
@@ -324,16 +325,16 @@ namespace UnrealBuildTool
                 // allow standalone tools to use target platform modules, without needing Engine
                 if (UEBuildConfiguration.bForceBuildTargetPlatforms)
                 {
-                    InModule.AddDynamicallyLoadedModule("LinuxTargetPlatform");
-                    InModule.AddDynamicallyLoadedModule("LinuxNoEditorTargetPlatform");
-                    InModule.AddDynamicallyLoadedModule("LinuxServerTargetPlatform");
-					InModule.AddDynamicallyLoadedModule("AllDesktopTargetPlatform");
+                    Rules.DynamicallyLoadedModuleNames.Add("LinuxTargetPlatform");
+                    Rules.DynamicallyLoadedModuleNames.Add("LinuxNoEditorTargetPlatform");
+                    Rules.DynamicallyLoadedModuleNames.Add("LinuxServerTargetPlatform");
+					Rules.DynamicallyLoadedModuleNames.Add("AllDesktopTargetPlatform");
                 }
 
                 if (bBuildShaderFormats)
                 {
-					// InModule.AddDynamicallyLoadedModule("ShaderFormatD3D");
-                    InModule.AddDynamicallyLoadedModule("ShaderFormatOpenGL");
+					// Rules.DynamicallyLoadedModuleNames.Add("ShaderFormatD3D");
+                    Rules.DynamicallyLoadedModuleNames.Add("ShaderFormatOpenGL");
                 }
             }
         }
diff --git a/Engine/Source/Programs/UnrealBuildTool/Mac/UEBuildMac.cs b/Engine/Source/Programs/UnrealBuildTool/Mac/UEBuildMac.cs
--- a/Engine/Source/Programs/UnrealBuildTool/Mac/UEBuildMac.cs
+++ b/Engine/Source/Programs/UnrealBuildTool/Mac/UEBuildMac.cs
@@ -92,7 +92,7 @@ namespace UnrealBuildTool
 			return BuildConfiguration.bGeneratedSYMFile || BuildConfiguration.bUsePDBFiles ? ".dSYM" : "";
 		}
 
-		public override void ModifyNewlyLoadedModule(UEBuildModule InModule, TargetInfo Target)
+		public override void ModifyModuleRules(string ModuleName, ModuleRules Rules, TargetInfo Target)
 		{
 			if (Target.Platform == UnrealTargetPlatform.Mac)
 			{
@@ -100,7 +100,7 @@ namespace UnrealBuildTool
 
                 if (!UEBuildConfiguration.bBuildRequiresCookedData)
                 {
-                    if (InModule.ToString() == "TargetPlatform")
+                    if (ModuleName == "TargetPlatform")
                     {
                         bBuildShaderFormats = true;
                     }
@@ -109,17 +109,17 @@ namespace UnrealBuildTool
 				// allow standalone tools to use target platform modules, without needing Engine
 				if (UEBuildConfiguration.bForceBuildTargetPlatforms)
 				{
-					InModule.AddDynamicallyLoadedModule("MacTargetPlatform");
-					InModule.AddDynamicallyLoadedModule("MacNoEditorTargetPlatform");
-					InModule.AddDynamicallyLoadedModule("MacClientTargetPlatform");
-					InModule.AddDynamicallyLoadedModule("MacServerTargetPlatform");
-					InModule.AddDynamicallyLoadedModule("AllDesktopTargetPlatform");
+					Rules.DynamicallyLoadedModuleNames.Add("MacTargetPlatform");
+					Rules.DynamicallyLoadedModuleNames.Add("MacNoEditorTargetPlatform");
+					Rules.DynamicallyLoadedModuleNames.Add("MacClientTargetPlatform");
+					Rules.DynamicallyLoadedModuleNames.Add("MacServerTargetPlatform");
+					Rules.DynamicallyLoadedModuleNames.Add("AllDesktopTargetPlatform");
 				}
 
                 if (bBuildShaderFormats)
                 {
-					// InModule.AddDynamicallyLoadedModule("ShaderFormatD3D");
-                    InModule.AddDynamicallyLoadedModule("ShaderFormatOpenGL");
+					// Rules.DynamicallyLoadedModuleNames.Add("ShaderFormatD3D");
+                    Rules.DynamicallyLoadedModuleNames.Add("ShaderFormatOpenGL");
                 }
 			}
 		}
diff --git a/Engine/Source/Programs/UnrealBuildTool/UWP/UEBuildUWP.cs b/Engine/Source/Programs/UnrealBuildTool/UWP/UEBuildUWP.cs
--- a/Engine/Source/Programs/UnrealBuildTool/UWP/UEBuildUWP.cs
+++ b/Engine/Source/Programs/UnrealBuildTool/UWP/UEBuildUWP.cs
@@ -283,52 +283,51 @@ namespace UnrealBuildTool
 			UEBuildConfiguration.bCompileICU = true;
 		}
 
-
-		public override void ModifyNewlyLoadedModule(UEBuildModule InModule, TargetInfo Target)
+		public override void ModifyModuleRules(string ModuleName, ModuleRules Rules, TargetInfo Target)
 		{
 			if ((Target.Platform == UnrealTargetPlatform.UWP))
 			{
-				if (InModule.ToString() == "Core")
+				if (ModuleName == "Core")
 				{
-					InModule.AddPrivateDependencyModule("UWPSDK");
+					Rules.PrivateDependencyModuleNames.Add("UWPSDK");
 				}
-				else if (InModule.ToString() == "Engine")
+				else if (ModuleName == "Engine")
 				{
-					InModule.AddPrivateDependencyModule("zlib");
-					InModule.AddPrivateDependencyModule("UElibPNG");
-					InModule.AddPublicDependencyModule("UEOgg");
-					InModule.AddPublicDependencyModule("Vorbis");
+					Rules.PrivateDependencyModuleNames.Add("zlib");
+					Rules.PrivateDependencyModuleNames.Add("UElibPNG");
+					Rules.PublicDependencyModuleNames.Add("UEOgg");
+					Rules.PublicDependencyModuleNames.Add("Vorbis");
 				}
-				else if (InModule.ToString() == "D3D11RHI")
+				else if (ModuleName == "D3D11RHI")
 				{
-					InModule.AddPublicDefinition("D3D11_WITH_DWMAPI=0");
-					InModule.AddPublicDefinition("WITH_DX_PERF=0");
+					Rules.Definitions.Add("D3D11_WITH_DWMAPI=0");
+					Rules.Definitions.Add("WITH_DX_PERF=0");
 				}
-				else if (InModule.ToString() == "DX11")
+				else if (ModuleName == "DX11")
 				{
 					// Clear out all the Windows include paths and libraries...
 					// The UWPSDK module handles proper paths and libs for UWP.
 					// However, the D3D11RHI module will include the DX11 module.
-					InModule.ClearPublicIncludePaths();
-					InModule.ClearPublicLibraryPaths();
-					InModule.ClearPublicAdditionalLibraries();
-					InModule.RemovePublicDefinition("WITH_D3DX_LIBS=1");
-					InModule.AddPublicDefinition("WITH_D3DX_LIBS=0");
-					InModule.RemovePublicAdditionalLibrary("X3DAudio.lib");
-					InModule.RemovePublicAdditionalLibrary("XAPOFX.lib");
+					Rules.PublicIncludePaths.Clear();
+					Rules.PublicLibraryPaths.Clear();
+					Rules.PublicAdditionalLibraries.Clear();
+					Rules.Definitions.Remove("WITH_D3DX_LIBS=1");
+					Rules.Definitions.Add("WITH_D3DX_LIBS=0");
+					Rules.PublicAdditionalLibraries.Remove("X3DAudio.lib");
+					Rules.PublicAdditionalLibraries.Remove("XAPOFX.lib");
 				}
-				else if (InModule.ToString() == "XAudio2")
+				else if (ModuleName == "XAudio2")
 				{
-					InModule.AddPublicDefinition("XAUDIO_SUPPORTS_XMA2WAVEFORMATEX=0");
-					InModule.AddPublicDefinition("XAUDIO_SUPPORTS_DEVICE_DETAILS=0");
-					InModule.AddPublicDefinition("XAUDIO2_SUPPORTS_MUSIC=0");
-					InModule.AddPublicDefinition("XAUDIO2_SUPPORTS_SENDLIST=0");
-					InModule.AddPublicAdditionalLibrary("XAudio2.lib");
+					Rules.Definitions.Add("XAUDIO_SUPPORTS_XMA2WAVEFORMATEX=0");
+					Rules.Definitions.Add("XAUDIO_SUPPORTS_DEVICE_DETAILS=0");
+					Rules.Definitions.Add("XAUDIO2_SUPPORTS_MUSIC=0");
+					Rules.Definitions.Add("XAUDIO2_SUPPORTS_SENDLIST=0");
+					Rules.PublicAdditionalLibraries.Add("XAudio2.lib");
 				}
-				else if (InModule.ToString() == "DX11Audio")
+				else if (ModuleName == "DX11Audio")
 				{
-					InModule.RemovePublicAdditionalLibrary("X3DAudio.lib");
-					InModule.RemovePublicAdditionalLibrary("XAPOFX.lib");
+					Rules.PublicAdditionalLibraries.Remove("X3DAudio.lib");
+					Rules.PublicAdditionalLibraries.Remove("XAPOFX.lib");
 				}
 			}
 		}
diff --git a/Engine/Source/Programs/UnrealBuildTool/WinRT/UEBuildWinRT.cs b/Engine/Source/Programs/UnrealBuildTool/WinRT/UEBuildWinRT.cs
--- a/Engine/Source/Programs/UnrealBuildTool/WinRT/UEBuildWinRT.cs
+++ b/Engine/Source/Programs/UnrealBuildTool/WinRT/UEBuildWinRT.cs
@@ -246,238 +246,239 @@ namespace UnrealBuildTool
          *	This is not required - but allows for hiding details of a
          *	particular platform.
          *	
-         *	@param	InModule		The newly loaded module
-         *	@param	Target			The target being build
-         */
-        public override void ModifyNewlyLoadedModule(UEBuildModule InModule, TargetInfo Target)
+		 *  @param  Name			The name of the module
+		 *	@param	Rules			The module rules
+		 *	@param	Target			The target being build
+		 */
+		public override void ModifyModuleRules(string ModuleName, ModuleRules Rules, TargetInfo Target)
         {
             if ((Target.Platform == UnrealTargetPlatform.WinRT) || 
                 (Target.Platform == UnrealTargetPlatform.WinRT_ARM))
             {
-                if (InModule.ToString() == "Core")
+                if (ModuleName == "Core")
                 {
-                    InModule.AddPublicIncludePath("Runtime/Core/Public/WinRT");
-                    InModule.AddPublicDependencyModule("zlib");
+                    Rules.PublicIncludePaths.Add("Runtime/Core/Public/WinRT");
+                    Rules.PublicDependencyModuleNames.Add("zlib");
                 }
-                else if (InModule.ToString() == "Engine")
+                else if (ModuleName == "Engine")
                 {
-                    InModule.AddPrivateDependencyModule("zlib");
-                    InModule.AddPrivateDependencyModule("UElibPNG");
-                    InModule.AddPublicDependencyModule("UEOgg");
-                    InModule.AddPublicDependencyModule("Vorbis");
+                    Rules.PrivateDependencyModuleNames.Add("zlib");
+                    Rules.PrivateDependencyModuleNames.Add("UElibPNG");
+                    Rules.PublicDependencyModuleNames.Add("UEOgg");
+                    Rules.PublicDependencyModuleNames.Add("Vorbis");
                 }
-                else if (InModule.ToString() == "Launch")
+                else if (ModuleName == "Launch")
                 {
                 }
-                else if (InModule.ToString() == "D3D11RHI")
+                else if (ModuleName == "D3D11RHI")
                 {
-                    InModule.AddPublicDefinition("D3D11_CUSTOM_VIEWPORT_CONSTRUCTOR=1");
+                    Rules.Definitions.Add("D3D11_CUSTOM_VIEWPORT_CONSTRUCTOR=1");
                     // To enable platform specific D3D11 RHI Types
-                    InModule.AddPrivateIncludePath("Runtime/Windows/D3D11RHI/Private/WinRT");
+                    Rules.PrivateIncludePaths.Add("Runtime/Windows/D3D11RHI/Private/WinRT");
                     // Hack to enable AllowWindowsPlatformTypes.h/HideWindowsPlatformTypes.h
-                    InModule.AddPublicIncludePath("Runtime/Core/Public/Windows");
+                    Rules.PublicIncludePaths.Add("Runtime/Core/Public/Windows");
                 }
-                else if (InModule.ToString() == "Sockets")
+                else if (ModuleName == "Sockets")
                 {
                     // Hack to enable AllowWindowsPlatformTypes.h/HideWindowsPlatformTypes.h
-                    InModule.AddPublicIncludePath("Runtime/Core/Public/Windows");
+                    Rules.PublicIncludePaths.Add("Runtime/Core/Public/Windows");
                 }
-                else if (InModule.ToString() == "PhysX")
+                else if (ModuleName == "PhysX")
                 {
                     string PhysXDir = UEBuildConfiguration.UEThirdPartySourceDirectory + "PhysX/PhysX-3.3/";
 
-                    InModule.AddPublicIncludePath("include/foundation/WinRT");
+                    Rules.PublicIncludePaths.Add("include/foundation/WinRT");
                     if (Target.Platform == UnrealTargetPlatform.WinRT)
                     {
-                        InModule.AddPublicLibraryPath(PhysXDir + "Lib/WinRT");
+                        Rules.PublicLibraryPaths.Add(PhysXDir + "Lib/WinRT");
                     }
                     else
                     {
-                        InModule.AddPublicLibraryPath(PhysXDir + "Lib/WinRT/ARM");
+                        Rules.PublicLibraryPaths.Add(PhysXDir + "Lib/WinRT/ARM");
                     }
 
                     if (Target.Configuration == UnrealTargetConfiguration.Debug)
                     {
-                        InModule.AddPublicAdditionalLibrary("PhysX3DEBUG.lib");
-                        InModule.AddPublicAdditionalLibrary("PhysX3ExtensionsDEBUG.lib");
-                        InModule.AddPublicAdditionalLibrary("PhysX3CookingDEBUG.lib");
-                        InModule.AddPublicAdditionalLibrary("PhysX3CommonDEBUG.lib");
-                        InModule.AddPublicAdditionalLibrary("PhysX3VehicleDEBUG.lib");
-                        InModule.AddPublicAdditionalLibrary("PxTaskDEBUG.lib");
-                        InModule.AddPublicAdditionalLibrary("PhysXVisualDebuggerSDKDEBUG.lib");
-                        InModule.AddPublicAdditionalLibrary("PhysXProfileSDKDEBUG.lib");
+                        Rules.PublicAdditionalLibraries.Add("PhysX3DEBUG.lib");
+                        Rules.PublicAdditionalLibraries.Add("PhysX3ExtensionsDEBUG.lib");
+                        Rules.PublicAdditionalLibraries.Add("PhysX3CookingDEBUG.lib");
+                        Rules.PublicAdditionalLibraries.Add("PhysX3CommonDEBUG.lib");
+                        Rules.PublicAdditionalLibraries.Add("PhysX3VehicleDEBUG.lib");
+                        Rules.PublicAdditionalLibraries.Add("PxTaskDEBUG.lib");
+                        Rules.PublicAdditionalLibraries.Add("PhysXVisualDebuggerSDKDEBUG.lib");
+                        Rules.PublicAdditionalLibraries.Add("PhysXProfileSDKDEBUG.lib");
                     }
                     else if (Target.Configuration == UnrealTargetConfiguration.Development)
                     {
-                        InModule.AddPublicAdditionalLibrary("PhysX3PROFILE.lib");
-                        InModule.AddPublicAdditionalLibrary("PhysX3ExtensionsPROFILE.lib");
-                        InModule.AddPublicAdditionalLibrary("PhysX3CookingPROFILE.lib");
-                        InModule.AddPublicAdditionalLibrary("PhysX3CommonPROFILE.lib");
-                        InModule.AddPublicAdditionalLibrary("PhysX3VehiclePROFILE.lib");
-                        InModule.AddPublicAdditionalLibrary("PxTaskPROFILE.lib");
-                        InModule.AddPublicAdditionalLibrary("PhysXVisualDebuggerSDKPROFILE.lib");
-                        InModule.AddPublicAdditionalLibrary("PhysXProfileSDKPROFILE.lib");
+                        Rules.PublicAdditionalLibraries.Add("PhysX3PROFILE.lib");
+                        Rules.PublicAdditionalLibraries.Add("PhysX3ExtensionsPROFILE.lib");
+                        Rules.PublicAdditionalLibraries.Add("PhysX3CookingPROFILE.lib");
+                        Rules.PublicAdditionalLibraries.Add("PhysX3CommonPROFILE.lib");
+                        Rules.PublicAdditionalLibraries.Add("PhysX3VehiclePROFILE.lib");
+                        Rules.PublicAdditionalLibraries.Add("PxTaskPROFILE.lib");
+                        Rules.PublicAdditionalLibraries.Add("PhysXVisualDebuggerSDKPROFILE.lib");
+                        Rules.PublicAdditionalLibraries.Add("PhysXProfileSDKPROFILE.lib");
                     }
                     else // Test or Shipping
                     {
-                        InModule.AddPublicAdditionalLibrary("PhysX3.lib");
-                        InModule.AddPublicAdditionalLibrary("PhysX3Extensions.lib");
-                        InModule.AddPublicAdditionalLibrary("PhysX3Cooking.lib");
-                        InModule.AddPublicAdditionalLibrary("PhysX3Common.lib");
-                        InModule.AddPublicAdditionalLibrary("PhysX3Vehicle.lib");
-                        InModule.AddPublicAdditionalLibrary("PxTask.lib");
-                        InModule.AddPublicAdditionalLibrary("PhysXVisualDebuggerSDK.lib");
-                        InModule.AddPublicAdditionalLibrary("PhysXProfileSDK.lib");
+                        Rules.PublicAdditionalLibraries.Add("PhysX3.lib");
+                        Rules.PublicAdditionalLibraries.Add("PhysX3Extensions.lib");
+                        Rules.PublicAdditionalLibraries.Add("PhysX3Cooking.lib");
+                        Rules.PublicAdditionalLibraries.Add("PhysX3Common.lib");
+                        Rules.PublicAdditionalLibraries.Add("PhysX3Vehicle.lib");
+                        Rules.PublicAdditionalLibraries.Add("PxTask.lib");
+                        Rules.PublicAdditionalLibraries.Add("PhysXVisualDebuggerSDK.lib");
+                        Rules.PublicAdditionalLibraries.Add("PhysXProfileSDK.lib");
                     }
                 }
-                else if (InModule.ToString() == "APEX")
+                else if (ModuleName == "APEX")
                 {
-                    InModule.RemovePublicDefinition("APEX_STATICALLY_LINKED=0");
-                    InModule.AddPublicDefinition("APEX_STATICALLY_LINKED=1");
+                    Rules.Definitions.Remove("APEX_STATICALLY_LINKED=0");
+                    Rules.Definitions.Add("APEX_STATICALLY_LINKED=1");
                     
                     string APEXDir = UEBuildConfiguration.UEThirdPartySourceDirectory + "PhysX/APEX-1.3/";
                     if (Target.Platform == UnrealTargetPlatform.WinRT)
                     {
-                        InModule.AddPublicLibraryPath(APEXDir + "lib/WinRT");
+                        Rules.PublicLibraryPaths.Add(APEXDir + "lib/WinRT");
                     }
                     else
                     {
-                        InModule.AddPublicLibraryPath(APEXDir + "lib/WinRT/ARM");
+                        Rules.PublicLibraryPaths.Add(APEXDir + "lib/WinRT/ARM");
                     }
 
                     if (Target.Configuration == UnrealTargetConfiguration.Debug)
                     {
-                        InModule.AddPublicAdditionalLibrary("ApexCommonDEBUG.lib");
-                        InModule.AddPublicAdditionalLibrary("ApexFrameworkDEBUG.lib");
-                        InModule.AddPublicAdditionalLibrary("ApexSharedDEBUG.lib");
-                        InModule.AddPublicAdditionalLibrary("APEX_DestructibleDEBUG.lib");
+                        Rules.PublicAdditionalLibraries.Add("ApexCommonDEBUG.lib");
+                        Rules.PublicAdditionalLibraries.Add("ApexFrameworkDEBUG.lib");
+                        Rules.PublicAdditionalLibraries.Add("ApexSharedDEBUG.lib");
+                        Rules.PublicAdditionalLibraries.Add("APEX_DestructibleDEBUG.lib");
 
                     }
                     else if (Target.Configuration == UnrealTargetConfiguration.Development)
                     {
-                        InModule.AddPublicAdditionalLibrary("ApexCommonPROFILE.lib");
-                        InModule.AddPublicAdditionalLibrary("ApexFrameworkPROFILE.lib");
-                        InModule.AddPublicAdditionalLibrary("ApexSharedPROFILE.lib");
-                        InModule.AddPublicAdditionalLibrary("APEX_DestructiblePROFILE.lib");
+                        Rules.PublicAdditionalLibraries.Add("ApexCommonPROFILE.lib");
+                        Rules.PublicAdditionalLibraries.Add("ApexFrameworkPROFILE.lib");
+                        Rules.PublicAdditionalLibraries.Add("ApexSharedPROFILE.lib");
+                        Rules.PublicAdditionalLibraries.Add("APEX_DestructiblePROFILE.lib");
                     }
                     else // Test or Shipping
                     {
-                        InModule.AddPublicAdditionalLibrary("ApexCommon.lib");
-                        InModule.AddPublicAdditionalLibrary("ApexFramework.lib");
-                        InModule.AddPublicAdditionalLibrary("ApexShared.lib");
-                        InModule.AddPublicAdditionalLibrary("APEX_Destructible.lib");
+                        Rules.PublicAdditionalLibraries.Add("ApexCommon.lib");
+                        Rules.PublicAdditionalLibraries.Add("ApexFramework.lib");
+                        Rules.PublicAdditionalLibraries.Add("ApexShared.lib");
+                        Rules.PublicAdditionalLibraries.Add("APEX_Destructible.lib");
                     }
                 }
-                else if (InModule.ToString() == "FreeType2")
+                else if (ModuleName == "FreeType2")
                 {
                     string FreeType2Path = UEBuildConfiguration.UEThirdPartySourceDirectory + "FreeType2/FreeType2-2.4.12/";
                     if (Target.Platform == UnrealTargetPlatform.WinRT)
                     {
-                        InModule.AddPublicLibraryPath(FreeType2Path + "Lib/WinRT/Win64");
+                        Rules.PublicLibraryPaths.Add(FreeType2Path + "Lib/WinRT/Win64");
                     }
                     else
                     {
-                        InModule.AddPublicLibraryPath(FreeType2Path + "Lib/WinRT/ARM");
+                        Rules.PublicLibraryPaths.Add(FreeType2Path + "Lib/WinRT/ARM");
                     }
-                    InModule.AddPublicAdditionalLibrary("freetype2412MT.lib");
+                    Rules.PublicAdditionalLibraries.Add("freetype2412MT.lib");
                 }
-                else if (InModule.ToString() == "UElibPNG")
+                else if (ModuleName == "UElibPNG")
                 {
                     string libPNGPath = UEBuildConfiguration.UEThirdPartySourceDirectory + "libPNG/libPNG-1.5.2";
                     if (Target.Platform == UnrealTargetPlatform.WinRT)
                     {
-                        InModule.AddPublicLibraryPath(libPNGPath + "/lib/WinRT/Win64");
+                        Rules.PublicLibraryPaths.Add(libPNGPath + "/lib/WinRT/Win64");
                     }
                     else
                     {
-                        InModule.AddPublicLibraryPath(libPNGPath + "/lib/WinRT/ARM");
+                        Rules.PublicLibraryPaths.Add(libPNGPath + "/lib/WinRT/ARM");
                     }
-                    InModule.AddPublicAdditionalLibrary("libpng125.lib");
+                    Rules.PublicAdditionalLibraries.Add("libpng125.lib");
                 }
-                else if (InModule.ToString() == "DX11")
+                else if (ModuleName == "DX11")
                 {
                     // Clear out all the Windows include paths and libraries...
                     // The WinRTSDK module handles proper paths and libs for WinRT.
                     // However, the D3D11RHI module will include the DX11 module.
-                    InModule.ClearPublicIncludePaths();
-                    InModule.ClearPublicLibraryPaths();
-                    InModule.ClearPublicAdditionalLibraries();
-                    InModule.RemovePublicDefinition("WITH_D3DX_LIBS=1");
-                    InModule.AddPublicDefinition("D3D11_WITH_DWMAPI=0");
-                    InModule.AddPublicDefinition("WITH_D3DX_LIBS=0");
-                    InModule.AddPublicDefinition("WITH_DX_PERF=0");
-                    InModule.RemovePublicAdditionalLibrary("X3DAudio.lib");
-                    InModule.RemovePublicAdditionalLibrary("XAPOFX.lib");
+                    Rules.PublicIncludePaths.Clear();
+                    Rules.PublicLibraryPaths.Clear();
+                    Rules.PublicAdditionalLibraries.Clear();
+                    Rules.Definitions.Remove("WITH_D3DX_LIBS=1");
+                    Rules.Definitions.Add("D3D11_WITH_DWMAPI=0");
+                    Rules.Definitions.Add("WITH_D3DX_LIBS=0");
+                    Rules.Definitions.Add("WITH_DX_PERF=0");
+                    Rules.PublicAdditionalLibraries.Remove("X3DAudio.lib");
+                    Rules.PublicAdditionalLibraries.Remove("XAPOFX.lib");
                 }
-                else if (InModule.ToString() == "XInput")
+                else if (ModuleName == "XInput")
                 {
-                    InModule.AddPublicAdditionalLibrary("XInput.lib");
+                    Rules.PublicAdditionalLibraries.Add("XInput.lib");
                 }
-                else if (InModule.ToString() == "XAudio2")
+                else if (ModuleName == "XAudio2")
                 {
-                    InModule.AddPublicDefinition("XAUDIO_SUPPORTS_XMA2WAVEFORMATEX=0");
-                    InModule.AddPublicDefinition("XAUDIO_SUPPORTS_DEVICE_DETAILS=0");
-                    InModule.AddPublicDefinition("XAUDIO2_SUPPORTS_MUSIC=0");
-                    InModule.AddPublicDefinition("XAUDIO2_SUPPORTS_SENDLIST=0");
-                    InModule.AddPublicAdditionalLibrary("XAudio2.lib");
+                    Rules.Definitions.Add("XAUDIO_SUPPORTS_XMA2WAVEFORMATEX=0");
+                    Rules.Definitions.Add("XAUDIO_SUPPORTS_DEVICE_DETAILS=0");
+                    Rules.Definitions.Add("XAUDIO2_SUPPORTS_MUSIC=0");
+                    Rules.Definitions.Add("XAUDIO2_SUPPORTS_SENDLIST=0");
+                    Rules.PublicAdditionalLibraries.Add("XAudio2.lib");
                     // Hack to enable AllowWindowsPlatformTypes.h/HideWindowsPlatformTypes.h
-                    InModule.AddPublicIncludePath("Runtime/Core/Public/Windows");
+                    Rules.PublicIncludePaths.Add("Runtime/Core/Public/Windows");
                 }
-                else if (InModule.ToString() == "UEOgg")
+                else if (ModuleName == "UEOgg")
                 {
                     string OggPath = UEBuildConfiguration.UEThirdPartySourceDirectory + "Ogg/libogg-1.2.2/";
                     if (Target.Platform == UnrealTargetPlatform.WinRT)
                     {
-                        InModule.AddPublicLibraryPath(OggPath + "WinRT/VS2012/WinRT/x64/Release");
+                        Rules.PublicLibraryPaths.Add(OggPath + "WinRT/VS2012/WinRT/x64/Release");
                     }
                     else
                     {
-                        InModule.AddPublicLibraryPath(OggPath + "WinRT/VS2012/WinRT/ARM/Release");
+                        Rules.PublicLibraryPaths.Add(OggPath + "WinRT/VS2012/WinRT/ARM/Release");
                     }
-                    InModule.AddPublicAdditionalLibrary("libogg_static.lib");
+                    Rules.PublicAdditionalLibraries.Add("libogg_static.lib");
                 }
-                else if (InModule.ToString() == "Vorbis")
+                else if (ModuleName == "Vorbis")
                 {
                     string VorbisPath = UEBuildConfiguration.UEThirdPartySourceDirectory + "Vorbis/libvorbis-1.3.2/";
                     if (Target.Platform == UnrealTargetPlatform.WinRT)
                     {
-                        InModule.AddPublicLibraryPath(VorbisPath + "WinRT/VS2012/WinRT/x64/Release");
+                        Rules.PublicLibraryPaths.Add(VorbisPath + "WinRT/VS2012/WinRT/x64/Release");
                     }
                     else
                     {
-                        InModule.AddPublicLibraryPath(VorbisPath + "WinRT/VS2012/WinRT/ARM/Release");
+                        Rules.PublicLibraryPaths.Add(VorbisPath + "WinRT/VS2012/WinRT/ARM/Release");
                     }
-                    InModule.AddPublicAdditionalLibrary("libvorbis_static.lib");
+                    Rules.PublicAdditionalLibraries.Add("libvorbis_static.lib");
                 }
-                else if (InModule.ToString() == "VorbisFile")
+                else if (ModuleName == "VorbisFile")
                 {
                     string VorbisPath = UEBuildConfiguration.UEThirdPartySourceDirectory + "Vorbis/libvorbis-1.3.2/";
                     if (Target.Platform == UnrealTargetPlatform.WinRT)
                     {
-                        InModule.AddPublicLibraryPath(VorbisPath + "WinRT/VS2012/WinRT/x64/Release");
+                        Rules.PublicLibraryPaths.Add(VorbisPath + "WinRT/VS2012/WinRT/x64/Release");
                     }
                     else
                     {
-                        InModule.AddPublicLibraryPath(VorbisPath + "WinRT/VS2012/WinRT/ARM/Release");
+                        Rules.PublicLibraryPaths.Add(VorbisPath + "WinRT/VS2012/WinRT/ARM/Release");
                     }
-                    InModule.AddPublicAdditionalLibrary("libvorbisfile_static.lib");
+                    Rules.PublicAdditionalLibraries.Add("libvorbisfile_static.lib");
                 }
-                else if (InModule.ToString() == "DX11Audio")
+                else if (ModuleName == "DX11Audio")
                 {
-                    InModule.RemovePublicAdditionalLibrary("X3DAudio.lib");
-                    InModule.RemovePublicAdditionalLibrary("XAPOFX.lib");
+                    Rules.PublicAdditionalLibraries.Remove("X3DAudio.lib");
+                    Rules.PublicAdditionalLibraries.Remove("XAPOFX.lib");
                 }
-                else if (InModule.ToString() == "zlib")
+                else if (ModuleName == "zlib")
                 {
                     if (Target.Platform == UnrealTargetPlatform.WinRT)
                     {
-                        InModule.AddPublicLibraryPath(UEBuildConfiguration.UEThirdPartySourceDirectory + "zlib/zlib-1.2.5/Lib/WinRT/Win64");
+                        Rules.PublicLibraryPaths.Add(UEBuildConfiguration.UEThirdPartySourceDirectory + "zlib/zlib-1.2.5/Lib/WinRT/Win64");
                     }
                     else
                     {
-                        InModule.AddPublicLibraryPath(UEBuildConfiguration.UEThirdPartySourceDirectory + "zlib/zlib-1.2.5/Lib/WinRT/ARM");
+                        Rules.PublicLibraryPaths.Add(UEBuildConfiguration.UEThirdPartySourceDirectory + "zlib/zlib-1.2.5/Lib/WinRT/ARM");
                     }
-                    InModule.AddPublicAdditionalLibrary("zlib125.lib");
+                    Rules.PublicAdditionalLibraries.Add("zlib125.lib");
                 }
             }
             else if ((Target.Platform == UnrealTargetPlatform.Win32) || (Target.Platform == UnrealTargetPlatform.Win64))
@@ -485,14 +486,14 @@ namespace UnrealBuildTool
 //              bool bBuildShaderFormats = UEBuildConfiguration.bForceBuildShaderFormats;
 // 				if (!UEBuildConfiguration.bBuildRequiresCookedData)
 // 				{
-// 					if (InModule.ToString() == "Engine")
+// 					if (ModuleName == "Engine")
 // 					{
 // 						if (UEBuildConfiguration.bBuildDeveloperTools)
 // 						{
-                // 							InModule.AddPlatformSpecificDynamicallyLoadedModule("WinRTTargetPlatform");
+                // 							Rules.PlatformSpecificDynamicallyLoadedModuleNames.Add("WinRTTargetPlatform");
 // 						}
 // 					}
-// 					else if (InModule.ToString() == "TargetPlatform")
+// 					else if (ModuleName == "TargetPlatform")
 // 					{
 // 		                bBuildShaderFormats = true;
 // 					}
@@ -501,12 +502,12 @@ namespace UnrealBuildTool
 // 				// allow standalone tools to use targetplatform modules, without needing Engine
 // 				if (UEBuildConfiguration.bForceBuildTargetPlatforms)
 // 				{
-                // 					InModule.AddPlatformSpecificDynamicallyLoadedModule("WinRTTargetPlatform");
+                // 					Rules.PlatformSpecificDynamicallyLoadedModuleNames.Add("WinRTTargetPlatform");
 // 				}
 
 //              if (bBuildShaderFormats)
 //              {
-                //                  InModule.AddPlatformSpecificDynamicallyLoadedModule("ShaderFormatWinRT");
+                //                  Rules.PlatformSpecificDynamicallyLoadedModuleNames.Add("ShaderFormatWinRT");
 //              }
             }
         }
diff --git a/Engine/Source/Programs/UnrealBuildTool/Windows/UEBuildWindows.cs b/Engine/Source/Programs/UnrealBuildTool/Windows/UEBuildWindows.cs
--- a/Engine/Source/Programs/UnrealBuildTool/Windows/UEBuildWindows.cs
+++ b/Engine/Source/Programs/UnrealBuildTool/Windows/UEBuildWindows.cs
@@ -480,7 +480,7 @@ namespace UnrealBuildTool
             UEBuildConfiguration.bCompileICU = true;
         }
 
-        public override void ModifyNewlyLoadedModule(UEBuildModule InModule, TargetInfo Target)
+		public override void ModifyModuleRules(string ModuleName, ModuleRules Rules, TargetInfo Target)
         {
             if ((Target.Platform == UnrealTargetPlatform.Win32) || (Target.Platform == UnrealTargetPlatform.Win64))
             {
@@ -488,7 +488,7 @@ namespace UnrealBuildTool
 
                 if (!UEBuildConfiguration.bBuildRequiresCookedData)
                 {
-                    if (InModule.ToString() == "TargetPlatform")
+                    if (ModuleName == "TargetPlatform")
                     {
                         bBuildShaderFormats = true;
                     }
@@ -497,38 +497,38 @@ namespace UnrealBuildTool
                 // allow standalone tools to use target platform modules, without needing Engine
                 if (UEBuildConfiguration.bForceBuildTargetPlatforms)
                 {
-                    InModule.AddDynamicallyLoadedModule("WindowsTargetPlatform");
-                    InModule.AddDynamicallyLoadedModule("WindowsNoEditorTargetPlatform");
-                    InModule.AddDynamicallyLoadedModule("WindowsServerTargetPlatform");
-                    InModule.AddDynamicallyLoadedModule("WindowsClientTargetPlatform");
-					InModule.AddDynamicallyLoadedModule("AllDesktopTargetPlatform");
+                    Rules.DynamicallyLoadedModuleNames.Add("WindowsTargetPlatform");
+                    Rules.DynamicallyLoadedModuleNames.Add("WindowsNoEditorTargetPlatform");
+                    Rules.DynamicallyLoadedModuleNames.Add("WindowsServerTargetPlatform");
+                    Rules.DynamicallyLoadedModuleNames.Add("WindowsClientTargetPlatform");
+					Rules.DynamicallyLoadedModuleNames.Add("AllDesktopTargetPlatform");
                 }
 
                 if (bBuildShaderFormats)
                 {
-                    InModule.AddDynamicallyLoadedModule("ShaderFormatD3D");
-                    InModule.AddDynamicallyLoadedModule("ShaderFormatOpenGL");
+                    Rules.DynamicallyLoadedModuleNames.Add("ShaderFormatD3D");
+                    Rules.DynamicallyLoadedModuleNames.Add("ShaderFormatOpenGL");
 
 //#todo-rco: Remove when public
 					{
 						string VulkanSDKPath = Environment.GetEnvironmentVariable("VK_SDK_PATH");
 						if (!String.IsNullOrEmpty(VulkanSDKPath))
 						{
-							InModule.AddDynamicallyLoadedModule("VulkanShaderFormat");
+							Rules.DynamicallyLoadedModuleNames.Add("VulkanShaderFormat");
 						}
 					}
 				}
 
-                if (InModule.ToString() == "D3D11RHI")
+                if (ModuleName == "D3D11RHI")
                 {
                     // To enable platform specific D3D11 RHI Types
-                    InModule.AddPrivateIncludePath("Runtime/Windows/D3D11RHI/Private/Windows");
+                    Rules.PrivateIncludePaths.Add("Runtime/Windows/D3D11RHI/Private/Windows");
                 }
 
-				if (InModule.ToString() == "D3D12RHI")
+				if (ModuleName == "D3D12RHI")
 				{
 					// To enable platform specific D3D12 RHI Types
-					InModule.AddPrivateIncludePath("Runtime/Windows/D3D12RHI/Private/Windows");
+					Rules.PrivateIncludePaths.Add("Runtime/Windows/D3D12RHI/Private/Windows");
 				}
 			}
         }
