Pass the module rules into the constructor module instances, rather than passing all the parameters individually. The number of things that get passed around is getting silly, and keeping the two things as ideologically separate is time consuming and hard to maintain.

diff --git a/Engine/Source/Programs/UnrealBuildTool/Configuration/UEBuildModule.cs b/Engine/Source/Programs/UnrealBuildTool/Configuration/UEBuildModule.cs
--- a/Engine/Source/Programs/UnrealBuildTool/Configuration/UEBuildModule.cs
+++ b/Engine/Source/Programs/UnrealBuildTool/Configuration/UEBuildModule.cs
@@ -169,6 +169,9 @@ namespace UnrealBuildTool
 		/** The type of module being built. Used to switch between debug/development and precompiled/source configurations. */
 		public UEBuildModuleType Type;
 
+		/** The rules for this module */
+		public ModuleRules Rules;
+
 		/** The distribution level of the module being built. Used to check where build products are placed. */
 		public UEBuildModuleDistribution DistributionLevel;
 
@@ -248,55 +251,36 @@ namespace UnrealBuildTool
 			string InName,
 			UEBuildModuleType InType,
 			string InModuleDirectory,
-			bool? InIsRedistributableOverride,
-			IEnumerable<string> InPublicDefinitions,
-			IEnumerable<string> InPublicIncludePaths,
-			IEnumerable<string> InPublicSystemIncludePaths,
-			IEnumerable<string> InPublicLibraryPaths,
-			IEnumerable<string> InPublicAdditionalLibraries,
-			IEnumerable<string> InPublicFrameworks,
-			IEnumerable<string> InPublicWeakFrameworks,
-			IEnumerable<UEBuildFramework> InPublicAdditionalFrameworks,
-			IEnumerable<string> InPublicAdditionalShadowFiles,
-			IEnumerable<UEBuildBundleResource> InPublicAdditionalBundleResources,
-			IEnumerable<string> InPublicIncludePathModuleNames,
-			IEnumerable<string> InPublicDependencyModuleNames,
-			IEnumerable<string> InPublicDelayLoadDLLs,
-			IEnumerable<string> InPrivateIncludePaths,
-			IEnumerable<string> InPrivateIncludePathModuleNames,
-			IEnumerable<string> InPrivateDependencyModuleNames,
-			IEnumerable<string> InCircularlyReferencedDependentModules,
-			IEnumerable<string> InDynamicallyLoadedModuleNames,
-            IEnumerable<string> InPlatformSpecificDynamicallyLoadedModuleNames,
-			IEnumerable<RuntimeDependency> InRuntimeDependencies,
+			ModuleRules InRules,
 			string InBuildCsFilename
 			)
 		{
 			Target = InTarget;
 			Name = InName;
 			Type = InType;
+			Rules = InRules;
 			ModuleDirectory = InModuleDirectory;
-			PublicDefinitions = HashSetFromOptionalEnumerableStringParameter(InPublicDefinitions);
-			PublicIncludePaths = HashSetFromOptionalEnumerableStringParameter(InPublicIncludePaths);
-			PublicSystemIncludePaths = HashSetFromOptionalEnumerableStringParameter(InPublicSystemIncludePaths);
-			PublicLibraryPaths = HashSetFromOptionalEnumerableStringParameter(InPublicLibraryPaths);
-			PublicAdditionalLibraries = HashSetFromOptionalEnumerableStringParameter(InPublicAdditionalLibraries);
-			PublicFrameworks = HashSetFromOptionalEnumerableStringParameter(InPublicFrameworks);
-			PublicWeakFrameworks = HashSetFromOptionalEnumerableStringParameter(InPublicWeakFrameworks);
-			PublicAdditionalFrameworks = InPublicAdditionalFrameworks == null ? new HashSet<UEBuildFramework>() : new HashSet<UEBuildFramework>(InPublicAdditionalFrameworks);
-			PublicAdditionalShadowFiles = HashSetFromOptionalEnumerableStringParameter(InPublicAdditionalShadowFiles);
-			PublicAdditionalBundleResources = InPublicAdditionalBundleResources == null ? new HashSet<UEBuildBundleResource>() : new HashSet<UEBuildBundleResource>(InPublicAdditionalBundleResources);
-			PublicIncludePathModuleNames = HashSetFromOptionalEnumerableStringParameter( InPublicIncludePathModuleNames );
-			PublicDependencyModuleNames = HashSetFromOptionalEnumerableStringParameter(InPublicDependencyModuleNames);
-			PublicDelayLoadDLLs = HashSetFromOptionalEnumerableStringParameter(InPublicDelayLoadDLLs);
-			PrivateIncludePaths = HashSetFromOptionalEnumerableStringParameter(InPrivateIncludePaths);
-			PrivateIncludePathModuleNames = HashSetFromOptionalEnumerableStringParameter( InPrivateIncludePathModuleNames );
-			PrivateDependencyModuleNames = HashSetFromOptionalEnumerableStringParameter( InPrivateDependencyModuleNames );
-			CircularlyReferencedDependentModules = new HashSet<string>( HashSetFromOptionalEnumerableStringParameter( InCircularlyReferencedDependentModules ) );
-			DynamicallyLoadedModuleNames = HashSetFromOptionalEnumerableStringParameter( InDynamicallyLoadedModuleNames );
-            PlatformSpecificDynamicallyLoadedModuleNames = HashSetFromOptionalEnumerableStringParameter(InPlatformSpecificDynamicallyLoadedModuleNames);
-			RuntimeDependencies = (InRuntimeDependencies == null)? new List<RuntimeDependency>() : new List<RuntimeDependency>(InRuntimeDependencies);
-			IsRedistributableOverride = InIsRedistributableOverride;
+			PublicDefinitions = HashSetFromOptionalEnumerableStringParameter(InRules.Definitions);
+			PublicIncludePaths = HashSetFromOptionalEnumerableStringParameter(InRules.PublicIncludePaths);
+			PublicSystemIncludePaths = HashSetFromOptionalEnumerableStringParameter(InRules.PublicSystemIncludePaths);
+			PublicLibraryPaths = HashSetFromOptionalEnumerableStringParameter(InRules.PublicLibraryPaths);
+			PublicAdditionalLibraries = HashSetFromOptionalEnumerableStringParameter(InRules.PublicAdditionalLibraries);
+			PublicFrameworks = HashSetFromOptionalEnumerableStringParameter(InRules.PublicFrameworks);
+			PublicWeakFrameworks = HashSetFromOptionalEnumerableStringParameter(InRules.PublicWeakFrameworks);
+			PublicAdditionalFrameworks = InRules.PublicAdditionalFrameworks == null ? new HashSet<UEBuildFramework>() : new HashSet<UEBuildFramework>(InRules.PublicAdditionalFrameworks);
+			PublicAdditionalShadowFiles = HashSetFromOptionalEnumerableStringParameter(InRules.PublicAdditionalShadowFiles);
+			PublicAdditionalBundleResources = InRules.AdditionalBundleResources == null ? new HashSet<UEBuildBundleResource>() : new HashSet<UEBuildBundleResource>(InRules.AdditionalBundleResources);
+			PublicIncludePathModuleNames = HashSetFromOptionalEnumerableStringParameter( InRules.PublicIncludePathModuleNames );
+			PublicDependencyModuleNames = HashSetFromOptionalEnumerableStringParameter(InRules.PublicDependencyModuleNames);
+			PublicDelayLoadDLLs = HashSetFromOptionalEnumerableStringParameter(InRules.PublicDelayLoadDLLs);
+			PrivateIncludePaths = HashSetFromOptionalEnumerableStringParameter(InRules.PrivateIncludePaths);
+			PrivateIncludePathModuleNames = HashSetFromOptionalEnumerableStringParameter( InRules.PrivateIncludePathModuleNames );
+			PrivateDependencyModuleNames = HashSetFromOptionalEnumerableStringParameter( InRules.PrivateDependencyModuleNames );
+			CircularlyReferencedDependentModules = new HashSet<string>( HashSetFromOptionalEnumerableStringParameter( InRules.CircularlyReferencedDependentModules ) );
+			DynamicallyLoadedModuleNames = HashSetFromOptionalEnumerableStringParameter( InRules.DynamicallyLoadedModuleNames );
+            PlatformSpecificDynamicallyLoadedModuleNames = HashSetFromOptionalEnumerableStringParameter(InRules.PlatformSpecificDynamicallyLoadedModuleNames);
+			RuntimeDependencies = (InRules.RuntimeDependencies == null)? new List<RuntimeDependency>() : new List<RuntimeDependency>(InRules.RuntimeDependencies);
+			IsRedistributableOverride = InRules.IsRedistributableOverride;
 
 			Debug.Assert(InBuildCsFilename == null || InBuildCsFilename.EndsWith(".Build.cs", StringComparison.InvariantCultureIgnoreCase));
 			BuildCsFilenameField = InBuildCsFilename;
@@ -647,20 +631,7 @@ namespace UnrealBuildTool
 			UEBuildModuleType InType,
 			string InName,
 			string InModuleDirectory,
-			bool? InIsRedistributableOverride,
-			IEnumerable<string> InPublicDefinitions,
-			IEnumerable<string> InPublicIncludePaths,
-			IEnumerable<string> InPublicSystemIncludePaths,
-			IEnumerable<string> InPublicLibraryPaths,
-			IEnumerable<string> InPublicAdditionalLibraries,
-			IEnumerable<string> InPublicFrameworks,
-			IEnumerable<string> InPublicWeakFrameworks,
-			IEnumerable<UEBuildFramework> InPublicAdditionalFrameworks,
-			IEnumerable<string> InPublicAdditionalShadowFiles,
-			IEnumerable<UEBuildBundleResource> InPublicAdditionalBundleResources,
-			IEnumerable<string> InPublicDependencyModuleNames,
-			IEnumerable<string> InPublicDelayLoadDLLs,		// Delay loaded DLLs that should be setup when including this module
-			IEnumerable<RuntimeDependency> InRuntimeDependencies,
+			ModuleRules InRules,
 			string InBuildCsFilename
 			)
 		: base(
@@ -668,27 +639,7 @@ namespace UnrealBuildTool
 			InType:											InType,
 			InName:											InName,
 			InModuleDirectory:								InModuleDirectory,
-			InIsRedistributableOverride:					InIsRedistributableOverride,
-			InPublicDefinitions:							InPublicDefinitions,
-			InPublicIncludePaths:							InPublicIncludePaths,
-			InPublicSystemIncludePaths:						InPublicSystemIncludePaths,
-			InPublicLibraryPaths:							InPublicLibraryPaths,
-			InPublicAdditionalLibraries:					InPublicAdditionalLibraries,
-			InPublicFrameworks:								InPublicFrameworks,
-			InPublicWeakFrameworks:							InPublicWeakFrameworks,
-			InPublicAdditionalFrameworks:					InPublicAdditionalFrameworks,
-			InPublicAdditionalShadowFiles:					InPublicAdditionalShadowFiles,
-			InPublicAdditionalBundleResources:				InPublicAdditionalBundleResources,
-			InPublicIncludePathModuleNames:					null,
-			InPublicDependencyModuleNames:					InPublicDependencyModuleNames,
-			InPublicDelayLoadDLLs:							InPublicDelayLoadDLLs,
-			InPrivateIncludePaths:							null,
-			InPrivateIncludePathModuleNames:				null,
-			InPrivateDependencyModuleNames:					null,
-			InCircularlyReferencedDependentModules:			null,
-			InDynamicallyLoadedModuleNames:					null,
-			InPlatformSpecificDynamicallyLoadedModuleNames: null,
-			InRuntimeDependencies:							InRuntimeDependencies,
+			InRules:										InRules,
 			InBuildCsFilename:								InBuildCsFilename
 			)
 		{
@@ -808,45 +759,6 @@ namespace UnrealBuildTool
 		/// When set, allows this module to report compiler definitions and include paths for Intellisense
 		IntelliSenseGatherer IntelliSenseGatherer;
 
-		/** Whether this module contains performance critical code. */
-		ModuleRules.CodeOptimization OptimizeCode = ModuleRules.CodeOptimization.Default;
-
-		/** Whether this module should use a "shared PCH" header if possible.  Disable this for modules that have lots of global includes
-			in a private PCH header (e.g. game modules) */
-		public bool bAllowSharedPCH = true;
-
-		/** Header file name for a shared PCH provided by this module.  Must be a valid relative path to a public C++ header file.
-			This should only be set for header files that are included by a significant number of other C++ modules. */
-		public string SharedPCHHeaderFile = String.Empty;
-
-		/** Use run time type information */
-		public bool bUseRTTI = false;
-
-		/** Use AVX instructions */
-		public bool bUseAVX = false;
-
-		/** Enable buffer security checks. This should usually be enabled as it prevents severe security risks. */
-		public bool bEnableBufferSecurityChecks = true;
-
-		/** If true and unity builds are enabled, this module will build without unity. */
-		public bool bFasterWithoutUnity = false;
-
-		/** The number of source files in this module before unity build will be activated for that module.  If set to
-		    anything besides 0, will override the default setting which is controlled by MinGameModuleSourceFilesForUnityBuild */
-		public int MinSourceFilesForUnityBuildOverride = 0;
-
-		/** Overrides BuildConfiguration.MinFilesUsingPrecompiledHeader if non-zero. */
-		public int MinFilesUsingPrecompiledHeaderOverride = 0;
-
-		/**  Module uses a #import so must be built locally when compiling with SN-DBS */
-		public bool bBuildLocallyWithSNDBS = false;
-		
-		/** Enable exception handling */
-		public bool bEnableExceptions = false;
-
-		/** Enable warnings for shadowed variables */
-		public bool bEnableShadowVariableWarnings;
-
 		public List<string> IncludeSearchPaths = new List<string>();
 
 		public class ProcessedDependenciesClass
@@ -904,69 +816,17 @@ namespace UnrealBuildTool
 			UEBuildModuleType InType,
 			string InModuleDirectory,
 			string InGeneratedCodeDirectory,
-			bool? InIsRedistributableOverride,
 			IntelliSenseGatherer InIntelliSenseGatherer,
 			IEnumerable<FileItem> InSourceFiles,
-			IEnumerable<string> InPublicIncludePaths,
-			IEnumerable<string> InPublicSystemIncludePaths,
-            IEnumerable<string> InPublicLibraryPaths,
-			IEnumerable<string> InDefinitions,
-			IEnumerable<string> InPublicIncludePathModuleNames,
-			IEnumerable<string> InPublicDependencyModuleNames,
-			IEnumerable<string> InPublicDelayLoadDLLs,
-			IEnumerable<string> InPublicAdditionalLibraries,
-			IEnumerable<string> InPublicFrameworks,
-			IEnumerable<string> InPublicWeakFrameworks,
-			IEnumerable<UEBuildFramework> InPublicAdditionalFrameworks,
-			IEnumerable<string> InPublicAdditionalShadowFiles,
-			IEnumerable<UEBuildBundleResource> InPublicAdditionalBundleResources,
-			IEnumerable<string> InPrivateIncludePaths,
-			IEnumerable<string> InPrivateIncludePathModuleNames,
-			IEnumerable<string> InPrivateDependencyModuleNames,
-			IEnumerable<string> InCircularlyReferencedDependentModules,
-			IEnumerable<string> InDynamicallyLoadedModuleNames,
-            IEnumerable<string> InPlatformSpecificDynamicallyLoadedModuleNames,
-			IEnumerable<RuntimeDependency> InRuntimeDependencies,
-            ModuleRules.CodeOptimization InOptimizeCode,
-			bool InAllowSharedPCH,
-			string InSharedPCHHeaderFile,
-			bool InUseRTTI,
-			bool InEnableBufferSecurityChecks,
-			bool InFasterWithoutUnity,
-			int InMinSourceFilesForUnityBuildOverride,
-			int InMinFilesUsingPrecompiledHeaderOverride,
-			bool InBuildLocallyWithSNDBS,
-			bool InEnableExceptions,
-			bool InEnableShadowVariableWarnings,
+			ModuleRules InRules,
 			bool bInBuildSourceFiles,
-			string InBuildCsFilename,
-			bool bInUseAVX
+			string InBuildCsFilename
 			)
 			: base(	InTarget,
 					InName, 
 					InType,
 					InModuleDirectory,
-					InIsRedistributableOverride,
-					InDefinitions, 
-					InPublicIncludePaths,
-					InPublicSystemIncludePaths, 
-					InPublicLibraryPaths, 
-					InPublicAdditionalLibraries,
-					InPublicFrameworks,
-					InPublicWeakFrameworks,
-					InPublicAdditionalFrameworks,
-					InPublicAdditionalShadowFiles,
-					InPublicAdditionalBundleResources,
-					InPublicIncludePathModuleNames,
-					InPublicDependencyModuleNames,
-					InPublicDelayLoadDLLs,
-					InPrivateIncludePaths, 
-					InPrivateIncludePathModuleNames,
-					InPrivateDependencyModuleNames, 
-					InCircularlyReferencedDependentModules,
-					InDynamicallyLoadedModuleNames,
-                    InPlatformSpecificDynamicallyLoadedModuleNames,
-					InRuntimeDependencies,
+					InRules,
 					InBuildCsFilename
 				)
 		{
@@ -979,23 +839,11 @@ namespace UnrealBuildTool
 				SourceFilesToBuild.CopyFrom(SourceFilesFound);
 			}
 
-			Definitions = HashSetFromOptionalEnumerableStringParameter(InDefinitions);
+			Definitions = HashSetFromOptionalEnumerableStringParameter(InRules.Definitions);
 			foreach (var Def in Definitions)
 			{
 				Log.TraceVerbose("Compile Env {0}: {1}", Name, Def);
 			}
-			OptimizeCode                           = InOptimizeCode;
-			bAllowSharedPCH                        = InAllowSharedPCH;
-			SharedPCHHeaderFile                    = InSharedPCHHeaderFile;
-			bUseRTTI                               = InUseRTTI;
-			bUseAVX                                = bInUseAVX;
-			bEnableBufferSecurityChecks 		   = InEnableBufferSecurityChecks;
-			bFasterWithoutUnity                    = InFasterWithoutUnity;
-			MinSourceFilesForUnityBuildOverride    = InMinSourceFilesForUnityBuildOverride;
-            MinFilesUsingPrecompiledHeaderOverride = InMinFilesUsingPrecompiledHeaderOverride;
-			bBuildLocallyWithSNDBS				   = InBuildLocallyWithSNDBS;
-			bEnableExceptions                      = InEnableExceptions;
-			bEnableShadowVariableWarnings          = InEnableShadowVariableWarnings;
 		}
 
 		// UEBuildModule interface.
@@ -1080,9 +928,9 @@ namespace UnrealBuildTool
 			// builds.)  But for game modules, it can be convenient to always have a precompiled header to single-file
 			// changes to code is really quick to compile.
 			int MinFilesUsingPrecompiledHeader = BuildConfiguration.MinFilesUsingPrecompiledHeader;
-			if( MinFilesUsingPrecompiledHeaderOverride != 0 )
+			if( Rules.MinFilesUsingPrecompiledHeaderOverride != 0 )
 			{
-				MinFilesUsingPrecompiledHeader = MinFilesUsingPrecompiledHeaderOverride;
+				MinFilesUsingPrecompiledHeader = Rules.MinFilesUsingPrecompiledHeaderOverride;
 			}
 			else if( IsGameModule && BuildConfiguration.bForcePrecompiledHeaderForGameModules )
 			{
@@ -1097,9 +945,9 @@ namespace UnrealBuildTool
 			// to be worthwhile.  If you have a lot of small game modules, consider specifying MinSourceFilesForUnityBuildOverride=0
 			// in the modules that you don't typically iterate on source files in very frequently.
 			int MinSourceFilesForUnityBuild = 0;
-			if (MinSourceFilesForUnityBuildOverride != 0)
+			if (Rules.MinSourceFilesForUnityBuildOverride != 0)
 			{
-				MinSourceFilesForUnityBuild = MinSourceFilesForUnityBuildOverride;
+				MinSourceFilesForUnityBuild = Rules.MinSourceFilesForUnityBuildOverride;
 			}
 			else if ( IsGameModule )
 			{
@@ -1118,7 +966,7 @@ namespace UnrealBuildTool
 					Log.TraceVerbose("Module '{0}' using unity build mode (bForceUnityBuild enabled for this module)", this.Name);
 					bModuleUsesUnityBuild = true;
 				}
-				else if (bFasterWithoutUnity)
+				else if (Rules.bFasterWithoutUnity)
 				{
 					Log.TraceVerbose("Module '{0}' not using unity build mode (bFasterWithoutUnity enabled for this module)", this.Name);
 					bModuleUsesUnityBuild = false;
@@ -1169,11 +1017,12 @@ namespace UnrealBuildTool
 					bool bIsASharedPCHModule = bUseSharedPCHFiles && GlobalCompileEnvironment.SharedPCHHeaderFiles.Any( PCH => PCH.Module == this );
 					if( bIsASharedPCHModule )
 					{
-						SharingPCHHeaderFilePath = Path.GetFullPath( Path.Combine( ProjectFileGenerator.RootRelativePath, "Engine", "Source", this.SharedPCHHeaderFile ) );
+						SharingPCHHeaderFilePath = Path.GetFullPath( Path.Combine( ProjectFileGenerator.RootRelativePath, "Engine", "Source", Rules.SharedPCHHeaderFile ) );
 					}
 
 					// We can't use a shared PCH file when compiling a module
 					// with exports, because the shared PCH can only have imports in it to work correctly.
+					bool bAllowSharedPCH = (Rules.PCHUsage == ModuleRules.PCHUsageMode.NoSharedPCHs)? false : true;
 					bool bCanModuleUseOwnSharedPCH = bAllowSharedPCH && bIsASharedPCHModule && !Binary.Config.bAllowExports && ProcessedDependencies.UniquePCHHeaderFile.AbsolutePath.Equals( SharingPCHHeaderFilePath, StringComparison.InvariantCultureIgnoreCase );
 					if( bAllowSharedPCH && ( !bIsASharedPCHModule || bCanModuleUseOwnSharedPCH ) )
 					{
@@ -1342,7 +1191,7 @@ namespace UnrealBuildTool
 								UEBuildModuleCPP SharedPCHModule = (UEBuildModuleCPP)Target.FindOrCreateModuleByName(SharedPCHModuleName);
 
 								CPPEnvironment SharedPCHCompileEnvironment = GlobalCompileEnvironment.DeepCopy();
-								SharedPCHCompileEnvironment.Config.bEnableShadowVariableWarning = SharedPCHModule.bEnableShadowVariableWarnings;
+								SharedPCHCompileEnvironment.Config.bEnableShadowVariableWarning = SharedPCHModule.Rules.bEnableShadowVariableWarnings;
 
 								SharedPCHModule.SetupPublicCompileEnvironment(
 									Binary, 
@@ -1454,7 +1303,7 @@ namespace UnrealBuildTool
 				return null;
 			}
 			// Don't mix RTTI modes
-			if (bUseRTTI)
+			if (Rules.bUseRTTI)
 			{
 				Log.TraceVerbose("Module {0} cannot use existing Shared PCH '{1}' (from module '{2}') because RTTI modes don't match", Name, SharedPCHEnvironment.PrecompiledHeaderIncludeFilename.AbsolutePath, SharedPCHEnvironment.ModuleName);
 				SharedPCHHeaderFile = null;
@@ -1679,16 +1528,16 @@ namespace UnrealBuildTool
 			var Result = BaseCompileEnvironment.DeepCopy();
 
 			// Override compile environment
-			Result.Config.bFasterWithoutUnity                    = bFasterWithoutUnity;
-			Result.Config.OptimizeCode                           = OptimizeCode;
-			Result.Config.bUseRTTI                               = bUseRTTI;
-			Result.Config.bUseAVX                                = bUseAVX;
-			Result.Config.bEnableBufferSecurityChecks            = bEnableBufferSecurityChecks;
-			Result.Config.MinSourceFilesForUnityBuildOverride	 = MinSourceFilesForUnityBuildOverride;
-			Result.Config.MinFilesUsingPrecompiledHeaderOverride = MinFilesUsingPrecompiledHeaderOverride;
-			Result.Config.bBuildLocallyWithSNDBS				 = bBuildLocallyWithSNDBS;
-			Result.Config.bEnableExceptions                      = bEnableExceptions;
-			Result.Config.bEnableShadowVariableWarning          = bEnableShadowVariableWarnings;
+			Result.Config.bFasterWithoutUnity                    = Rules.bFasterWithoutUnity;
+			Result.Config.OptimizeCode                           = Rules.OptimizeCode;
+			Result.Config.bUseRTTI                               = Rules.bUseRTTI;
+			Result.Config.bUseAVX                                = Rules.bUseAVX;
+			Result.Config.bEnableBufferSecurityChecks            = Rules.bEnableBufferSecurityChecks;
+			Result.Config.MinSourceFilesForUnityBuildOverride	 = Rules.MinSourceFilesForUnityBuildOverride;
+			Result.Config.MinFilesUsingPrecompiledHeaderOverride = Rules.MinFilesUsingPrecompiledHeaderOverride;
+			Result.Config.bBuildLocallyWithSNDBS				 = Rules.bBuildLocallyWithSNDBS;
+			Result.Config.bEnableExceptions                      = Rules.bEnableExceptions;
+			Result.Config.bEnableShadowVariableWarning           = Rules.bEnableShadowVariableWarnings;
 			Result.Config.bUseStaticCRT							 = (Target.Rules != null && Target.Rules.bUseStaticCRT);
 			Result.Config.OutputDirectory                        = Path.Combine(Binary.Config.IntermediateDirectory, Name);
 
@@ -1906,52 +1755,17 @@ namespace UnrealBuildTool
 			UEBuildModuleType InType,
 			string InModuleDirectory,
 			string InGeneratedCodeDirectory,
-			bool? InIsRedistributableOverride,
 			IntelliSenseGatherer InIntelliSenseGatherer,
 			IEnumerable<FileItem> InSourceFiles,
-			IEnumerable<string> InPublicIncludePaths,
-			IEnumerable<string> InPublicSystemIncludePaths,
-			IEnumerable<string> InDefinitions,
-			IEnumerable<string> InPrivateAssemblyReferences,
-			IEnumerable<string> InPublicIncludePathModuleNames,
-			IEnumerable<string> InPublicDependencyModuleNames,
-			IEnumerable<string> InPublicDelayLoadDLLs,
-			IEnumerable<string> InPublicAdditionalLibraries,
-			IEnumerable<string> InPublicFrameworks,
-			IEnumerable<string> InPublicWeakFrameworks,
-			IEnumerable<UEBuildFramework> InPublicAdditionalFrameworks,
-			IEnumerable<string> InPublicAdditionalShadowFiles,
-			IEnumerable<UEBuildBundleResource> InPublicAdditionalBundleResources,
-			IEnumerable<string> InPrivateIncludePaths,
-			IEnumerable<string> InPrivateIncludePathModuleNames,
-			IEnumerable<string> InPrivateDependencyModuleNames,
-			IEnumerable<string> InCircularlyReferencedDependentModules,
-			IEnumerable<string> InDynamicallyLoadedModuleNames,
-            IEnumerable<string> InPlatformSpecificDynamicallyLoadedModuleNames,
-			IEnumerable<RuntimeDependency> InRuntimeDependencies,
-            ModuleRules.CodeOptimization InOptimizeCode,
-			bool InAllowSharedPCH,
-			string InSharedPCHHeaderFile,
-			bool InUseRTTI,
-			bool InEnableBufferSecurityChecks,
-			bool InFasterWithoutUnity,
-			int InMinSourceFilesForUnityBuildOverride,
-			int InMinFilesUsingPrecompiledHeaderOverride,
-			bool InEnableExceptions,
-			bool InEnableShadowVariableWarnings,
+			ModuleRules InRules,
 			bool bInBuildSourceFiles,
-			string InBuildCsFilename,
-			bool bInUseAVX
+			string InBuildCsFilename
 			)
-			: base(InTarget,InName,InType,InModuleDirectory,InGeneratedCodeDirectory,InIsRedistributableOverride,InIntelliSenseGatherer,
-			InSourceFiles,InPublicIncludePaths,InPublicSystemIncludePaths,null,InDefinitions,
-			InPublicIncludePathModuleNames,InPublicDependencyModuleNames,InPublicDelayLoadDLLs,InPublicAdditionalLibraries,InPublicFrameworks,InPublicWeakFrameworks,InPublicAdditionalFrameworks,InPublicAdditionalShadowFiles,InPublicAdditionalBundleResources,
-			InPrivateIncludePaths,InPrivateIncludePathModuleNames,InPrivateDependencyModuleNames,
-            InCircularlyReferencedDependentModules, InDynamicallyLoadedModuleNames, InPlatformSpecificDynamicallyLoadedModuleNames, InRuntimeDependencies, InOptimizeCode,
-			InAllowSharedPCH, InSharedPCHHeaderFile, InUseRTTI, InEnableBufferSecurityChecks, InFasterWithoutUnity, InMinSourceFilesForUnityBuildOverride, InMinFilesUsingPrecompiledHeaderOverride, true,
-			InEnableExceptions, InEnableShadowVariableWarnings, bInBuildSourceFiles, InBuildCsFilename, bInUseAVX)
+			: base(InTarget,InName,InType,InModuleDirectory,InGeneratedCodeDirectory,InIntelliSenseGatherer,
+			InSourceFiles,InRules,
+			bInBuildSourceFiles, InBuildCsFilename)
 		{
-			PrivateAssemblyReferences = HashSetFromOptionalEnumerableStringParameter(InPrivateAssemblyReferences);
+			PrivateAssemblyReferences = HashSetFromOptionalEnumerableStringParameter(InRules.PrivateAssemblyReferences);
 		}
 
 		// UEBuildModule interface.
diff --git a/Engine/Source/Programs/UnrealBuildTool/Configuration/UEBuildTarget.cs b/Engine/Source/Programs/UnrealBuildTool/Configuration/UEBuildTarget.cs
--- a/Engine/Source/Programs/UnrealBuildTool/Configuration/UEBuildTarget.cs
+++ b/Engine/Source/Programs/UnrealBuildTool/Configuration/UEBuildTarget.cs
@@ -2417,49 +2417,20 @@ namespace UnrealBuildTool
 		/// <returns>Created module.</returns>
 		private UEBuildModuleCPP CreateArtificialModule(string Name, string Directory, IEnumerable<FileItem> SourceFiles, IEnumerable<string> PrivateDependencyModuleNames)
 		{
+			ModuleRules Rules = new ModuleRules();
+			Rules.PrivateDependencyModuleNames.AddRange(PrivateDependencyModuleNames);
+
 			return new UEBuildModuleCPP(
 				InTarget: this,
 				InName: Name,
 				InType: UEBuildModuleType.Game,
 				InModuleDirectory: Directory,
 				InGeneratedCodeDirectory: null,
-				InIsRedistributableOverride: null,
 				InIntelliSenseGatherer: null,
 				InSourceFiles: SourceFiles.ToList(),
-				InPublicIncludePaths: new List<string>(),
-                InPublicLibraryPaths: new List<string>(),
-				InPublicSystemIncludePaths: new List<string>(),
-				InDefinitions: new List<string>(),
-				InPublicIncludePathModuleNames: new List<string>(),
-				InPublicDependencyModuleNames: new List<string>(),
-				InPublicDelayLoadDLLs: new List<string>(),
-				InPublicAdditionalLibraries: new List<string>(),
-				InPublicFrameworks: new List<string>(),
-				InPublicWeakFrameworks: new List<string>(),
-				InPublicAdditionalShadowFiles: new List<string>(),
-				InPublicAdditionalBundleResources: new List<UEBuildBundleResource>(),
-				InPublicAdditionalFrameworks: new List<UEBuildFramework>(),
-				InPrivateIncludePaths: new List<string>(),
-				InPrivateIncludePathModuleNames: new List<string>(),
-				InPrivateDependencyModuleNames: PrivateDependencyModuleNames.ToList(),
-				InCircularlyReferencedDependentModules: new List<string>(),
-				InDynamicallyLoadedModuleNames: new List<string>(),
-				InPlatformSpecificDynamicallyLoadedModuleNames: new List<string>(),
-				InRuntimeDependencies: new List<RuntimeDependency>(),
-				InOptimizeCode: ModuleRules.CodeOptimization.Default,
-				InAllowSharedPCH: false,
-				InSharedPCHHeaderFile: "",
-				InUseRTTI: false,
-				InEnableBufferSecurityChecks: true,
-				InFasterWithoutUnity: true,
-				InMinSourceFilesForUnityBuildOverride: 0,
-				InMinFilesUsingPrecompiledHeaderOverride: 0,
-				InBuildLocallyWithSNDBS: false,
-				InEnableExceptions: false,
-				InEnableShadowVariableWarnings: true,
+				InRules: Rules,
 				bInBuildSourceFiles: true,
-				InBuildCsFilename: null,
-				bInUseAVX: false);
+				InBuildCsFilename: null);
 		}
 
 
@@ -2486,12 +2457,12 @@ namespace UnrealBuildTool
 						var CPPModule = GetModuleByName( ModuleName ) as UEBuildModuleCPP;
 						if( CPPModule != null )
 						{
-							if( !String.IsNullOrEmpty( CPPModule.SharedPCHHeaderFile ) && CPPModule.Binary.Config.bAllowCompilation )
+							if( !String.IsNullOrEmpty( CPPModule.Rules.SharedPCHHeaderFile ) && CPPModule.Binary.Config.bAllowCompilation )
 							{
 								// @todo SharedPCH: Ideally we could figure the PCH header name automatically, and simply use a boolean in the module
 								//     definition to opt into exposing a shared PCH.  Unfortunately we don't determine which private PCH header "goes with"
 								//     a module until a bit later in the process.  It shouldn't be hard to change that though.
-								var SharedPCHHeaderFilePath = ProjectFileGenerator.RootRelativePath + "/Engine/Source/" + CPPModule.SharedPCHHeaderFile;
+								var SharedPCHHeaderFilePath = ProjectFileGenerator.RootRelativePath + "/Engine/Source/" + CPPModule.Rules.SharedPCHHeaderFile;
 								var SharedPCHHeaderFileItem = FileItem.GetExistingItemByPath( SharedPCHHeaderFilePath );
 								if( SharedPCHHeaderFileItem != null )
 								{
@@ -3475,43 +3446,11 @@ namespace UnrealBuildTool
 							InType: ModuleType,
 							InModuleDirectory: ModuleDirectory,
 							InGeneratedCodeDirectory: GeneratedCodeDirectory,
-							InIsRedistributableOverride: RulesObject.IsRedistributableOverride,
 							InIntelliSenseGatherer: IntelliSenseGatherer,
 							InSourceFiles: ModuleSourceFiles,
-							InPublicSystemIncludePaths: RulesObject.PublicSystemIncludePaths,
-                            InPublicLibraryPaths: RulesObject.PublicLibraryPaths,
-							InPublicIncludePaths: RulesObject.PublicIncludePaths,
-							InDefinitions: RulesObject.Definitions,
-							InPublicIncludePathModuleNames: RulesObject.PublicIncludePathModuleNames,
-							InPublicDependencyModuleNames: RulesObject.PublicDependencyModuleNames,
-							InPublicDelayLoadDLLs: RulesObject.PublicDelayLoadDLLs,
-							InPublicAdditionalLibraries: RulesObject.PublicAdditionalLibraries,
-							InPublicFrameworks: RulesObject.PublicFrameworks,
-							InPublicWeakFrameworks: RulesObject.PublicWeakFrameworks,
-							InPublicAdditionalFrameworks: RulesObject.PublicAdditionalFrameworks,
-							InPublicAdditionalShadowFiles: RulesObject.PublicAdditionalShadowFiles,
-							InPublicAdditionalBundleResources: RulesObject.AdditionalBundleResources,
-							InPrivateIncludePaths: RulesObject.PrivateIncludePaths,
-							InPrivateIncludePathModuleNames: RulesObject.PrivateIncludePathModuleNames,
-							InPrivateDependencyModuleNames: RulesObject.PrivateDependencyModuleNames,
-							InCircularlyReferencedDependentModules: RulesObject.CircularlyReferencedDependentModules,
-							InDynamicallyLoadedModuleNames: RulesObject.DynamicallyLoadedModuleNames,
-							InPlatformSpecificDynamicallyLoadedModuleNames: RulesObject.PlatformSpecificDynamicallyLoadedModuleNames,
-							InRuntimeDependencies: RulesObject.RuntimeDependencies,
-							InOptimizeCode: RulesObject.OptimizeCode,
-							InAllowSharedPCH: (RulesObject.PCHUsage == ModuleRules.PCHUsageMode.NoSharedPCHs) ? false : true,
-							InSharedPCHHeaderFile: RulesObject.SharedPCHHeaderFile,
-							InUseRTTI: RulesObject.bUseRTTI,
-							InEnableBufferSecurityChecks: RulesObject.bEnableBufferSecurityChecks,
-							InFasterWithoutUnity: RulesObject.bFasterWithoutUnity,
-							InMinSourceFilesForUnityBuildOverride: RulesObject.MinSourceFilesForUnityBuildOverride,
-							InMinFilesUsingPrecompiledHeaderOverride: RulesObject.MinFilesUsingPrecompiledHeaderOverride,
-							InBuildLocallyWithSNDBS: RulesObject.bBuildLocallyWithSNDBS,
-							InEnableExceptions: RulesObject.bEnableExceptions,
-							InEnableShadowVariableWarnings: RulesObject.bEnableShadowVariableWarnings,
+							InRules: RulesObject,
 							bInBuildSourceFiles: bBuildSourceFiles,
-							InBuildCsFilename: InBuildCsFile,
-							bInUseAVX: RulesObject.bUseAVX
+							InBuildCsFilename: InBuildCsFile
 						);
 
 				case ModuleRules.ModuleType.CPlusPlusCLR:
@@ -3521,42 +3460,11 @@ namespace UnrealBuildTool
 							InType: ModuleType,
 							InModuleDirectory: ModuleDirectory,
 							InGeneratedCodeDirectory: GeneratedCodeDirectory,
-							InIsRedistributableOverride: RulesObject.IsRedistributableOverride,
 							InIntelliSenseGatherer: IntelliSenseGatherer,
 							InSourceFiles: ModuleSourceFiles,
-							InDefinitions: RulesObject.Definitions,
-							InPublicSystemIncludePaths: RulesObject.PublicSystemIncludePaths,
-							InPublicIncludePaths: RulesObject.PublicIncludePaths,
-							InPublicIncludePathModuleNames: RulesObject.PublicIncludePathModuleNames,
-							InPublicDependencyModuleNames: RulesObject.PublicDependencyModuleNames,
-							InPublicDelayLoadDLLs: RulesObject.PublicDelayLoadDLLs,
-							InPublicAdditionalLibraries: RulesObject.PublicAdditionalLibraries,
-							InPublicFrameworks: RulesObject.PublicFrameworks,
-							InPublicWeakFrameworks: RulesObject.PublicWeakFrameworks,
-							InPublicAdditionalFrameworks: RulesObject.PublicAdditionalFrameworks,
-							InPublicAdditionalShadowFiles: RulesObject.PublicAdditionalShadowFiles,
-							InPublicAdditionalBundleResources: RulesObject.AdditionalBundleResources,
-							InPrivateIncludePaths: RulesObject.PrivateIncludePaths,
-							InPrivateIncludePathModuleNames: RulesObject.PrivateIncludePathModuleNames,
-							InPrivateDependencyModuleNames: RulesObject.PrivateDependencyModuleNames,
-							InPrivateAssemblyReferences: RulesObject.PrivateAssemblyReferences,
-							InCircularlyReferencedDependentModules: RulesObject.CircularlyReferencedDependentModules,
-							InDynamicallyLoadedModuleNames: RulesObject.DynamicallyLoadedModuleNames,
-							InPlatformSpecificDynamicallyLoadedModuleNames: RulesObject.PlatformSpecificDynamicallyLoadedModuleNames,
-							InRuntimeDependencies: RulesObject.RuntimeDependencies,
-							InOptimizeCode: RulesObject.OptimizeCode,
-							InAllowSharedPCH: (RulesObject.PCHUsage == ModuleRules.PCHUsageMode.NoSharedPCHs) ? false : true,
-							InSharedPCHHeaderFile: RulesObject.SharedPCHHeaderFile,
-							InUseRTTI: RulesObject.bUseRTTI,
-							InEnableBufferSecurityChecks: RulesObject.bEnableBufferSecurityChecks,
-							InFasterWithoutUnity: RulesObject.bFasterWithoutUnity,
-							InMinSourceFilesForUnityBuildOverride: RulesObject.MinSourceFilesForUnityBuildOverride,
-							InMinFilesUsingPrecompiledHeaderOverride: RulesObject.MinFilesUsingPrecompiledHeaderOverride,
-							InEnableExceptions: RulesObject.bEnableExceptions,
-							InEnableShadowVariableWarnings: RulesObject.bEnableShadowVariableWarnings,
+							InRules: RulesObject,
 							bInBuildSourceFiles : bBuildSourceFiles,
-							InBuildCsFilename: InBuildCsFile,
-							bInUseAVX: RulesObject.bUseAVX
+							InBuildCsFilename: InBuildCsFile
 						);
 
 				case ModuleRules.ModuleType.External:
@@ -3565,20 +3473,7 @@ namespace UnrealBuildTool
 							InName: ModuleName,
 							InType: ModuleType,
 							InModuleDirectory: ModuleDirectory,
-							InIsRedistributableOverride: RulesObject.IsRedistributableOverride,
-							InPublicDefinitions: RulesObject.Definitions,
-							InPublicSystemIncludePaths: RulesObject.PublicSystemIncludePaths,
-							InPublicIncludePaths: RulesObject.PublicIncludePaths,
-							InPublicLibraryPaths: RulesObject.PublicLibraryPaths,
-							InPublicAdditionalLibraries: RulesObject.PublicAdditionalLibraries,
-							InPublicFrameworks: RulesObject.PublicFrameworks,
-							InPublicWeakFrameworks: RulesObject.PublicWeakFrameworks,
-							InPublicAdditionalFrameworks: RulesObject.PublicAdditionalFrameworks,
-							InPublicAdditionalShadowFiles: RulesObject.PublicAdditionalShadowFiles,
-							InPublicAdditionalBundleResources: RulesObject.AdditionalBundleResources,
-							InPublicDependencyModuleNames: RulesObject.PublicDependencyModuleNames,
-							InPublicDelayLoadDLLs: RulesObject.PublicDelayLoadDLLs,
-							InRuntimeDependencies: RulesObject.RuntimeDependencies,
+							InRules: RulesObject,
 							InBuildCsFilename: InBuildCsFile
 						);
 
diff --git a/Engine/Source/Programs/UnrealBuildTool/System/RulesCompiler.cs b/Engine/Source/Programs/UnrealBuildTool/System/RulesCompiler.cs
--- a/Engine/Source/Programs/UnrealBuildTool/System/RulesCompiler.cs
+++ b/Engine/Source/Programs/UnrealBuildTool/System/RulesCompiler.cs
@@ -129,7 +129,7 @@ namespace UnrealBuildTool
 	/// <summary>
 	/// ModuleRules is a data structure that contains the rules for defining a module
 	/// </summary>
-	public abstract class ModuleRules
+	public class ModuleRules
 	{
 		/// Type of module
 		public enum ModuleType
