UE-20492 - fix for utilizing the path instead of the uproject name for finding the plist file
#ios

diff --git a/Engine/Source/Programs/UnrealBuildTool/IOS/IOSToolChain.cs b/Engine/Source/Programs/UnrealBuildTool/IOS/IOSToolChain.cs
--- a/Engine/Source/Programs/UnrealBuildTool/IOS/IOSToolChain.cs
+++ b/Engine/Source/Programs/UnrealBuildTool/IOS/IOSToolChain.cs
@@ -1165,6 +1165,7 @@ namespace UnrealBuildTool
 								throw new BuildException("The target was not in a /Source subdirectory");
 							}
 						}
+						PathToApp += "\\" + AppName + ".uproject";
 					}
 
 					if (bUseDangerouslyFastMode)
