Track references to dependent modules rather than just keeping their name. Saves having to look up module names repeatedly.

diff --git a/Engine/Source/Programs/UnrealBuildTool/Configuration/UEBuildBinary.cs b/Engine/Source/Programs/UnrealBuildTool/Configuration/UEBuildBinary.cs
--- a/Engine/Source/Programs/UnrealBuildTool/Configuration/UEBuildBinary.cs
+++ b/Engine/Source/Programs/UnrealBuildTool/Configuration/UEBuildBinary.cs
@@ -500,6 +500,7 @@ namespace UnrealBuildTool
 				for(int Idx = 0; Idx < ModuleNames.Count; Idx++)
 				{
 					UEBuildModule Module = Target.FindOrCreateModuleByName(ModuleNames[Idx]);
+					Module.RecursivelyCreateModules();
 					Module.RecursivelyProcessUnboundModules();
 				}
 			}
diff --git a/Engine/Source/Programs/UnrealBuildTool/Configuration/UEBuildModule.cs b/Engine/Source/Programs/UnrealBuildTool/Configuration/UEBuildModule.cs
--- a/Engine/Source/Programs/UnrealBuildTool/Configuration/UEBuildModule.cs
+++ b/Engine/Source/Programs/UnrealBuildTool/Configuration/UEBuildModule.cs
@@ -216,32 +216,25 @@ namespace UnrealBuildTool
 		protected readonly HashSet<UEBuildBundleResource> PublicAdditionalBundleResources;
 
 		/** Names of modules with header files that this module's public interface needs access to. */
-		protected HashSet<string> PublicIncludePathModuleNames;
+		protected List<UEBuildModule> PublicIncludePathModules;
 
 		/** Names of modules that this module's public interface depends on. */
-		protected HashSet<string> PublicDependencyModuleNames;
+		protected List<UEBuildModule> PublicDependencyModules;
 
 		/** Names of DLLs that this module should delay load */
 		protected HashSet<string> PublicDelayLoadDLLs;
 
 		/** Names of modules with header files that this module's private implementation needs access to. */
-		protected HashSet<string> PrivateIncludePathModuleNames;
+		protected List<UEBuildModule> PrivateIncludePathModules;
 	
 		/** Names of modules that this module's private implementation depends on. */
-		protected HashSet<string> PrivateDependencyModuleNames;
-
-		/** If any of this module's dependent modules circularly reference this module, then they must be added to this list */
-		public HashSet<string> CircularlyReferencedDependentModules
-		{
-			get;
-			protected set;
-		}
+		protected List<UEBuildModule> PrivateDependencyModules;
 
 		/** Extra modules this module may require at run time */
-		protected HashSet<string> DynamicallyLoadedModuleNames;
+		protected List<UEBuildModule> DynamicallyLoadedModules;
 
         /** Extra modules this module may require at run time, that are on behalf of another platform (i.e. shader formats and the like) */
-		protected HashSet<string> PlatformSpecificDynamicallyLoadedModuleNames;
+		protected List<UEBuildModule> PlatformSpecificDynamicallyLoadedModules;
 
 		/** Files which this module depends on at runtime. */
 		public List<RuntimeDependency> RuntimeDependencies;
@@ -270,15 +263,8 @@ namespace UnrealBuildTool
 			PublicAdditionalFrameworks = InRules.PublicAdditionalFrameworks == null ? new HashSet<UEBuildFramework>() : new HashSet<UEBuildFramework>(InRules.PublicAdditionalFrameworks);
 			PublicAdditionalShadowFiles = HashSetFromOptionalEnumerableStringParameter(InRules.PublicAdditionalShadowFiles);
 			PublicAdditionalBundleResources = InRules.AdditionalBundleResources == null ? new HashSet<UEBuildBundleResource>() : new HashSet<UEBuildBundleResource>(InRules.AdditionalBundleResources);
-			PublicIncludePathModuleNames = HashSetFromOptionalEnumerableStringParameter( InRules.PublicIncludePathModuleNames );
-			PublicDependencyModuleNames = HashSetFromOptionalEnumerableStringParameter(InRules.PublicDependencyModuleNames);
 			PublicDelayLoadDLLs = HashSetFromOptionalEnumerableStringParameter(InRules.PublicDelayLoadDLLs);
 			PrivateIncludePaths = HashSetFromOptionalEnumerableStringParameter(InRules.PrivateIncludePaths);
-			PrivateIncludePathModuleNames = HashSetFromOptionalEnumerableStringParameter( InRules.PrivateIncludePathModuleNames );
-			PrivateDependencyModuleNames = HashSetFromOptionalEnumerableStringParameter( InRules.PrivateDependencyModuleNames );
-			CircularlyReferencedDependentModules = new HashSet<string>( HashSetFromOptionalEnumerableStringParameter( InRules.CircularlyReferencedDependentModules ) );
-			DynamicallyLoadedModuleNames = HashSetFromOptionalEnumerableStringParameter( InRules.DynamicallyLoadedModuleNames );
-            PlatformSpecificDynamicallyLoadedModuleNames = HashSetFromOptionalEnumerableStringParameter(InRules.PlatformSpecificDynamicallyLoadedModuleNames);
 			RuntimeDependencies = (InRules.RuntimeDependencies == null)? new List<RuntimeDependency>() : new List<RuntimeDependency>(InRules.RuntimeDependencies);
 			IsRedistributableOverride = InRules.IsRedistributableOverride;
 
@@ -290,6 +276,12 @@ namespace UnrealBuildTool
 			Target.RegisterModule(this);
 		}
 
+		/** Determines whether this module has a circular dependency on the given module */
+		public bool HasCircularDependencyOn(string ModuleName)
+		{
+			return Rules.CircularlyReferencedDependentModules.Contains(ModuleName);
+		}
+
 		/** Sets up the environment for compiling any module that includes the public interface of this module. */
 		protected virtual void SetupPublicCompileEnvironment(
 			UEBuildBinary SourceBinary,
@@ -373,18 +365,16 @@ namespace UnrealBuildTool
 				if( !bIncludePathsOnly )
 				{
 					// Recurse on this module's public dependencies.
-					foreach(var DependencyName in PublicDependencyModuleNames)
+					foreach(UEBuildModule DependencyModule in PublicDependencyModules)
 					{
-						var DependencyModule = Target.GetModuleByName(DependencyName);
 						DependencyModule.SetupPublicCompileEnvironment(SourceBinary,bIncludePathsOnly, IncludePaths, SystemIncludePaths, Definitions, AdditionalFrameworks, VisitedModules);
 					}
 				}
 
 				// Now add an include paths from modules with header files that we need access to, but won't necessarily be importing
-				foreach( var IncludePathModuleName in PublicIncludePathModuleNames )
+				foreach( UEBuildModule IncludePathModule in PublicIncludePathModules )
 				{
 					bool bInnerIncludePathsOnly = true;
-					var IncludePathModule = Target.GetModuleByName(IncludePathModuleName);
 					IncludePathModule.SetupPublicCompileEnvironment( SourceBinary, bInnerIncludePathsOnly, IncludePaths, SystemIncludePaths, Definitions, AdditionalFrameworks, VisitedModules );
 				}
 
@@ -441,7 +431,7 @@ namespace UnrealBuildTool
 			List<UEBuildFramework> AdditionalFrameworks
 			)
 		{
-			var VisitedModules = new Dictionary<UEBuildModule, bool>();
+			Dictionary<UEBuildModule, bool> VisitedModules = new Dictionary<UEBuildModule, bool>();
 
 			if (this.Type == UEBuildModuleType.Game)
 			{
@@ -456,17 +446,15 @@ namespace UnrealBuildTool
 			SetupPublicCompileEnvironment(Binary,bIncludePathsOnly, IncludePaths, SystemIncludePaths, Definitions, AdditionalFrameworks, VisitedModules);
 
 			// Also allow the module's private dependencies to modify the compile environment.
-			foreach(var DependencyName in PrivateDependencyModuleNames)
+			foreach(UEBuildModule DependencyModule in PrivateDependencyModules)
 			{
-				var DependencyModule = Target.GetModuleByName(DependencyName);
 				DependencyModule.SetupPublicCompileEnvironment(Binary,bIncludePathsOnly, IncludePaths, SystemIncludePaths, Definitions, AdditionalFrameworks, VisitedModules);
 			}
 
 			// Add include paths from modules with header files that our private files need access to, but won't necessarily be importing
-			foreach( var IncludePathModuleName in PrivateIncludePathModuleNames )
+			foreach(UEBuildModule IncludePathModule in PrivateIncludePathModules)
 			{
 				bool bInnerIncludePathsOnly = true;
-				var IncludePathModule = Target.GetModuleByName(IncludePathModuleName);
 				IncludePathModule.SetupPublicCompileEnvironment(Binary, bInnerIncludePathsOnly, IncludePaths, SystemIncludePaths, Definitions, AdditionalFrameworks, VisitedModules);
 			}
 		}
@@ -510,12 +498,12 @@ namespace UnrealBuildTool
 					if (!bIsBuildingAStaticLibrary && bIsModuleBinaryAStaticLibrary)
 					{
 						// Gather all dependencies and recursively call SetupPublicLinkEnvironmnet
-						List<string> AllDependencyModuleNames = new List<string>(PrivateDependencyModuleNames);
-						AllDependencyModuleNames.AddRange(PublicDependencyModuleNames);
+						List<UEBuildModule> AllDependencyModules = new List<UEBuildModule>();
+						AllDependencyModules.AddRange(PrivateDependencyModules);
+						AllDependencyModules.AddRange(PublicDependencyModules);
 
-						foreach (var DependencyName in AllDependencyModuleNames)
+						foreach (UEBuildModule DependencyModule in AllDependencyModules)
 						{
-							var DependencyModule = Target.GetModuleByName(DependencyName);
 							bool bIsExternalModule = (DependencyModule as UEBuildExternalModule != null);
 							bool bIsInStaticLibrary = (DependencyModule.Binary != null && DependencyModule.Binary.Config.Type == UEBuildBinaryType.StaticLibrary);
 							if (bIsExternalModule || bIsInStaticLibrary)
@@ -557,12 +545,12 @@ namespace UnrealBuildTool
 				LinkEnvironment.Config.AdditionalFrameworks,LinkEnvironment.Config.AdditionalShadowFiles, LinkEnvironment.Config.AdditionalBundleResources, LinkEnvironment.Config.DelayLoadDLLs, BinaryDependencies, VisitedModules);
 
 			// Also allow the module's public and private dependencies to modify the link environment.
-			List<string> AllDependencyModuleNames = new List<string>(PrivateDependencyModuleNames);
-			AllDependencyModuleNames.AddRange(PublicDependencyModuleNames);
+			List<UEBuildModule> AllDependencyModules = new List<UEBuildModule>();
+			AllDependencyModules.AddRange(PrivateDependencyModules);
+			AllDependencyModules.AddRange(PublicDependencyModules);
 
-			foreach (var DependencyName in AllDependencyModuleNames)
+			foreach (UEBuildModule DependencyModule in AllDependencyModules)
 			{
-				var DependencyModule = Target.GetModuleByName(DependencyName);
 				DependencyModule.SetupPublicLinkEnvironment(SourceBinary, LinkEnvironment.Config.LibraryPaths, LinkEnvironment.Config.AdditionalLibraries, LinkEnvironment.Config.Frameworks, LinkEnvironment.Config.WeakFrameworks,
 					LinkEnvironment.Config.AdditionalFrameworks, LinkEnvironment.Config.AdditionalShadowFiles, LinkEnvironment.Config.AdditionalBundleResources, LinkEnvironment.Config.DelayLoadDLLs, BinaryDependencies, VisitedModules);
 			}
@@ -613,13 +601,50 @@ namespace UnrealBuildTool
 		}
 
 		/**
-		 * Recursively adds modules that are referenced only for include paths (not actual dependencies).
-		 * 
-		 * @param	Target	The target we are currently building
-		 * @param	bPublicIncludesOnly	True to only add modules for public includes, false to add all include path modules.  Regardless of what you pass here, recursive adds will only add public include path modules.
+		 * Creates all the modules required for this target
 		 */
-		public virtual void RecursivelyAddIncludePathModules( UEBuildTarget Target, bool bPublicIncludesOnly )
+		public void RecursivelyCreateModules()
+		{
+			// Create all the include path modules. These modules may not be added to the target (and we don't process their dependencies), but they need 
+			// to be created to set up their compile environment.
+			RecursivelyCreateIncludePathModulesByName(Target, Rules.PublicIncludePathModuleNames, ref PublicIncludePathModules);
+			RecursivelyCreateIncludePathModulesByName(Target, Rules.PrivateIncludePathModuleNames, ref PrivateIncludePathModules);
+
+			// Create all the dependency modules
+			RecursivelyCreateModulesByName(Target, Rules.PublicDependencyModuleNames, ref PublicDependencyModules);
+			RecursivelyCreateModulesByName(Target, Rules.PrivateDependencyModuleNames, ref PrivateDependencyModules);
+			RecursivelyCreateModulesByName(Target, Rules.DynamicallyLoadedModuleNames, ref DynamicallyLoadedModules);
+			RecursivelyCreateModulesByName(Target, Rules.PlatformSpecificDynamicallyLoadedModuleNames, ref PlatformSpecificDynamicallyLoadedModules);
+		}
+
+		private static void RecursivelyCreateModulesByName(UEBuildTarget Target, List<string> ModuleNames, ref List<UEBuildModule> Modules)
 		{
+			// Check whether the module list is already set. We set this immediately (via the ref) to avoid infinite recursion.
+			if(Modules == null)
+			{
+				Modules = new List<UEBuildModule>();
+				foreach(string ModuleName in ModuleNames.Distinct(StringComparer.InvariantCultureIgnoreCase))
+				{
+					UEBuildModule Module = Target.FindOrCreateModuleByName(ModuleName);
+					Module.RecursivelyCreateModules();
+					Modules.Add(Module);
+				}
+			}
+		}
+
+		private static void RecursivelyCreateIncludePathModulesByName(UEBuildTarget Target, List<string> ModuleNames, ref List<UEBuildModule> Modules)
+		{
+			// Check whether the module list is already set. We set this immediately (via the ref) to avoid infinite recursion.
+			if(Modules == null)
+			{
+				Modules = new List<UEBuildModule>();
+				foreach(string ModuleName in ModuleNames)
+				{
+					UEBuildModule Module = Target.FindOrCreateModuleByName(ModuleName);
+					RecursivelyCreateIncludePathModulesByName(Target, Module.Rules.PublicIncludePathModuleNames, ref Module.PublicIncludePathModules);
+					Modules.Add(Module);
+				}
+			}
 		}
 	};
 
@@ -1606,33 +1631,32 @@ namespace UnrealBuildTool
 
 		public override void GetAllDependencyModules( CaselessDictionary<ModuleIndexPair> ReferencedModules, bool bIncludeDynamicallyLoaded, bool bForceCircular, bool bOnlyDirectDependencies )
 		{
-			var AllModuleNames = new List<string>();
-			AllModuleNames.AddRange(PrivateDependencyModuleNames);
-			AllModuleNames.AddRange(PublicDependencyModuleNames);
+			List<UEBuildModule> AllDependencyModules = new List<UEBuildModule>();
+			AllDependencyModules.AddRange(PrivateDependencyModules);
+			AllDependencyModules.AddRange(PublicDependencyModules);
 			if( bIncludeDynamicallyLoaded )
 			{
-				AllModuleNames.AddRange(DynamicallyLoadedModuleNames);
-                AllModuleNames.AddRange(PlatformSpecificDynamicallyLoadedModuleNames);
+				AllDependencyModules.AddRange(DynamicallyLoadedModules);
+                AllDependencyModules.AddRange(PlatformSpecificDynamicallyLoadedModules);
             }
 
-			foreach (var DependencyName in AllModuleNames)
+			foreach (UEBuildModule DependencyModule in AllDependencyModules)
 			{
-				if (!ReferencedModules.ContainsKey(DependencyName))
+				if (!ReferencedModules.ContainsKey(DependencyModule.Name))
 				{
 					// Don't follow circular back-references!
-					bool bIsCircular = CircularlyReferencedDependentModules.Contains( DependencyName );
+					bool bIsCircular = HasCircularDependencyOn(DependencyModule.Name);
 					if (bForceCircular || !bIsCircular)
 					{
-						var Module = Target.GetModuleByName( DependencyName );
-						ReferencedModules[ DependencyName ] = null;
+						ReferencedModules[ DependencyModule.Name ] = null;
 
 						if( !bOnlyDirectDependencies )
 						{ 
 							// Recurse into dependent modules first
-							Module.GetAllDependencyModules(ReferencedModules, bIncludeDynamicallyLoaded, bForceCircular, bOnlyDirectDependencies);
+							DependencyModule.GetAllDependencyModules(ReferencedModules, bIncludeDynamicallyLoaded, bForceCircular, bOnlyDirectDependencies);
 						}
 
-						ReferencedModules[ DependencyName ] = new ModuleIndexPair{ Module = Module, Index = ReferencedModules.Where(x => x.Value != null).Count() - 1 };
+						ReferencedModules[ DependencyModule.Name ] = new ModuleIndexPair{ Module = DependencyModule, Index = ReferencedModules.Where(x => x.Value != null).Count() - 1 };
 					}
 				}
 			}
@@ -1645,42 +1669,32 @@ namespace UnrealBuildTool
 				Modules.Add(this);
 
 				// Get the dependent modules
-				List<string> DependentModuleNames = new List<string>();
-				DependentModuleNames.AddRange(PrivateDependencyModuleNames);
-				DependentModuleNames.AddRange(PublicDependencyModuleNames);
-				DependentModuleNames.AddRange(DynamicallyLoadedModuleNames);
-				DependentModuleNames.AddRange(PlatformSpecificDynamicallyLoadedModuleNames);
+				List<UEBuildModule> DependentModules = new List<UEBuildModule>();
+				if(PrivateDependencyModules != null)
+				{
+					DependentModules.AddRange(PrivateDependencyModules);
+				}
+				if(PublicDependencyModules != null)
+				{
+					DependentModules.AddRange(PublicDependencyModules);
+				}
+				if(DynamicallyLoadedModules != null)
+				{
+					DependentModules.AddRange(DynamicallyLoadedModules);
+				}
+				if(PlatformSpecificDynamicallyLoadedModules != null)
+				{
+					DependentModules.AddRange(PlatformSpecificDynamicallyLoadedModules);
+				}
 
 				// Find modules for each of them, and add their dependencies too
-				foreach(string DependentModuleName in DependentModuleNames)
+				foreach(UEBuildModule DependentModule in DependentModules)
 				{
-					UEBuildModule DependentModule = Target.FindOrCreateModuleByName(DependentModuleName);
 					DependentModule.RecursivelyAddPrecompiledModules(Modules);
 				}
 			}
 		}
 
-		public override void RecursivelyAddIncludePathModules( UEBuildTarget Target, bool bPublicIncludesOnly )
-		{
-			var AllIncludePathModuleNames = new List<string>();
-			AllIncludePathModuleNames.AddRange( PublicIncludePathModuleNames );
-			if( !bPublicIncludesOnly )
-			{
-				AllIncludePathModuleNames.AddRange( PrivateIncludePathModuleNames );
-			}
-			foreach( var IncludePathModuleName in AllIncludePathModuleNames )
-			{
-				var IncludePathModule = Target.FindOrCreateModuleByName( IncludePathModuleName );
-
-				// No need to do anything here.  We just want to make sure that we've instantiated our representation of the
-				// module so that we can find it later when processing include path module names.  Unlike actual dependency
-				// modules, these include path modules may never be "bound" (have Binary or bIncludedInTarget member variables set)
-
-				// We'll also need to make sure we know about all dependent public include path modules, too!
-				IncludePathModule.RecursivelyAddIncludePathModules( Target, bPublicIncludesOnly:true );
-			}
-		}
-
 		public override void RecursivelyProcessUnboundModules()
 		{
 			try
@@ -1691,23 +1705,21 @@ namespace UnrealBuildTool
 					throw new BuildException( "Module '{0}' should already have been bound to a binary!", Name );
 				}
 
-				var AllModuleNames = new List<string>();
-				AllModuleNames.AddRange( PrivateDependencyModuleNames );
-				AllModuleNames.AddRange( PublicDependencyModuleNames );
-				AllModuleNames.AddRange( DynamicallyLoadedModuleNames );
-				AllModuleNames.AddRange( PlatformSpecificDynamicallyLoadedModuleNames );
+				List<UEBuildModule> AllDependencyModules = new List<UEBuildModule>();
+				AllDependencyModules.AddRange( PrivateDependencyModules );
+				AllDependencyModules.AddRange( PublicDependencyModules );
+				AllDependencyModules.AddRange( DynamicallyLoadedModules );
+				AllDependencyModules.AddRange( PlatformSpecificDynamicallyLoadedModules );
 
-				foreach( var DependencyName in AllModuleNames )
+				foreach( UEBuildModule DependencyModule in AllDependencyModules )
 				{
-					var DependencyModule = Target.FindOrCreateModuleByName(DependencyName);
-
 					// Skip modules that are included with the target (externals)
 					if( !DependencyModule.bIncludedInTarget )
 					{
-						bool bIsCrossTarget = PlatformSpecificDynamicallyLoadedModuleNames.Contains(DependencyName) && !DynamicallyLoadedModuleNames.Contains(DependencyName);
+						bool bIsCrossTarget = Rules.PlatformSpecificDynamicallyLoadedModuleNames.Contains(DependencyModule.Name) && !Rules.DynamicallyLoadedModuleNames.Contains(DependencyModule.Name);
 
 						// Get the binary that this module should be bound to
-						UEBuildBinary BinaryToBindTo = Target.FindOrAddBinaryForModule(DependencyName, bIsCrossTarget);
+						UEBuildBinary BinaryToBindTo = Target.FindOrAddBinaryForModule(DependencyModule.Name, bIsCrossTarget);
 
 						// Bind this module
 						DependencyModule.Binary = BinaryToBindTo;
@@ -1722,19 +1734,13 @@ namespace UnrealBuildTool
 						// Check to see if there is a circular relationship between the module and it's referencer
 						if( DependencyModule.Binary != null )
 						{
-							if( CircularlyReferencedDependentModules.Contains( DependencyName ) )
+							if( Rules.CircularlyReferencedDependentModules.Contains( DependencyModule.Name ) )
 							{
 								DependencyModule.Binary.SetCreateImportLibrarySeparately( true );
 							}
 						}
 					}
 				}
-
-				// Also make sure module entries are created for any module that is pulled in as an "include path" module.
-				// These modules are never linked in unless they were referenced as an actual dependency of a different module,
-				// but we still need to keep track of them so that we can find their include paths when setting up our
-				// module's include paths.
-				RecursivelyAddIncludePathModules( Target, bPublicIncludesOnly:false );
 			}
 			catch (System.Exception ex)
 			{
diff --git a/Engine/Source/Programs/UnrealBuildTool/Configuration/UEBuildTarget.cs b/Engine/Source/Programs/UnrealBuildTool/Configuration/UEBuildTarget.cs
--- a/Engine/Source/Programs/UnrealBuildTool/Configuration/UEBuildTarget.cs
+++ b/Engine/Source/Programs/UnrealBuildTool/Configuration/UEBuildTarget.cs
@@ -2311,6 +2311,7 @@ namespace UnrealBuildTool
 
 				// Now bind this new module to the executable binary so it will link the plugin libs correctly
 				NewModule.bSkipDefinitionsForCompileEnvironment = true;
+				NewModule.RecursivelyCreateModules();
 				BindArtificialModuleToBinary(NewModule, ExecutableBinary);
 
 				// Create the cpp file
@@ -2489,7 +2490,7 @@ namespace UnrealBuildTool
 												// that UnrealEd depends on Engine.  We use this to avoid having modules such as UnrealEd
 												// appear before Engine in our shared PCH list.
 												// @todo SharedPCH: This is not very easy for people to discover.  Luckily we won't have many shared PCHs in total.
-												if( !ExistingModule.CircularlyReferencedDependentModules.Contains( CPPModule.Name ) )
+												if( !ExistingModule.HasCircularDependencyOn( CPPModule.Name ) )
 												{
 													// We are at least dependent on this module.  We'll keep searching the list to find
 													// further-descendant modules we might be dependent on
