Move the RulesAssembly class out of the RulesCompiler, so we don't have to rely on state in RulesCompiler.

diff --git a/Engine/Source/Programs/UnrealBuildTool/System/RulesCompiler.cs b/Engine/Source/Programs/UnrealBuildTool/System/RulesCompiler.cs
--- a/Engine/Source/Programs/UnrealBuildTool/System/RulesCompiler.cs
+++ b/Engine/Source/Programs/UnrealBuildTool/System/RulesCompiler.cs
@@ -1070,6 +1070,360 @@ namespace UnrealBuildTool
 		}
 	}
 
+	public class RulesAssembly
+	{
+		// The compiled assembly
+		private Assembly CompiledAssembly;
+
+		/// Maps module names to their actual xxx.Module.cs file on disk
+		private Dictionary<string, string> ModuleNameToModuleFileMap = new Dictionary<string, string>(StringComparer.InvariantCultureIgnoreCase);
+
+		/// Maps target names to their actual xxx.Target.cs file on disk
+		private Dictionary<string, string> TargetNameToTargetFileMap = new Dictionary<string, string>(StringComparer.InvariantCultureIgnoreCase);
+
+		/// <summary>
+		/// Constructor. Compiles a rules assembly from the given source files.
+		/// </summary>
+		/// <param name="ModuleFileNames">List of module files to compile</param>
+		/// <param name="TargetFileNames">List of target files to compile</param>
+		/// <param name="AssemblyFileName">The output path for the compiled assembly</param>
+		public RulesAssembly(List<string> ModuleFileNames, List<string> TargetFileNames, string AssemblyFileName)
+		{
+			// Find all the source files
+			var AssemblySourceFiles = new List<string>();
+			AssemblySourceFiles.AddRange( ModuleFileNames );
+			if( AssemblySourceFiles.Count == 0 )
+			{
+				throw new BuildException("No module rules source files were found in any of the module base directories!");
+			}
+			AssemblySourceFiles.AddRange( TargetFileNames );
+
+			// Compile the assembly
+			CompiledAssembly = DynamicCompilation.CompileAndLoadAssembly( AssemblyFileName, AssemblySourceFiles );
+
+			// Setup the module map
+			foreach( var CurModuleFileName in ModuleFileNames )
+			{
+				var CleanFileName = Utils.CleanDirectorySeparators( CurModuleFileName );
+				var ModuleName = Path.GetFileNameWithoutExtension( Path.GetFileNameWithoutExtension( CleanFileName ) );	// Strip both extensions
+				if( !ModuleNameToModuleFileMap.ContainsKey( ModuleName ) )
+				{
+					ModuleNameToModuleFileMap.Add( ModuleName, CurModuleFileName );
+				}
+			}
+
+			// Setup the target map
+			foreach( var CurTargetFileName in TargetFileNames )
+			{
+				var CleanFileName = Utils.CleanDirectorySeparators( CurTargetFileName );
+				var TargetName = Path.GetFileNameWithoutExtension( Path.GetFileNameWithoutExtension( CleanFileName ) );	// Strip both extensions
+				if( !TargetNameToTargetFileMap.ContainsKey( TargetName ) )
+				{
+					TargetNameToTargetFileMap.Add( TargetName, CurTargetFileName );
+				}
+			}
+		}
+
+		/// <summary>
+		/// The name of the compiled assembly
+		/// </summary>
+		public string Name
+		{
+			get { return Path.GetFileNameWithoutExtension( CompiledAssembly.Location ); }
+		}
+
+		/// <summary>
+		/// 
+		/// </summary>
+		/// <param name="InModuleName"></param>
+		/// <returns></returns>
+		public string GetModuleFilename(string InModuleName)
+		{
+			// Make sure the module file is known to us
+			if (!ModuleNameToModuleFileMap.ContainsKey(InModuleName))
+			{
+				return "";
+			}
+
+			// Return the module file name to the caller
+			return ModuleNameToModuleFileMap[InModuleName];
+		}
+
+		public string GetTargetFilename(string InTargetName)
+		{
+			// Make sure the target file is known to us
+			if (!TargetNameToTargetFileMap.ContainsKey(InTargetName))
+			{
+				return "";
+			}
+
+			// Return the target file name to the caller
+			return TargetNameToTargetFileMap[InTargetName];
+		}
+
+		/// <summary>
+		/// Creates an instance of a module rules descriptor object for the specified module name
+		/// </summary>
+		/// <param name="ModuleName">Name of the module</param>
+		/// <param name="Target">Information about the target associated with this module</param>
+		/// <param name="ModuleFileName">The original source file name for the Module.cs file for this module</param>
+		/// <returns>Compiled module rule info</returns>
+		public ModuleRules CreateModuleRules(string ModuleName, TargetInfo Target, out string ModuleFileName)
+		{
+			var AssemblyFileName = Path.GetFileNameWithoutExtension(CompiledAssembly.Location);
+
+			// Currently, we expect the user's rules object type name to be the same as the module name
+			var ModuleTypeName = ModuleName;
+
+			// Make sure the module file is known to us
+			if (!ModuleNameToModuleFileMap.ContainsKey(ModuleName))
+			{
+				throw new MissingModuleException(ModuleName);
+			}
+
+			// Return the module file name to the caller
+			ModuleFileName = ModuleNameToModuleFileMap[ModuleName];
+
+			UnrealTargetPlatform LocalPlatform = Target.Platform;
+			UnrealTargetConfiguration LocalConfiguration = Target.Configuration;
+			TargetInfo LocalTarget = new TargetInfo(LocalPlatform, LocalConfiguration, Target.Type.Value, Target.bIsMonolithic.Value);
+
+			// The build module must define a type named 'Rules' that derives from our 'ModuleRules' type.  
+			var RulesObjectType = CompiledAssembly.GetType(ModuleName);
+
+			if (RulesObjectType == null)
+			{
+				// Temporary hack to avoid System namespace collisions
+				// @todo projectfiles: Make rules assemblies require namespaces.
+				RulesObjectType = CompiledAssembly.GetType("UnrealBuildTool.Rules." + ModuleName);
+			}
+
+			if (RulesObjectType == null)
+			{
+				throw new BuildException("Expecting to find a type to be declared in a module rules named '{0}' in {1}.  This type must derive from the 'ModuleRules' type defined by Unreal Build Tool.", ModuleTypeName, CompiledAssembly.FullName);
+			}
+
+			// Create an instance of the module's rules object
+			ModuleRules RulesObject;
+			try
+			{
+				RulesObject = (ModuleRules)Activator.CreateInstance(RulesObjectType, LocalTarget);
+			}
+			catch (Exception Ex)
+			{
+				throw new BuildException(Ex, "Unable to instantiate instance of '{0}' object type from compiled assembly '{1}'.  Unreal Build Tool creates an instance of your module's 'Rules' object in order to find out about your module's requirements.  The CLR exception details may provide more information:  {2}", ModuleTypeName, AssemblyFileName, Ex.ToString());
+			}
+
+			// Have to do absolute here as this could be a project that is under the root
+			var FullUProjectPath = string.IsNullOrWhiteSpace(UnrealBuildTool.GetUProjectPath())
+				? ""
+				: Path.GetFullPath(UnrealBuildTool.GetUProjectPath());
+			var bProjectModule = string.IsNullOrWhiteSpace(FullUProjectPath)
+				? false
+				: Utils.IsFileUnderDirectory(ModuleFileName, FullUProjectPath);
+
+			if (bProjectModule)
+			{
+				RulesObject.ReadAdditionalDependencies(UnrealBuildTool.GetUProjectFile(), ModuleName);
+			}
+
+			// Validate rules object
+			{
+				if (RulesObject.Type == ModuleRules.ModuleType.CPlusPlus)
+				{
+					if (RulesObject.PrivateAssemblyReferences.Count > 0)
+					{
+						throw new BuildException("Module rules for '{0}' may not specify PrivateAssemblyReferences unless it is a CPlusPlusCLR module type.", AssemblyFileName);
+					}
+
+					var InvalidDependencies = RulesObject.DynamicallyLoadedModuleNames.Intersect(RulesObject.PublicDependencyModuleNames.Concat(RulesObject.PrivateDependencyModuleNames)).ToList();
+					if (InvalidDependencies.Count != 0)
+					{
+						throw new BuildException("Module rules for '{0}' should not be dependent on modules which are also dynamically loaded: {1}", ModuleName, String.Join(", ", InvalidDependencies));
+					}
+
+					// Choose code optimization options based on module type (game/engine) if
+					// default optimization method is selected.
+					bool bIsEngineModule = Utils.IsFileUnderDirectory(ModuleFileName, ProjectFileGenerator.EngineRelativePath);
+					bool IsGamePluginModule = bProjectModule ? Utils.IsFileUnderDirectory(ModuleFileName, Path.Combine(UnrealBuildTool.GetUProjectPath(), "Plugins")) : false;
+					if (RulesObject.OptimizeCode == ModuleRules.CodeOptimization.Default)
+					{
+						// Engine/Source and Engine/Plugins are considered 'Engine' code...
+						if (bIsEngineModule)
+						{
+							// Engine module - always optimize (except Debug).
+							RulesObject.OptimizeCode = ModuleRules.CodeOptimization.Always;
+						}
+						else
+						{
+							// Game module - do not optimize in Debug and DebugGame builds.
+							RulesObject.OptimizeCode = ModuleRules.CodeOptimization.InNonDebugBuilds;
+						}
+					}
+
+					// Disable shared PCHs for game modules by default (but not game plugins, since they won't depend on the game's PCH!)
+					if (RulesObject.PCHUsage == ModuleRules.PCHUsageMode.Default)
+					{
+						// Note that bIsEngineModule includes Engine/Plugins, so Engine/Plugins will use shared PCHs.
+						var IsProgramTarget = Target.Type != null && Target.Type == TargetRules.TargetType.Program;
+						if (bIsEngineModule || IsProgramTarget || IsGamePluginModule)
+						{
+							// Engine module or plugin module -- allow shared PCHs
+							RulesObject.PCHUsage = ModuleRules.PCHUsageMode.UseSharedPCHs;
+						}
+						else
+						{
+							// Game module.  Do not enable shared PCHs by default, because games usually have a large precompiled header of their own and compile times would suffer.
+							RulesObject.PCHUsage = ModuleRules.PCHUsageMode.NoSharedPCHs;
+						}
+					}
+				}
+			}
+
+			return RulesObject;
+		}
+
+		protected bool GetTargetTypeAndRulesInstance(string InTargetName, TargetInfo InTarget, out System.Type OutRulesObjectType, out TargetRules OutRulesObject)
+		{
+			// The build module must define a type named '<TargetName>Target' that derives from our 'TargetRules' type.  
+			OutRulesObjectType = CompiledAssembly.GetType(InTargetName);
+			if (OutRulesObjectType == null)
+			{
+				throw new BuildException(
+					"Expecting to find a type to be declared in a target rules named '{0}'.  This type must derive from the 'TargetRules' type defined by Unreal Build Tool.",
+					InTargetName);
+			}
+
+			// Create an instance of the module's rules object
+			try
+			{
+				OutRulesObject = (TargetRules)Activator.CreateInstance(OutRulesObjectType, InTarget);
+			}
+			catch (Exception Ex)
+			{
+				var AssemblyFileName = Path.GetFileNameWithoutExtension(CompiledAssembly.Location);
+				throw new BuildException(Ex,
+					"Unable to instantiate instance of '{0}' object type from compiled assembly '{1}'.  Unreal Build Tool creates an instance of your module's 'Rules' object in order to find out about your module's requirements.  The CLR exception details may provide more information:  {2}",
+					InTargetName, AssemblyFileName, Ex.ToString());
+			}
+
+			OutRulesObject.TargetName = InTargetName;
+
+			return true;
+		}
+
+		/// <summary>
+		/// Creates a target rules object for the specified target name.
+		/// </summary>
+		/// <param name="TargetName">Name of the target</param>
+		/// <param name="Target">Information about the target associated with this target</param>
+		/// <param name="TargetFileName">The original source file name of the Target.cs file for this target</param>
+		/// <returns>The build target rules for the specified target</returns>
+		public TargetRules CreateTargetRules(string TargetName, TargetInfo Target, bool bInEditorRecompile, out string TargetFileName)
+		{
+			// Make sure the target file is known to us
+			bool bFoundTargetName = TargetNameToTargetFileMap.ContainsKey(TargetName);
+			if (bFoundTargetName == false)
+			{
+				if (UnrealBuildTool.RunningRocket())
+				{
+					//@todo Rocket: Remove this when full game support is implemented
+					// If we are Rocket, they will currently only have an editor target.
+					// See if that exists
+					bFoundTargetName = TargetNameToTargetFileMap.ContainsKey(TargetName + "Editor");
+					if (bFoundTargetName)
+					{
+						TargetName += "Editor";
+					}
+				}
+			}
+
+			if (bFoundTargetName == false)
+			{
+				//				throw new BuildException("Couldn't find target rules file for target '{0}' in rules assembly '{1}'.", TargetName, RulesAssembly.FullName);
+				string ExceptionMessage = "Couldn't find target rules file for target '";
+				ExceptionMessage += TargetName;
+				ExceptionMessage += "' in rules assembly '";
+				ExceptionMessage += CompiledAssembly.FullName;
+				ExceptionMessage += "'." + Environment.NewLine;
+
+				ExceptionMessage += "Location: " + CompiledAssembly.Location + Environment.NewLine;
+
+				ExceptionMessage += "Target rules found:" + Environment.NewLine;
+				foreach (KeyValuePair<string, string> entry in TargetNameToTargetFileMap)
+				{
+					ExceptionMessage += "\t" + entry.Key + " - " + entry.Value + Environment.NewLine;
+				}
+
+				throw new BuildException(ExceptionMessage);
+			}
+
+			// Return the target file name to the caller
+			TargetFileName = TargetNameToTargetFileMap[TargetName];
+
+			// Currently, we expect the user's rules object type name to be the same as the module name + 'Target'
+			string TargetTypeName = TargetName + "Target";
+
+			// The build module must define a type named '<TargetName>Target' that derives from our 'TargetRules' type.  
+			System.Type RulesObjectType;
+			TargetRules RulesObject;
+			GetTargetTypeAndRulesInstance(TargetTypeName, Target, out RulesObjectType, out RulesObject);
+			if (bInEditorRecompile)
+			{
+				// Make sure this is an editor module.
+				if (RulesObject != null)
+				{
+					if (RulesObject.Type != TargetRules.TargetType.Editor)
+					{
+						// Not the editor... determine the editor project
+						string TargetSourceFolder = TargetFileName;
+						Int32 SourceFolderIndex = -1;
+						if (Utils.IsRunningOnMono)
+						{
+							TargetSourceFolder = TargetSourceFolder.Replace("\\", "/");
+							SourceFolderIndex = TargetSourceFolder.LastIndexOf("/Source/", StringComparison.InvariantCultureIgnoreCase);
+						}
+						else
+						{
+							TargetSourceFolder = TargetSourceFolder.Replace("/", "\\");
+							SourceFolderIndex = TargetSourceFolder.LastIndexOf("\\Source\\", StringComparison.InvariantCultureIgnoreCase);
+						}
+						if (SourceFolderIndex != -1)
+						{
+							TargetSourceFolder = TargetSourceFolder.Substring(0, SourceFolderIndex + 8);
+							foreach (KeyValuePair<string, string> CheckEntry in TargetNameToTargetFileMap)
+							{
+								if (CheckEntry.Value.StartsWith(TargetSourceFolder, StringComparison.InvariantCultureIgnoreCase))
+								{
+									if (CheckEntry.Key.Equals(TargetName, StringComparison.InvariantCultureIgnoreCase) == false)
+									{
+										// We have found a target in the same source folder that is not the original target found.
+										// See if it is the editor project
+										string CheckTargetTypeName = CheckEntry.Key + "Target";
+										System.Type CheckRulesObjectType;
+										TargetRules CheckRulesObject;
+										GetTargetTypeAndRulesInstance(CheckTargetTypeName, Target, out CheckRulesObjectType, out CheckRulesObject);
+										if (CheckRulesObject != null)
+										{
+											if (CheckRulesObject.Type == TargetRules.TargetType.Editor)
+											{
+												// Found it
+												// NOTE: This prevents multiple Editor targets from co-existing...
+												RulesObject = CheckRulesObject;
+												break;
+											}
+										}
+									}
+								}
+							}
+						}
+					}
+				}
+			}
+
+			return RulesObject;
+		}
+	}
 
 	public class RulesCompiler
 	{
@@ -1325,315 +1679,6 @@ namespace UnrealBuildTool
 			return SourceFiles;
 		}
 
-		class RulesAssembly
-		{
-			public Assembly Rules;
-
-			/// Maps module names to their actual xxx.Module.cs file on disk
-			public Dictionary<string, string> ModuleNameToModuleFileMap = new Dictionary<string, string>(StringComparer.InvariantCultureIgnoreCase);
-
-			/// Maps target names to their actual xxx.Target.cs file on disk
-			public Dictionary<string, string> TargetNameToTargetFileMap = new Dictionary<string, string>(StringComparer.InvariantCultureIgnoreCase);
-
-			public RulesAssembly(Assembly InRules)
-			{
-				Rules = InRules;
-			}
-
-			/// <summary>
-			/// 
-			/// </summary>
-			/// <param name="InModuleName"></param>
-			/// <returns></returns>
-			public string GetModuleFilename(string InModuleName)
-			{
-				// Make sure the module file is known to us
-				if (!ModuleNameToModuleFileMap.ContainsKey(InModuleName))
-				{
-					return "";
-				}
-
-				// Return the module file name to the caller
-				return ModuleNameToModuleFileMap[InModuleName];
-			}
-
-			public string GetTargetFilename(string InTargetName)
-			{
-				// Make sure the target file is known to us
-				if (!TargetNameToTargetFileMap.ContainsKey(InTargetName))
-				{
-					return "";
-				}
-
-				// Return the target file name to the caller
-				return TargetNameToTargetFileMap[InTargetName];
-			}
-
-			/// <summary>
-			/// Creates an instance of a module rules descriptor object for the specified module name
-			/// </summary>
-			/// <param name="ModuleName">Name of the module</param>
-			/// <param name="Target">Information about the target associated with this module</param>
-			/// <param name="ModuleFileName">The original source file name for the Module.cs file for this module</param>
-			/// <returns>Compiled module rule info</returns>
-			public ModuleRules CreateModuleRules(string ModuleName, TargetInfo Target, out string ModuleFileName)
-			{
-				var AssemblyFileName = Path.GetFileNameWithoutExtension(Rules.Location);
-
-				// Currently, we expect the user's rules object type name to be the same as the module name
-				var ModuleTypeName = ModuleName;
-
-				// Make sure the module file is known to us
-				if (!ModuleNameToModuleFileMap.ContainsKey(ModuleName))
-				{
-					throw new MissingModuleException(ModuleName);
-				}
-
-				// Return the module file name to the caller
-				ModuleFileName = ModuleNameToModuleFileMap[ModuleName];
-
-				UnrealTargetPlatform LocalPlatform = Target.Platform;
-				UnrealTargetConfiguration LocalConfiguration = Target.Configuration;
-				TargetInfo LocalTarget = new TargetInfo(LocalPlatform, LocalConfiguration, Target.Type.Value, Target.bIsMonolithic.Value);
-
-				// The build module must define a type named 'Rules' that derives from our 'ModuleRules' type.  
-				var RulesObjectType = Rules.GetType(ModuleName);
-
-				if (RulesObjectType == null)
-				{
-					// Temporary hack to avoid System namespace collisions
-					// @todo projectfiles: Make rules assemblies require namespaces.
-					RulesObjectType = Rules.GetType("UnrealBuildTool.Rules." + ModuleName);
-				}
-
-				if (RulesObjectType == null)
-				{
-					throw new BuildException("Expecting to find a type to be declared in a module rules named '{0}' in {1}.  This type must derive from the 'ModuleRules' type defined by Unreal Build Tool.", ModuleTypeName, Rules.FullName);
-				}
-
-				// Create an instance of the module's rules object
-				ModuleRules RulesObject;
-				try
-				{
-					RulesObject = (ModuleRules)Activator.CreateInstance(RulesObjectType, LocalTarget);
-				}
-				catch (Exception Ex)
-				{
-					throw new BuildException(Ex, "Unable to instantiate instance of '{0}' object type from compiled assembly '{1}'.  Unreal Build Tool creates an instance of your module's 'Rules' object in order to find out about your module's requirements.  The CLR exception details may provide more information:  {2}", ModuleTypeName, AssemblyFileName, Ex.ToString());
-				}
-
-				// Have to do absolute here as this could be a project that is under the root
-				var FullUProjectPath = string.IsNullOrWhiteSpace(UnrealBuildTool.GetUProjectPath())
-					? ""
-					: Path.GetFullPath(UnrealBuildTool.GetUProjectPath());
-				var bProjectModule = string.IsNullOrWhiteSpace(FullUProjectPath)
-					? false
-					: Utils.IsFileUnderDirectory(ModuleFileName, FullUProjectPath);
-
-				if (bProjectModule)
-				{
-					RulesObject.ReadAdditionalDependencies(UnrealBuildTool.GetUProjectFile(), ModuleName);
-				}
-
-				// Validate rules object
-				{
-					if (RulesObject.Type == ModuleRules.ModuleType.CPlusPlus)
-					{
-						if (RulesObject.PrivateAssemblyReferences.Count > 0)
-						{
-							throw new BuildException("Module rules for '{0}' may not specify PrivateAssemblyReferences unless it is a CPlusPlusCLR module type.", AssemblyFileName);
-						}
-
-						var InvalidDependencies = RulesObject.DynamicallyLoadedModuleNames.Intersect(RulesObject.PublicDependencyModuleNames.Concat(RulesObject.PrivateDependencyModuleNames)).ToList();
-						if (InvalidDependencies.Count != 0)
-						{
-							throw new BuildException("Module rules for '{0}' should not be dependent on modules which are also dynamically loaded: {1}", ModuleName, String.Join(", ", InvalidDependencies));
-						}
-
-						// Choose code optimization options based on module type (game/engine) if
-						// default optimization method is selected.
-						bool bIsEngineModule = Utils.IsFileUnderDirectory(ModuleFileName, ProjectFileGenerator.EngineRelativePath);
-						bool IsGamePluginModule = bProjectModule ? Utils.IsFileUnderDirectory(ModuleFileName, Path.Combine(UnrealBuildTool.GetUProjectPath(), "Plugins")) : false;
-						if (RulesObject.OptimizeCode == ModuleRules.CodeOptimization.Default)
-						{
-							// Engine/Source and Engine/Plugins are considered 'Engine' code...
-							if (bIsEngineModule)
-							{
-								// Engine module - always optimize (except Debug).
-								RulesObject.OptimizeCode = ModuleRules.CodeOptimization.Always;
-							}
-							else
-							{
-								// Game module - do not optimize in Debug and DebugGame builds.
-								RulesObject.OptimizeCode = ModuleRules.CodeOptimization.InNonDebugBuilds;
-							}
-						}
-
-						// Disable shared PCHs for game modules by default (but not game plugins, since they won't depend on the game's PCH!)
-						if (RulesObject.PCHUsage == ModuleRules.PCHUsageMode.Default)
-						{
-							// Note that bIsEngineModule includes Engine/Plugins, so Engine/Plugins will use shared PCHs.
-							var IsProgramTarget = Target.Type != null && Target.Type == TargetRules.TargetType.Program;
-							if (bIsEngineModule || IsProgramTarget || IsGamePluginModule)
-							{
-								// Engine module or plugin module -- allow shared PCHs
-								RulesObject.PCHUsage = ModuleRules.PCHUsageMode.UseSharedPCHs;
-							}
-							else
-							{
-								// Game module.  Do not enable shared PCHs by default, because games usually have a large precompiled header of their own and compile times would suffer.
-								RulesObject.PCHUsage = ModuleRules.PCHUsageMode.NoSharedPCHs;
-							}
-						}
-					}
-				}
-
-				return RulesObject;
-			}
-
-			protected bool GetTargetTypeAndRulesInstance(string InTargetName, TargetInfo InTarget, out System.Type OutRulesObjectType, out TargetRules OutRulesObject)
-			{
-				// The build module must define a type named '<TargetName>Target' that derives from our 'TargetRules' type.  
-				OutRulesObjectType = Rules.GetType(InTargetName);
-				if (OutRulesObjectType == null)
-				{
-					throw new BuildException(
-						"Expecting to find a type to be declared in a target rules named '{0}'.  This type must derive from the 'TargetRules' type defined by Unreal Build Tool.",
-						InTargetName);
-				}
-
-				// Create an instance of the module's rules object
-				try
-				{
-					OutRulesObject = (TargetRules)Activator.CreateInstance(OutRulesObjectType, InTarget);
-				}
-				catch (Exception Ex)
-				{
-					var AssemblyFileName = Path.GetFileNameWithoutExtension(Rules.Location);
-					throw new BuildException(Ex,
-						"Unable to instantiate instance of '{0}' object type from compiled assembly '{1}'.  Unreal Build Tool creates an instance of your module's 'Rules' object in order to find out about your module's requirements.  The CLR exception details may provide more information:  {2}",
-						InTargetName, AssemblyFileName, Ex.ToString());
-				}
-
-				OutRulesObject.TargetName = InTargetName;
-
-				return true;
-			}
-
-			/// <summary>
-			/// Creates a target rules object for the specified target name.
-			/// </summary>
-			/// <param name="TargetName">Name of the target</param>
-			/// <param name="Target">Information about the target associated with this target</param>
-			/// <param name="TargetFileName">The original source file name of the Target.cs file for this target</param>
-			/// <returns>The build target rules for the specified target</returns>
-			public TargetRules CreateTargetRules(string TargetName, TargetInfo Target, bool bInEditorRecompile, out string TargetFileName)
-			{
-				// Make sure the target file is known to us
-				bool bFoundTargetName = TargetNameToTargetFileMap.ContainsKey(TargetName);
-				if (bFoundTargetName == false)
-				{
-					if (UnrealBuildTool.RunningRocket())
-					{
-						//@todo Rocket: Remove this when full game support is implemented
-						// If we are Rocket, they will currently only have an editor target.
-						// See if that exists
-						bFoundTargetName = TargetNameToTargetFileMap.ContainsKey(TargetName + "Editor");
-						if (bFoundTargetName)
-						{
-							TargetName += "Editor";
-						}
-					}
-				}
-
-				if (bFoundTargetName == false)
-				{
-					//				throw new BuildException("Couldn't find target rules file for target '{0}' in rules assembly '{1}'.", TargetName, RulesAssembly.FullName);
-					string ExceptionMessage = "Couldn't find target rules file for target '";
-					ExceptionMessage += TargetName;
-					ExceptionMessage += "' in rules assembly '";
-					ExceptionMessage += Rules.FullName;
-					ExceptionMessage += "'." + Environment.NewLine;
-
-					ExceptionMessage += "Location: " + Rules.Location + Environment.NewLine;
-
-					ExceptionMessage += "Target rules found:" + Environment.NewLine;
-					foreach (KeyValuePair<string, string> entry in TargetNameToTargetFileMap)
-					{
-						ExceptionMessage += "\t" + entry.Key + " - " + entry.Value + Environment.NewLine;
-					}
-
-					throw new BuildException(ExceptionMessage);
-				}
-
-				// Return the target file name to the caller
-				TargetFileName = TargetNameToTargetFileMap[TargetName];
-
-				// Currently, we expect the user's rules object type name to be the same as the module name + 'Target'
-				string TargetTypeName = TargetName + "Target";
-
-				// The build module must define a type named '<TargetName>Target' that derives from our 'TargetRules' type.  
-				System.Type RulesObjectType;
-				TargetRules RulesObject;
-				GetTargetTypeAndRulesInstance(TargetTypeName, Target, out RulesObjectType, out RulesObject);
-				if (bInEditorRecompile)
-				{
-					// Make sure this is an editor module.
-					if (RulesObject != null)
-					{
-						if (RulesObject.Type != TargetRules.TargetType.Editor)
-						{
-							// Not the editor... determine the editor project
-							string TargetSourceFolder = TargetFileName;
-							Int32 SourceFolderIndex = -1;
-							if (Utils.IsRunningOnMono)
-							{
-								TargetSourceFolder = TargetSourceFolder.Replace("\\", "/");
-								SourceFolderIndex = TargetSourceFolder.LastIndexOf("/Source/", StringComparison.InvariantCultureIgnoreCase);
-							}
-							else
-							{
-								TargetSourceFolder = TargetSourceFolder.Replace("/", "\\");
-								SourceFolderIndex = TargetSourceFolder.LastIndexOf("\\Source\\", StringComparison.InvariantCultureIgnoreCase);
-							}
-							if (SourceFolderIndex != -1)
-							{
-								TargetSourceFolder = TargetSourceFolder.Substring(0, SourceFolderIndex + 8);
-								foreach (KeyValuePair<string, string> CheckEntry in TargetNameToTargetFileMap)
-								{
-									if (CheckEntry.Value.StartsWith(TargetSourceFolder, StringComparison.InvariantCultureIgnoreCase))
-									{
-										if (CheckEntry.Key.Equals(TargetName, StringComparison.InvariantCultureIgnoreCase) == false)
-										{
-											// We have found a target in the same source folder that is not the original target found.
-											// See if it is the editor project
-											string CheckTargetTypeName = CheckEntry.Key + "Target";
-											System.Type CheckRulesObjectType;
-											TargetRules CheckRulesObject;
-											GetTargetTypeAndRulesInstance(CheckTargetTypeName, Target, out CheckRulesObjectType, out CheckRulesObject);
-											if (CheckRulesObject != null)
-											{
-												if (CheckRulesObject.Type == TargetRules.TargetType.Editor)
-												{
-													// Found it
-													// NOTE: This prevents multiple Editor targets from co-existing...
-													RulesObject = CheckRulesObject;
-													break;
-												}
-											}
-										}
-									}
-								}
-							}
-						}
-					}
-				}
-
-				return RulesObject;
-			}
-
-		}
-
 		private class LoadedAssemblyData
 		{
 			public LoadedAssemblyData(RulesAssembly InAssembly, List<string> InGameFolders)
@@ -1723,48 +1768,18 @@ namespace UnrealBuildTool
 			}
 			var ModuleFileNames = FindAllRulesSourceFiles(RulesFileType.Module, AdditionalSearchPaths);
 
-			var AssemblySourceFiles = new List<string>();
-			AssemblySourceFiles.AddRange( ModuleFileNames );
-			if( AssemblySourceFiles.Count == 0 )
-			{
-				throw new BuildException("No module rules source files were found in any of the module base directories!");
-			}
-			var TargetFileNames = FindAllRulesSourceFiles(RulesFileType.Target, AdditionalSearchPaths);
-			AssemblySourceFiles.AddRange( TargetFileNames );
-
 			// Create a path to the assembly that we'll either load or compile
 			string BaseIntermediatePath = UnrealBuildTool.HasUProjectFile() ?
 				Path.Combine(UnrealBuildTool.GetUProjectPath(), BuildConfiguration.BaseIntermediateFolder) : BuildConfiguration.BaseIntermediatePath;
 
 			string OutputAssemblyPath = Path.GetFullPath(Path.Combine(BaseIntermediatePath, "BuildRules", AssemblyName + ".dll"));
 
-			RulesAssembly Rules = new RulesAssembly(DynamicCompilation.CompileAndLoadAssembly( OutputAssemblyPath, AssemblySourceFiles ));
-
-			// Setup the module map
-			foreach( var CurModuleFileName in ModuleFileNames )
-			{
-				var CleanFileName = Utils.CleanDirectorySeparators( CurModuleFileName );
-				var ModuleName = Path.GetFileNameWithoutExtension( Path.GetFileNameWithoutExtension( CleanFileName ) );	// Strip both extensions
-				if( !Rules.ModuleNameToModuleFileMap.ContainsKey( ModuleName ) )
-				{
-					Rules.ModuleNameToModuleFileMap.Add( ModuleName, CurModuleFileName );
-				}
-			}
-
-			// Setup the target map
-			foreach( var CurTargetFileName in TargetFileNames )
-			{
-				var CleanFileName = Utils.CleanDirectorySeparators( CurTargetFileName );
-				var TargetName = Path.GetFileNameWithoutExtension( Path.GetFileNameWithoutExtension( CleanFileName ) );	// Strip both extensions
-				if( !Rules.TargetNameToTargetFileMap.ContainsKey( TargetName ) )
-				{
-					Rules.TargetNameToTargetFileMap.Add( TargetName, CurTargetFileName );
-				}
-			}
+			var TargetFileNames = FindAllRulesSourceFiles(RulesFileType.Target, AdditionalSearchPaths);
 
+			RulesAssembly Rules = new RulesAssembly(ModuleFileNames, TargetFileNames, OutputAssemblyPath);
+				
 			// Remember that we loaded this assembly
-			var RulesAssemblyName = Path.GetFileNameWithoutExtension( Rules.Rules.Location );
-			LoadedAssemblyMap[RulesAssemblyName] = new LoadedAssemblyData(Rules, AllGameFolders);
+			LoadedAssemblyMap[Rules.Name] = new LoadedAssemblyData(Rules, AllGameFolders);
 		}
 
 		/// <summary>
