Add launch image support for Android (disabled for GearVR)
#ue4
#android

diff --git a/Engine/Build/Android/Java/JavaTemplates/DownloaderActivity.java.template b/Engine/Build/Android/Java/JavaTemplates/DownloaderActivity.java.template
--- a/Engine/Build/Android/Java/JavaTemplates/DownloaderActivity.java.template
+++ b/Engine/Build/Android/Java/JavaTemplates/DownloaderActivity.java.template
@@ -233,6 +233,7 @@ public class DownloaderActivity extends Activity implements IDownloaderClient {
 				OutputData.putExtra(GameActivity.DOWNLOAD_RETURN_NAME, GameActivity.DOWNLOAD_FILES_PRESENT);		
 				setResult(RESULT_OK, OutputData);
 				finish();
+				overridePendingTransition(R.anim.noaction, R.anim.noaction);
 		}
 	}
 
diff --git a/Engine/Build/Android/Java/res/anim/noaction.xml b/Engine/Build/Android/Java/res/anim/noaction.xml
new file mode 100644
--- /dev/null
+++ b/Engine/Build/Android/Java/res/anim/noaction.xml
@@ -0,0 +1,4 @@
+<?xml version="1.0" encoding="utf-8"?>
+<set xmlns:android="http://schemas.android.com/apk/res/android">
+	<alpha android:fromAlpha="1.0" android:toAlpha="1.0" android:duration="5" />
+</set>
\ No newline at end of file
diff --git a/Engine/Build/Android/Java/res/values-land/styles.xml b/Engine/Build/Android/Java/res/values-land/styles.xml
new file mode 100644
--- /dev/null
+++ b/Engine/Build/Android/Java/res/values-land/styles.xml
@@ -0,0 +1,12 @@
+<?xml version="1.0" encoding="utf-8"?>
+<resources>
+
+    <style name="UE4BaseTheme" parent = "@android:style/Theme.Black.NoTitleBar.Fullscreen" />
+
+    <style name="UE4SplashTheme" parent = "UE4BaseTheme">
+        <item name="android:windowBackground">@drawable/splashscreen_landscape</item>
+        <item name="android:windowNoTitle">true</item>
+        <item name="android:windowFullscreen">true</item>
+    </style>
+
+</resources>
\ No newline at end of file
diff --git a/Engine/Build/Android/Java/res/values-port/styles.xml b/Engine/Build/Android/Java/res/values-port/styles.xml
new file mode 100644
--- /dev/null
+++ b/Engine/Build/Android/Java/res/values-port/styles.xml
@@ -0,0 +1,12 @@
+<?xml version="1.0" encoding="utf-8"?>
+<resources>
+
+    <style name="UE4BaseTheme" parent = "@android:style/Theme.Black.NoTitleBar.Fullscreen" />
+
+    <style name="UE4SplashTheme" parent = "UE4BaseTheme">
+        <item name="android:windowBackground">@drawable/splashscreen_portrait</item>
+        <item name="android:windowNoTitle">true</item>
+        <item name="android:windowFullscreen">true</item>
+    </style>
+
+</resources>
\ No newline at end of file
diff --git a/Engine/Build/Android/Java/src/com/epicgames/ue4/GameActivity.java b/Engine/Build/Android/Java/src/com/epicgames/ue4/GameActivity.java
--- a/Engine/Build/Android/Java/src/com/epicgames/ue4/GameActivity.java
+++ b/Engine/Build/Android/Java/src/com/epicgames/ue4/GameActivity.java
@@ -97,6 +97,9 @@ public class GameActivity extends NativeActivity
 	
 	static GameActivity _activity;
 
+	protected Dialog mSplashDialog;
+	private int noActionAnimID = -1;
+
 	// Console
 	AlertDialog consoleAlert;
 	EditText consoleInputBox;
@@ -221,10 +224,43 @@ public class GameActivity extends NativeActivity
 		}
 	}
 
+	private int getResourceId(String VariableName, String ResourceName, String PackageName)
+	{
+		try {
+			return getResources().getIdentifier(VariableName, ResourceName, PackageName);
+		}
+		catch (Exception e) {
+			e.printStackTrace();
+			return -1;
+		} 
+	}
+
 	@Override
 	public void onCreate(Bundle savedInstanceState)
 	{
 		super.onCreate(savedInstanceState);
+
+		// create splashscreen dialog (if launched by SplashActivity)
+		Bundle intentBundle = getIntent().getExtras();
+		if (intentBundle != null && intentBundle.getString("UseSplashScreen") != null)
+		{
+			try {
+				// try to get the splash theme (can't use R.style.UE4SplashTheme since we don't know the package name until runtime)
+				int SplashThemeId = getResources().getIdentifier("UE4SplashTheme", "style", getPackageName());
+				mSplashDialog = new Dialog(this, SplashThemeId);
+				mSplashDialog.setCancelable(false);
+				mSplashDialog.show();
+			}
+			catch (Exception e) {
+				e.printStackTrace();
+			}
+			try {
+				noActionAnimID = getResources().getIdentifier("noaction", "anim", getPackageName());
+			}
+			catch (Exception e) {
+				e.printStackTrace();
+			}
+		}
 		
 		// Suppress java logs in Shipping builds
 		if (nativeIsShippingBuild())
@@ -633,11 +669,16 @@ public class GameActivity extends NativeActivity
 		else
 		{
 			// Start the check activity here
-			Log.debug("==============> Starting activity to check files and download if required");			
+			Log.debug("==============> Starting activity to check files and download if required");
 			Intent intent = new Intent(this, DownloadShim.GetDownloaderType());
+			intent.addFlags(Intent.FLAG_ACTIVITY_NO_ANIMATION);
 			startActivityForResult(intent, DOWNLOAD_ACTIVITY_ID);
+			if (noActionAnimID != -1)
+			{
+				overridePendingTransition(noActionAnimID, noActionAnimID);
+			}
 		}
-		
+
 		Log.debug("==============> GameActive.onResume complete!");
 	}
 
@@ -1154,7 +1195,10 @@ public class GameActivity extends NativeActivity
 			|| errorCode == DOWNLOAD_FAILED 
 			|| errorCode == DOWNLOAD_INVALID
 			|| errorCode == DOWNLOAD_NO_PLAY_KEY)
+			{
 				finish();
+				return;
+			}
 		}
 		else if( IapStoreHelper != null )
 		{
@@ -1208,6 +1252,15 @@ public class GameActivity extends NativeActivity
 		return bIsAllowedToMakePurchase;
 	}
 
+	public void AndroidThunkJava_DismissSplashScreen()
+	{
+		if (mSplashDialog != null)
+		{
+			mSplashDialog.dismiss();
+			mSplashDialog = null;
+		}
+	}
+
 	public native boolean nativeIsShippingBuild();
 	public native void nativeSetGlobalActivity();
 	public native void nativeSetWindowInfo(boolean bIsPortrait, int DepthBufferPreference);
diff --git a/Engine/Build/Android/Java/src/com/epicgames/ue4/SplashActivity.java b/Engine/Build/Android/Java/src/com/epicgames/ue4/SplashActivity.java
new file mode 100644
--- /dev/null
+++ b/Engine/Build/Android/Java/src/com/epicgames/ue4/SplashActivity.java
@@ -0,0 +1,32 @@
+//This file needs to be here so the "ant" build step doesnt fail when looking for a /src folder.
+
+package com.epicgames.ue4;
+
+import android.os.Bundle;
+import android.app.Activity;
+import android.content.Intent;
+
+public class SplashActivity extends Activity
+{
+	@Override
+	protected void onCreate(Bundle savedInstanceState)
+	{
+		super.onCreate(savedInstanceState);
+
+		Intent intent = new Intent(this, GameActivity.class);
+		intent.addFlags(Intent.FLAG_ACTIVITY_NO_ANIMATION);
+		intent.putExtra("UseSplashScreen", "true");
+		startActivity(intent);
+		finish();
+		overridePendingTransition(0, 0);
+	}
+
+	@Override
+	protected void onPause()
+	{
+		super.onPause();
+		finish();
+		overridePendingTransition(0, 0);
+	}
+
+}
diff --git a/Engine/Config/BaseEngine.ini b/Engine/Config/BaseEngine.ini
--- a/Engine/Config/BaseEngine.ini
+++ b/Engine/Config/BaseEngine.ini
@@ -2319,6 +2319,7 @@ PackageName=com.YourCompany.[PROJECT]
 StoreVersion=1
 VersionDisplayName=1.0
 MinSDKVersion=9
+bShowLaunchImage=true
 
 [/Script/AndroidPlatformEditor.AndroidSDKSettings]
 SDKAPILevel=latest
diff --git a/Engine/Source/Developer/Android/AndroidPlatformEditor/Private/AndroidTargetSettingsCustomization.cpp b/Engine/Source/Developer/Android/AndroidPlatformEditor/Private/AndroidTargetSettingsCustomization.cpp
--- a/Engine/Source/Developer/Android/AndroidPlatformEditor/Private/AndroidTargetSettingsCustomization.cpp
+++ b/Engine/Source/Developer/Android/AndroidPlatformEditor/Private/AndroidTargetSettingsCustomization.cpp
@@ -46,7 +46,10 @@ FAndroidTargetSettingsCustomization::FAndroidTargetSettingsCustomization()
 	new (IconNames) FPlatformIconInfo(TEXT("res/drawable-mdpi/icon.png"), LOCTEXT("SettingsIcon_MDPI", "MDPI Icon"), FText::GetEmpty(), 48, 48, FPlatformIconInfo::Required);
 	new (IconNames) FPlatformIconInfo(TEXT("res/drawable-hdpi/icon.png"), LOCTEXT("SettingsIcon_HDPI", "HDPI Icon"), FText::GetEmpty(), 72, 72, FPlatformIconInfo::Required);
 	new (IconNames) FPlatformIconInfo(TEXT("res/drawable-xhdpi/icon.png"), LOCTEXT("SettingsIcon_XHDPI", "XHDPI Icon"), FText::GetEmpty(), 96, 96, FPlatformIconInfo::Required);
-	new (IconNames) FPlatformIconInfo(TEXT("res/drawable/downloadimagev.png"), LOCTEXT("SettingsIcon_DownloadImageV", "Download Background Vertical Image"), FText::GetEmpty(), 720, 1280, FPlatformIconInfo::Required);
+
+	new (LaunchImageNames)FPlatformIconInfo(TEXT("res/drawable/downloadimagev.png"), LOCTEXT("SettingsIcon_DownloadImageV", "Download Background Vertical Image"), FText::GetEmpty(), 720, 1280, FPlatformIconInfo::Required);
+	new (LaunchImageNames)FPlatformIconInfo(TEXT("res/drawable/splashscreen_portrait.png"), LOCTEXT("LaunchImage_Portrait", "Launch Portrait"), FText::GetEmpty(), 360, 640, FPlatformIconInfo::Required);
+	new (LaunchImageNames)FPlatformIconInfo(TEXT("res/drawable/splashscreen_landscape.png"), LOCTEXT("LaunchImage_Landscape", "Launch Landscape"), FText::GetEmpty(), 640, 360, FPlatformIconInfo::Required);
 }
 
 void FAndroidTargetSettingsCustomization::CustomizeDetails(IDetailLayoutBuilder& DetailLayout)
@@ -55,6 +58,7 @@ void FAndroidTargetSettingsCustomization::CustomizeDetails(IDetailLayoutBuilder&
 
 	BuildAppManifestSection(DetailLayout);
 	BuildIconSection(DetailLayout);
+	BuildLaunchImageSection(DetailLayout);
 }
 
 static void OnBrowserLinkClicked(const FSlateHyperlinkRun::FMetadata& Metadata)
@@ -266,6 +270,69 @@ void FAndroidTargetSettingsCustomization::BuildIconSection(IDetailLayoutBuilder&
 	}
 }
 
+void FAndroidTargetSettingsCustomization::BuildLaunchImageSection(IDetailLayoutBuilder& DetailLayout)
+{
+	// Add the launch images
+	IDetailCategoryBuilder& LaunchImageCategory = DetailLayout.EditCategory(TEXT("LaunchImages"));
+	LaunchImageCategory.AddCustomRow(LOCTEXT("LaunchImageInfo", "Launch Image Info"), false)
+		.WholeRowWidget
+		[
+			SNew(SBorder)
+			.Padding(1)
+			[
+				SNew(SHorizontalBox)
+				+ SHorizontalBox::Slot()
+				.Padding(FMargin(10, 10, 10, 10))
+				.FillWidth(1.0f)
+				[
+					SNew(SRichTextBlock)
+					.Text(LOCTEXT("LaunchImageInfoMessage", "The <RichTextBlock.TextHighlight>Download Background</> image is used as the background for OBB downloading.  The <RichTextBlock.TextHighlight>Launch Portrait</> image is used as a splash screen for applications with Portrait, Reverse Portrait, Sensor Portrait, Sensor, or Full Sensor orientation.  The <RichTextBlock.TextHighlight>Launch Landscape</> image is used as a spash screen for applications with Landscape, Sensor Landscape, Reverse Landscape, Sensor, or Full Sensor orientation.\n\nThe launch images will be scaled to fit the device in the active orientation. Additional optional launch images may be provided as overrides for LDPI, MDPI, HDPI, and XHDPI by placing them in the project's corresponding Build/Android/res/drawable-* directory."))
+					.TextStyle(FEditorStyle::Get(), "MessageLog")
+					.DecoratorStyleSet(&FEditorStyle::Get())
+					.AutoWrapText(true)
+					+ SRichTextBlock::HyperlinkDecorator(TEXT("browser"), FSlateHyperlinkRun::FOnClick::CreateStatic(&OnBrowserLinkClicked))
+				]
+			]
+		];
+
+	const FVector2D LaunchImageMaxSize(150.0f, 150.0f);
+
+	for (const FPlatformIconInfo& Info : LaunchImageNames)
+	{
+		const FString AutomaticImagePath = EngineAndroidPath / Info.IconPath;
+		const FString TargetImagePath = GameAndroidPath / Info.IconPath;
+
+		LaunchImageCategory.AddCustomRow(Info.IconName)
+			.NameContent()
+			[
+				SNew(SHorizontalBox)
+				+ SHorizontalBox::Slot()
+				.Padding(FMargin(0, 1, 0, 1))
+				.FillWidth(1.0f)
+				[
+					SNew(STextBlock)
+					.Text(Info.IconName)
+					.Font(DetailLayout.GetDetailFont())
+				]
+			]
+			.ValueContent()
+			.MaxDesiredWidth(400.0f)
+			.MinDesiredWidth(100.0f)
+			[
+				SNew(SHorizontalBox)
+				+ SHorizontalBox::Slot()
+				.FillWidth(1.0f)
+				.VAlign(VAlign_Center)
+				[
+					SNew(SExternalImageReference, AutomaticImagePath, TargetImagePath)
+					.FileDescription(Info.IconDescription)
+//					.RequiredSize(Info.IconRequiredSize)
+					.MaxDisplaySize(LaunchImageMaxSize)
+				]
+			];
+	}
+}
+
 FReply FAndroidTargetSettingsCustomization::OpenBuildFolder()
 {
 	const FString BuildFolder = FPaths::ConvertRelativePathToFull(FPaths::GetPath(GameProjectPropertiesPath));
@@ -298,6 +365,18 @@ void FAndroidTargetSettingsCustomization::CopySetupFilesIntoProject()
 			}
 		}
 
+		// Now try to copy all of the launch images... (these can be ignored if the file already exists)
+		for (const FPlatformIconInfo& Info : LaunchImageNames)
+		{
+			const FString EngineImagePath = EngineAndroidPath / Info.IconPath;
+			const FString ProjectImagePath = GameAndroidPath / Info.IconPath;
+
+			if (!FPaths::FileExists(ProjectImagePath))
+			{
+				SourceControlHelpers::CopyFileUnderSourceControl(ProjectImagePath, EngineImagePath, Info.IconName, /*out*/ ErrorMessage);
+			}
+		}
+
 		// and copy the other files (aren't required)
 		SourceControlHelpers::CopyFileUnderSourceControl(GameProguardPath, EngineProguardPath, LOCTEXT("Proguard", "Proguard Settings"), /*out*/ ErrorMessage);
 	}
diff --git a/Engine/Source/Developer/Android/AndroidPlatformEditor/Private/AndroidTargetSettingsCustomization.h b/Engine/Source/Developer/Android/AndroidPlatformEditor/Private/AndroidTargetSettingsCustomization.h
--- a/Engine/Source/Developer/Android/AndroidPlatformEditor/Private/AndroidTargetSettingsCustomization.h
+++ b/Engine/Source/Developer/Android/AndroidPlatformEditor/Private/AndroidTargetSettingsCustomization.h
@@ -24,6 +24,7 @@ private:
 
 	void BuildAppManifestSection(IDetailLayoutBuilder& DetailLayout);
 	void BuildIconSection(IDetailLayoutBuilder& DetailLayout);
+	void BuildLaunchImageSection(IDetailLayoutBuilder& DetailLayout);
 
 	// Navigates to the build files in the explorer or finder
 	FReply OpenBuildFolder();
@@ -53,6 +54,7 @@ private:
 	const FString GameProjectPropertiesPath;
 
 	TArray<struct FPlatformIconInfo> IconNames;
+	TArray<struct FPlatformIconInfo> LaunchImageNames;
 
 	// Is the manifest writable?
 	TAttribute<bool> SetupForPlatformAttribute;
diff --git a/Engine/Source/Programs/AutomationTool/Android/AndroidPlatform.Automation.cs b/Engine/Source/Programs/AutomationTool/Android/AndroidPlatform.Automation.cs
--- a/Engine/Source/Programs/AutomationTool/Android/AndroidPlatform.Automation.cs
+++ b/Engine/Source/Programs/AutomationTool/Android/AndroidPlatform.Automation.cs
@@ -800,6 +800,7 @@ public class AndroidPlatform : Platform
 	/** Internal usage for GetPackageName */
 	private static string PackageLine = null;
 	private static Mutex PackageInfoMutex = new Mutex();
+	private static string LaunchableActivityLine = null;
 
 	/** Run an external exe (and capture the output), given the exe path and the commandline. */
 	private static string GetPackageInfo(string ApkName, bool bRetrieveVersionCode)
@@ -815,6 +816,7 @@ public class AndroidPlatform : Platform
 		using (var GameProcess = Process.Start(ExeInfo))
 		{
 			PackageLine = null;
+			LaunchableActivityLine = null;
 			GameProcess.BeginOutputReadLine();
 			GameProcess.OutputDataReceived += ParsePackageName;
 			GameProcess.WaitForExit();
@@ -836,6 +838,22 @@ public class AndroidPlatform : Platform
 		return ReturnValue;
 	}
 
+	/** Returns the launch activity name to launch (must call GetPackageInfo first), returns "com.epicgames.ue4.SplashActivity" default if not found */
+	private static string GetLaunchableActivityName()
+	{
+		string ReturnValue = "com.epicgames.ue4.SplashActivity";
+		if (LaunchableActivityLine != null)
+		{
+			// the line should look like: launchable-activity: name='com.epicgames.ue4.SplashActivity'  label='TappyChicken' icon=''
+			string[] Tokens = LaunchableActivityLine.Split("'".ToCharArray());
+			if (Tokens.Length >= 2)
+			{
+				ReturnValue = Tokens[1];
+			}
+		}
+		return ReturnValue;
+	}
+
 	/** Simple function to pipe output asynchronously */
 	private static void ParsePackageName(object Sender, DataReceivedEventArgs Event)
 	{
@@ -851,6 +869,14 @@ public class AndroidPlatform : Platform
 					PackageLine = Line;
 				}
 			}
+			if (LaunchableActivityLine == null)
+			{
+				string Line = Event.Data;
+				if (Line.StartsWith("launchable-activity:"))
+				{
+					LaunchableActivityLine = Line;
+				}
+			}
 		}
 	}
 
@@ -986,7 +1012,7 @@ public class AndroidPlatform : Platform
 		}
 
 		// start the app on device!
-		string CommandLine = "shell am start -n " + PackageName + "/com.epicgames.ue4.GameActivity";
+		string CommandLine = "shell am start -n " + PackageName + "/" + GetLaunchableActivityName();
 		ProcessResult ClientProcess = RunAdbCommand(Params, CommandLine, null, ClientRunFlags);
 
 
diff --git a/Engine/Source/Programs/UnrealBuildTool/Android/UEDeployAndroid.cs b/Engine/Source/Programs/UnrealBuildTool/Android/UEDeployAndroid.cs
--- a/Engine/Source/Programs/UnrealBuildTool/Android/UEDeployAndroid.cs
+++ b/Engine/Source/Programs/UnrealBuildTool/Android/UEDeployAndroid.cs
@@ -945,6 +945,123 @@ namespace UnrealBuildTool.Android
 			}
 		}
 
+		private void PickSplashScreenOrientation(string UE4BuildPath)
+		{
+			ConfigCacheIni Ini = GetConfigCacheIni("Engine");
+			bool bShowLaunchImage = false;
+			Ini.GetBool("/Script/AndroidRuntimeSettings.AndroidRuntimeSettings", "bShowLaunchImage", out bShowLaunchImage);
+			bool bPackageForGearVR;
+			Ini.GetBool("/Script/AndroidRuntimeSettings.AndroidRuntimeSettings", "bPackageForGearVR", out bPackageForGearVR);
+
+			if (bPackageForGearVR)
+			{
+				bShowLaunchImage = false;
+			}
+
+			// Decide which splash screen orientation(s) are needed based on orientation setting if enabled
+			bool bNeedPortrait = false;
+			bool bNeedLandscape = false;
+			if (bShowLaunchImage)
+			{
+				string Orientation;
+				Ini.GetString("/Script/AndroidRuntimeSettings.AndroidRuntimeSettings", "Orientation", out Orientation);
+
+				switch (Orientation.ToLower())
+				{
+					case "portrait":
+						bNeedPortrait = true;
+						break;
+					case "reverseportrait":
+						bNeedPortrait = true;
+						break;
+					case "sensorportrait":
+						bNeedPortrait = true;
+						break;
+
+					case "landscape":
+						bNeedLandscape = true;
+						break;
+					case "reverselandscape":
+						bNeedLandscape = true;
+						break;
+					case "sensorlandscape":
+						bNeedLandscape = true;
+						break;
+
+					case "sensor":
+						bNeedPortrait = true;
+						bNeedLandscape = true;
+						break;
+					case "fullsensor":
+						bNeedPortrait = true;
+						bNeedLandscape = true;
+						break;
+
+					default:
+						bNeedPortrait = true;
+						bNeedLandscape = true;
+						break;
+				}
+			}
+
+			// Remove unused styles.xml to prevent missing resource
+			if (!bNeedPortrait)
+			{
+				string StylesPath = UE4BuildPath + "/res/values-port/styles.xml";
+				if (File.Exists(StylesPath))
+				{
+					File.Delete(StylesPath);
+				}
+			}
+			if (!bNeedLandscape)
+			{
+				string StylesPath = UE4BuildPath + "/res/values-land/styles.xml";
+				if (File.Exists(StylesPath))
+				{
+					File.Delete(StylesPath);
+				}
+			}
+
+			// Loop through each of the resolutions (only /res/drawable/ is required, others are optional)
+			string[] Resolutions = new string[] { "/res/drawable/", "/res/drawable-ldpi/", "/res/drawable-mdpi/", "/res/drawable-hdpi/", "/res/drawable-xhdpi/" };
+			foreach (string ResolutionPath in Resolutions)
+			{
+				string PortraitFilename = UE4BuildPath + ResolutionPath + "splashscreen_portrait.png";
+				if (bNeedPortrait)
+				{
+					if (!File.Exists(PortraitFilename) && (ResolutionPath == "/res/drawable/"))
+					{
+						Log.TraceWarning("Warning: Splash screen source image {0} not available, splash screen will not function properly!", PortraitFilename);
+					}
+				}
+				else
+				{
+					// Remove unused image
+					if (File.Exists(PortraitFilename))
+					{
+						File.Delete(PortraitFilename);
+					}
+				}
+
+				string LandscapeFilename = UE4BuildPath + ResolutionPath + "splashscreen_landscape.png";
+				if (bNeedLandscape)
+				{
+					if (!File.Exists(LandscapeFilename) && (ResolutionPath == "/res/drawable/"))
+					{
+						Log.TraceWarning("Warning: Splash screen source image {0} not available, splash screen will not function properly!", LandscapeFilename);
+					}
+				}
+				else
+				{
+					// Remove unused image
+					if (File.Exists(LandscapeFilename))
+					{
+						File.Delete(LandscapeFilename);
+					}
+				}
+			}
+		}
+
         private string GetPackageName(string ProjectName)
         {
             ConfigCacheIni Ini = GetConfigCacheIni("Engine");
@@ -1003,11 +1120,27 @@ namespace UnrealBuildTool.Android
 			Ini.GetBool("/Script/AndroidRuntimeSettings.AndroidRuntimeSettings", "bPackageForGearVR", out bPackageForGearVR);
 			bool bEnableIAP = false;
 			Ini.GetBool("OnlineSubsystemGooglePlay.Store", "bSupportsInAppPurchasing", out bEnableIAP);
+			bool bShowLaunchImage = false;
+			Ini.GetBool("/Script/AndroidRuntimeSettings.AndroidRuntimeSettings", "bShowLaunchImage", out bShowLaunchImage);
 
 			// disable GearVR if not supported platform (in this case only armv7 for now)
 			if (UE4Arch != "-armv7")
 			{
-				bPackageForGearVR = false;
+				if (bPackageForGearVR)
+				{
+					Log.TraceInformation("Disabling Package For GearVR for unsupported architecture {0}", UE4Arch);
+					bPackageForGearVR = false;
+				}
+			}
+
+			// disable splash screen for GearVR (for now)
+			if (bPackageForGearVR)
+			{
+				if (bShowLaunchImage)
+				{
+					Log.TraceInformation("Disabling Show Launch Image for GearVR enabled application");
+					bShowLaunchImage = false;
+				}
 			}
 
 			StringBuilder Text = new StringBuilder();
@@ -1037,23 +1170,39 @@ namespace UnrealBuildTool.Android
 				}
 			}
 			Text.AppendLine("\t             android:hasCode=\"true\">");
-			Text.AppendLine("\t\t<activity android:name=\"com.epicgames.ue4.GameActivity\"");
-			Text.AppendLine("\t\t          android:label=\"@string/app_name\"");
-			if (!bPackageForGearVR)
+			if (bShowLaunchImage)
 			{
 				// normal application settings
-				Text.AppendLine("\t\t          android:theme=\"@android:style/Theme.NoTitleBar.Fullscreen\"");
-				Text.AppendLine("\t\t          android:configChanges=\"screenSize|orientation|keyboardHidden\"");
+				Text.AppendLine("\t\t<activity android:name=\"com.epicgames.ue4.SplashActivity\"");
+				Text.AppendLine("\t\t          android:label=\"@string/app_name\"");
+				Text.AppendLine("\t\t          android:theme=\"@style/UE4SplashTheme\"");
+				Text.AppendLine("\t\t          android:launchMode=\"singleTask\"");
+				if (bPackageForGearVR && bIsForDistribution)
+				{
+					Text.AppendLine("\t\t          android:excludeFromRecents=\"true\"");
+				}
+				Text.AppendLine(string.Format("\t\t          android:screenOrientation=\"{0}\"", ConvertOrientationIniValue(Orientation)));
+				Text.AppendLine(string.Format("\t\t          android:debuggable=\"{0}\">", bIsForDistribution ? "false" : "true"));
+				Text.AppendLine("\t\t\t<intent-filter>");
+				Text.AppendLine("\t\t\t\t<action android:name=\"android.intent.action.MAIN\" />");
+				Text.AppendLine(string.Format("\t\t\t\t<category android:name=\"android.intent.category.{0}\" />", (bPackageForGearVR && bIsForDistribution) ? "INFO" : "LAUNCHER"));
+				Text.AppendLine("\t\t\t</intent-filter>");
+				Text.AppendLine("\t\t</activity>");
+				Text.AppendLine("\t\t<activity android:name=\"com.epicgames.ue4.GameActivity\"");
+				Text.AppendLine("\t\t          android:label=\"@string/app_name\"");
+				Text.AppendLine("\t\t          android:theme=\"@style/UE4SplashTheme\"");
+				Text.AppendLine("\t\t          android:configChanges=\"screenSize|orientation|keyboardHidden|keyboard\"");
 			}
 			else
 			{
-				// GearVR
+				Text.AppendLine("\t\t<activity android:name=\"com.epicgames.ue4.GameActivity\"");
+				Text.AppendLine("\t\t          android:label=\"@string/app_name\"");
 				Text.AppendLine("\t\t          android:theme=\"@android:style/Theme.Black.NoTitleBar.Fullscreen\"");
 				Text.AppendLine("\t\t          android:configChanges=\"screenSize|orientation|keyboardHidden|keyboard\"");
-				if (bIsForDistribution)
-				{
-					Text.AppendLine("\t\t          android:excludeFromRecents=\"true\"");
-				}
+			}
+			if (bPackageForGearVR && bIsForDistribution)
+			{
+				Text.AppendLine("\t\t          android:excludeFromRecents=\"true\"");
 			}
 			Text.AppendLine("\t\t          android:launchMode=\"singleTask\"");
 			Text.AppendLine(string.Format("\t\t          android:screenOrientation=\"{0}\"", ConvertOrientationIniValue(Orientation)));
@@ -1066,10 +1215,13 @@ namespace UnrealBuildTool.Android
 			}
 			Text.AppendLine(string.Format("\t\t          android:debuggable=\"{0}\">", bIsForDistribution ? "false" : "true"));
 			Text.AppendLine("\t\t\t<meta-data android:name=\"android.app.lib_name\" android:value=\"UE4\"/>");
-			Text.AppendLine("\t\t\t<intent-filter>");
-			Text.AppendLine("\t\t\t\t<action android:name=\"android.intent.action.MAIN\" />");
-			Text.AppendLine(string.Format("\t\t\t\t<category android:name=\"android.intent.category.{0}\" />", (bPackageForGearVR && bIsForDistribution) ? "INFO" : "LAUNCHER"));
-			Text.AppendLine("\t\t\t</intent-filter>");
+			if (!bShowLaunchImage)
+			{
+				Text.AppendLine("\t\t\t<intent-filter>");
+				Text.AppendLine("\t\t\t\t<action android:name=\"android.intent.action.MAIN\" />");
+				Text.AppendLine(string.Format("\t\t\t\t<category android:name=\"android.intent.category.{0}\" />", (bPackageForGearVR && bIsForDistribution) ? "INFO" : "LAUNCHER"));
+				Text.AppendLine("\t\t\t</intent-filter>");
+			}
 			if (!string.IsNullOrEmpty(ExtraActivitySettings))
 			{
 				ExtraActivitySettings = ExtraActivitySettings.Replace("\\n", "\n");
@@ -1089,7 +1241,15 @@ namespace UnrealBuildTool.Android
 			Text.AppendLine("\t\t</activity>");
 
             // For OBB download support
-            Text.AppendLine("\t\t<activity android:name=\".DownloaderActivity\" />");
+			if (bShowLaunchImage)
+			{
+				Text.AppendLine("\t\t<activity android:name=\".DownloaderActivity\"");
+				Text.AppendLine("\t\t          android:theme=\"@style/UE4SplashTheme\" />");
+			}
+			else
+			{
+				Text.AppendLine("\t\t<activity android:name=\".DownloaderActivity\" />");
+			}
 
 			Text.AppendLine(string.Format("\t\t<meta-data android:name=\"com.epicgames.ue4.GameActivity.DepthBufferPreference\" android:value=\"{0}\"/>", ConvertDepthBufferIniValue(DepthBufferPreference)));
 			Text.AppendLine(string.Format("\t\t<meta-data android:name=\"com.epicgames.ue4.GameActivity.bPackageDataInsideApk\" android:value=\"{0}\"/>", bPackageDataInsideApk ? "true" : "false"));
@@ -1448,6 +1608,9 @@ namespace UnrealBuildTool.Android
 			CopyFileDirectory(GameBuildFilesPath + "/NotForLicensees", UE4BuildPath, Replacements);
 			CopyFileDirectory(GameBuildFilesPath + "/NoRedist", UE4BuildPath, Replacements);	
 
+			//Now keep the splash screen images matching orientation requested
+			PickSplashScreenOrientation(UE4BuildPath);
+
 			// at this point, we can write out the cached build settings to compare for a next build
 			File.WriteAllText(BuildSettingsCacheFile, CurrentBuildSettings);
 
diff --git a/Engine/Source/Runtime/Android/AndroidRuntimeSettings/Classes/AndroidRuntimeSettings.h b/Engine/Source/Runtime/Android/AndroidRuntimeSettings/Classes/AndroidRuntimeSettings.h
--- a/Engine/Source/Runtime/Android/AndroidRuntimeSettings/Classes/AndroidRuntimeSettings.h
+++ b/Engine/Source/Runtime/Android/AndroidRuntimeSettings/Classes/AndroidRuntimeSettings.h
@@ -250,6 +250,10 @@ public:
 	UPROPERTY(GlobalConfig, EditAnywhere, Category = GooglePlayServices)
 	FString GooglePlayLicenseKey;
 
+	/** Show the launch image as a startup slash screen */
+	UPROPERTY(GlobalConfig, EditAnywhere, Category = LaunchImages, meta = (DisplayName = "Show launch image"))
+	bool bShowLaunchImage;
+
 	/** Android Audio encoding options */
 	UPROPERTY(GlobalConfig, EditAnywhere, Category = DataCooker, meta = (DisplayName = "Audio encoding"))
 	TEnumAsByte<EAndroidAudio::Type> AndroidAudio;
diff --git a/Engine/Source/Runtime/Core/Private/Android/AndroidMisc.cpp b/Engine/Source/Runtime/Core/Private/Android/AndroidMisc.cpp
--- a/Engine/Source/Runtime/Core/Private/Android/AndroidMisc.cpp
+++ b/Engine/Source/Runtime/Core/Private/Android/AndroidMisc.cpp
@@ -112,6 +112,16 @@ void FAndroidMisc::PlatformInit()
 	// Display memory info
 }
 
+extern void AndroidThunkCpp_DismissSplashScreen();
+
+void FAndroidMisc::PlatformPostInit(bool ShowSplashScreen)
+{
+	if (!ShowSplashScreen)
+	{
+		AndroidThunkCpp_DismissSplashScreen();
+	}
+}
+
 void FAndroidMisc::GetEnvironmentVariable(const TCHAR* VariableName, TCHAR* Result, int32 ResultLength)
 {
 	*Result = 0;
diff --git a/Engine/Source/Runtime/Core/Public/Android/AndroidMisc.h b/Engine/Source/Runtime/Core/Public/Android/AndroidMisc.h
--- a/Engine/Source/Runtime/Core/Public/Android/AndroidMisc.h
+++ b/Engine/Source/Runtime/Core/Public/Android/AndroidMisc.h
@@ -23,6 +23,7 @@ struct CORE_API FAndroidMisc : public FGenericPlatformMisc
 	static void LocalPrint(const TCHAR *Message);
 	static void PlatformPreInit();
 	static void PlatformInit();
+	static void PlatformPostInit(bool ShowSplashScreen);
 	static void GetEnvironmentVariable(const TCHAR* VariableName, TCHAR* Result, int32 ResultLength);
 	static void* GetHardwareWindow();
 	static void SetHardwareWindow(void* InWindow);
diff --git a/Engine/Source/Runtime/Launch/Private/Android/AndroidJNI.cpp b/Engine/Source/Runtime/Launch/Private/Android/AndroidJNI.cpp
--- a/Engine/Source/Runtime/Launch/Private/Android/AndroidJNI.cpp
+++ b/Engine/Source/Runtime/Launch/Private/Android/AndroidJNI.cpp
@@ -53,7 +53,7 @@ void FJavaWrapper::FindClassesAndMethods(JNIEnv* Env)
 	AndroidThunkJava_KeepScreenOn = FindMethod(Env, GameActivityClassID, "AndroidThunkJava_KeepScreenOn", "(Z)V", bIsOptional);
 	AndroidThunkJava_InitHMDs = FindMethod(Env, GameActivityClassID, "AndroidThunkJava_InitHMDs", "()V", bIsOptional);
 	AndroidThunkJava_IsGearVRApplication = FindMethod(Env, GameActivityClassID, "AndroidThunkJava_IsGearVRApplication", "()Z", bIsOptional);
-
+	AndroidThunkJava_DismissSplashScreen = FindMethod(Env, GameActivityClassID, "AndroidThunkJava_DismissSplashScreen", "()V", bIsOptional);
 
 	// the rest are optional
 	bIsOptional = true;
@@ -159,6 +159,7 @@ jmethodID FJavaWrapper::AndroidThunkJava_IsMusicActive;
 jmethodID FJavaWrapper::AndroidThunkJava_KeepScreenOn;
 jmethodID FJavaWrapper::AndroidThunkJava_InitHMDs;
 jmethodID FJavaWrapper::AndroidThunkJava_IsGearVRApplication;
+jmethodID FJavaWrapper::AndroidThunkJava_DismissSplashScreen;
 
 jclass FJavaWrapper::GoogleServicesClassID;
 jobject FJavaWrapper::GoogleServicesThis;
@@ -232,6 +233,14 @@ bool AndroidThunkCpp_IsGearVRApplication()
 	return bIsGearVRApplication;
 }
 
+void AndroidThunkCpp_DismissSplashScreen()
+{
+	if (JNIEnv* Env = FAndroidApplication::GetJavaEnv())
+	{
+		FJavaWrapper::CallVoidMethod(Env, FJavaWrapper::GameActivityThis, FJavaWrapper::AndroidThunkJava_DismissSplashScreen);
+	}
+}
+
 void AndroidThunkCpp_ShowConsoleWindow()
 {
 	if (JNIEnv* Env = FAndroidApplication::GetJavaEnv())
diff --git a/Engine/Source/Runtime/Launch/Private/Android/AndroidJNI.h b/Engine/Source/Runtime/Launch/Private/Android/AndroidJNI.h
--- a/Engine/Source/Runtime/Launch/Private/Android/AndroidJNI.h
+++ b/Engine/Source/Runtime/Launch/Private/Android/AndroidJNI.h
@@ -29,6 +29,7 @@ public:
 	static jmethodID AndroidThunkJava_KeepScreenOn;
 	static jmethodID AndroidThunkJava_InitHMDs;
 	static jmethodID AndroidThunkJava_IsGearVRApplication;
+	static jmethodID AndroidThunkJava_DismissSplashScreen;
 
 	// IDs related to google play services
 	static jclass GoogleServicesClassID;
diff --git a/Engine/Source/Runtime/Launch/Private/Android/LaunchAndroid.cpp b/Engine/Source/Runtime/Launch/Private/Android/LaunchAndroid.cpp
--- a/Engine/Source/Runtime/Launch/Private/Android/LaunchAndroid.cpp
+++ b/Engine/Source/Runtime/Launch/Private/Android/LaunchAndroid.cpp
@@ -203,6 +203,8 @@ static void InitCommandLine()
 	}
 }
 
+extern void AndroidThunkCpp_DismissSplashScreen();
+
 //Main function called from the android entry point
 int32 AndroidMain(struct android_app* state)
 {
@@ -295,16 +297,6 @@ int32 AndroidMain(struct android_app* state)
 	// OBBs and APK are found.
 	IPlatformFile::GetPlatformPhysical().Initialize(nullptr, FCommandLine::Get());
 
-#if 0
-	for (int32 i = 0; i < 10; i++)
-	{
-		sleep(1);
-		FPlatformMisc::LowLevelOutputDebugStringf(TEXT("[Patch %d]"), i);
-
-	}
-	FPlatformMisc::LowLevelOutputDebugStringf(TEXT("[Patch] : Dont Patch \n"));
-#endif
-
 	// initialize the engine
 	GEngineLoop.PreInit(0, NULL, FCommandLine::Get());
 
@@ -319,6 +311,8 @@ int32 AndroidMain(struct android_app* state)
 
 	UE_LOG(LogAndroid, Log, TEXT("Passed GEngineLoop.Init()"));
 
+	AndroidThunkCpp_DismissSplashScreen();
+
 	// tick until done
 	while (!GIsRequestingExit)
 	{
