Compile rules for the engine into separate assemblies than rules for each project, so we aren't redundantly compiling them more than once. Also compile rules for each foreign plugin into its own assembly. Reduces time taken to generate project files by almost half (~40s to ~26s) in local testing.

diff --git a/Engine/Source/Programs/UnrealBuildTool/Configuration/UEBuildTarget.cs b/Engine/Source/Programs/UnrealBuildTool/Configuration/UEBuildTarget.cs
--- a/Engine/Source/Programs/UnrealBuildTool/Configuration/UEBuildTarget.cs
+++ b/Engine/Source/Programs/UnrealBuildTool/Configuration/UEBuildTarget.cs
@@ -544,7 +544,22 @@ namespace UnrealBuildTool
 		{
 			var CreateTargetStartTime = DateTime.UtcNow;
 
-			RulesAssembly RulesAssembly = RulesCompiler.CreateRulesAssembly(UnrealBuildTool.GetUProjectFile(), Desc.ForeignPlugins);
+			RulesAssembly RulesAssembly;
+			if(UnrealBuildTool.HasUProjectFile())
+			{
+				RulesAssembly = RulesCompiler.CreateProjectRulesAssembly(UnrealBuildTool.GetUProjectFile());
+			}
+			else
+			{
+				RulesAssembly = RulesCompiler.CreateEngineRulesAssembly();
+			}
+			if(Desc.ForeignPlugins != null)
+			{
+				foreach(string ForeignPlugin in Desc.ForeignPlugins)
+				{
+					RulesAssembly = RulesCompiler.CreatePluginRulesAssembly(ForeignPlugin, RulesAssembly);
+				}
+			}
 
 			string TargetFileName;
 			TargetRules RulesObject = RulesAssembly.CreateTargetRules(Desc.TargetName, new TargetInfo(Desc.Platform, Desc.Configuration), Desc.bIsEditorRecompile, out TargetFileName);
diff --git a/Engine/Source/Programs/UnrealBuildTool/System/ProjectFileGenerator.cs b/Engine/Source/Programs/UnrealBuildTool/System/ProjectFileGenerator.cs
--- a/Engine/Source/Programs/UnrealBuildTool/System/ProjectFileGenerator.cs
+++ b/Engine/Source/Programs/UnrealBuildTool/System/ProjectFileGenerator.cs
@@ -1729,7 +1729,16 @@ namespace UnrealBuildTool
 				if (WantProjectFileForTarget)
 				{
 					string CheckProjectFile = UProjectInfo.GetProjectForTarget(TargetName);
-					RulesAssembly RulesAssembly = RulesCompiler.CreateRulesAssembly(CheckProjectFile, null);
+
+					RulesAssembly RulesAssembly;
+					if(String.IsNullOrEmpty(CheckProjectFile))
+					{
+						RulesAssembly = RulesCompiler.CreateEngineRulesAssembly();
+					}
+					else
+					{
+						RulesAssembly = RulesCompiler.CreateProjectRulesAssembly(CheckProjectFile);
+					}
 
 					// Create target rules for all of the platforms and configuration combinations that we want to enable support for.
 					// Just use the current platform as we only need to recover the target type and both should be supported for all targets...
diff --git a/Engine/Source/Programs/UnrealBuildTool/System/RulesCompiler.cs b/Engine/Source/Programs/UnrealBuildTool/System/RulesCompiler.cs
--- a/Engine/Source/Programs/UnrealBuildTool/System/RulesCompiler.cs
+++ b/Engine/Source/Programs/UnrealBuildTool/System/RulesCompiler.cs
@@ -1072,13 +1072,24 @@ namespace UnrealBuildTool
 
 	public class RulesAssembly
 	{
-		// The compiled assembly
+		/// <summary>
+		/// The parent rules assembly that this assembly inherits. Game assemblies inherit the engine assembly, and the engine assembly inherits nothing.
+		/// </summary>
+		private RulesAssembly Parent;
+
+		/// <summary>
+		/// The compiled assembly
+		/// </summary>
 		private Assembly CompiledAssembly;
 
+		/// <summary>
 		/// Maps module names to their actual xxx.Module.cs file on disk
+		/// </summary>
 		private Dictionary<string, string> ModuleNameToFileName = new Dictionary<string, string>(StringComparer.InvariantCultureIgnoreCase);
 
+		/// <summary>
 		/// Maps target names to their actual xxx.Target.cs file on disk
+		/// </summary>
 		private Dictionary<string, string> TargetNameToFileName = new Dictionary<string, string>(StringComparer.InvariantCultureIgnoreCase);
 
 		/// <summary>
@@ -1087,8 +1098,10 @@ namespace UnrealBuildTool
 		/// <param name="ModuleFileNames">List of module files to compile</param>
 		/// <param name="TargetFileNames">List of target files to compile</param>
 		/// <param name="AssemblyFileName">The output path for the compiled assembly</param>
-		public RulesAssembly(List<string> ModuleFileNames, List<string> TargetFileNames, string AssemblyFileName)
+		public RulesAssembly(List<string> ModuleFileNames, List<string> TargetFileNames, string AssemblyFileName, RulesAssembly InParent)
 		{
+			Parent = InParent;
+
 			// Find all the source files
 			var AssemblySourceFiles = new List<string>();
 			AssemblySourceFiles.AddRange( ModuleFileNames );
@@ -1152,6 +1165,13 @@ namespace UnrealBuildTool
 					return true;
 				}
 			}
+			else
+			{
+				if(Parent != null && Parent.TryGetFileNameFromType(ExistingType, out FileName))
+				{
+					return true;
+				}
+			}
 
 			FileName = null;
 			return false;
@@ -1165,11 +1185,14 @@ namespace UnrealBuildTool
 		public string GetModuleFileName(string ModuleName)
 		{
 			string FileName;
-			if(!ModuleNameToFileName.TryGetValue(ModuleName, out FileName))
+			if(ModuleNameToFileName.TryGetValue(ModuleName, out FileName))
+			{
+				return FileName;
+			}
+			else
 			{
-				FileName = "";
+				return (Parent == null)? "" : Parent.GetModuleFileName(ModuleName);
 			}
-			return FileName;
 		}
 
 		/// <summary>
@@ -1180,11 +1203,14 @@ namespace UnrealBuildTool
 		public string GetTargetFileName(string TargetName)
 		{
 			string FileName;
-			if(!TargetNameToFileName.TryGetValue(TargetName, out FileName))
+			if(TargetNameToFileName.TryGetValue(TargetName, out FileName))
+			{
+				return FileName;
+			}
+			else
 			{
-				FileName = "";
+				return (Parent == null)? "" : Parent.GetTargetFileName(TargetName);
 			}
-			return FileName;
 		}
 
 		/// <summary>
@@ -1216,7 +1242,14 @@ namespace UnrealBuildTool
 			// Make sure the module file is known to us
 			if (!ModuleNameToFileName.TryGetValue(ModuleName, out ModuleFileName))
 			{
-				throw new MissingModuleException(ModuleName);
+				if(Parent == null)
+				{
+					throw new MissingModuleException(ModuleName);
+				}
+				else
+				{
+					return Parent.CreateModuleRules(ModuleName, Target, out ModuleFileName);
+				}
 			}
 
 			UnrealTargetPlatform LocalPlatform = Target.Platform;
@@ -1405,22 +1438,29 @@ namespace UnrealBuildTool
 
 			if (bFoundTargetName == false)
 			{
-				//				throw new BuildException("Couldn't find target rules file for target '{0}' in rules assembly '{1}'.", TargetName, RulesAssembly.FullName);
-				string ExceptionMessage = "Couldn't find target rules file for target '";
-				ExceptionMessage += TargetName;
-				ExceptionMessage += "' in rules assembly '";
-				ExceptionMessage += CompiledAssembly.FullName;
-				ExceptionMessage += "'." + Environment.NewLine;
+				if(Parent == null)
+				{
+					//				throw new BuildException("Couldn't find target rules file for target '{0}' in rules assembly '{1}'.", TargetName, RulesAssembly.FullName);
+					string ExceptionMessage = "Couldn't find target rules file for target '";
+					ExceptionMessage += TargetName;
+					ExceptionMessage += "' in rules assembly '";
+					ExceptionMessage += CompiledAssembly.FullName;
+					ExceptionMessage += "'." + Environment.NewLine;
 
-				ExceptionMessage += "Location: " + CompiledAssembly.Location + Environment.NewLine;
+					ExceptionMessage += "Location: " + CompiledAssembly.Location + Environment.NewLine;
+
+					ExceptionMessage += "Target rules found:" + Environment.NewLine;
+					foreach (KeyValuePair<string, string> entry in TargetNameToFileName)
+					{
+						ExceptionMessage += "\t" + entry.Key + " - " + entry.Value + Environment.NewLine;
+					}
 
-				ExceptionMessage += "Target rules found:" + Environment.NewLine;
-				foreach (KeyValuePair<string, string> entry in TargetNameToFileName)
+					throw new BuildException(ExceptionMessage);
+				}
+				else
 				{
-					ExceptionMessage += "\t" + entry.Key + " - " + entry.Value + Environment.NewLine;
+					return Parent.CreateTargetRules(TargetName, Target, bInEditorRecompile, out TargetFileName);
 				}
-
-				throw new BuildException(ExceptionMessage);
 			}
 
 			// Return the target file name to the caller
@@ -1603,19 +1643,28 @@ namespace UnrealBuildTool
 		/// We cache these file names so we can avoid searching for them later on.
 		static Dictionary<string, RulesFileCache> RootFolderToRulesFileCache = new Dictionary<string, RulesFileCache>();
 
-		public static List<string> FindAllRulesSourceFiles( RulesFileType RulesFileType, List<string> GameFolders, List<string> ForeignPlugins, List<string> AdditionalSearchPaths )
+		public static List<string> FindAllRulesSourceFiles( RulesFileType RulesFileType, List<string> GameFolders, List<string> ForeignPlugins, List<string> AdditionalSearchPaths, bool bIncludeEngine = true )
 		{
 			List<string> Folders = new List<string>();
 
 			// Add all engine source (including third party source)
-			Folders.Add( Path.Combine( ProjectFileGenerator.EngineRelativePath, "Source" ) );
+			if(bIncludeEngine)
+			{
+				Folders.Add( Path.Combine( ProjectFileGenerator.EngineRelativePath, "Source" ) );
+			}
 
 			// @todo plugin: Disallow modules from including plugin modules as dependency modules? (except when the module is part of that plugin)
 
 			// Get all the root folders for plugins
 			List<string> RootFolders = new List<string>();
-			RootFolders.Add(ProjectFileGenerator.EngineRelativePath);
-			RootFolders.AddRange(GameFolders);
+			if(bIncludeEngine)
+			{
+				RootFolders.Add(ProjectFileGenerator.EngineRelativePath);
+			}
+			if(GameFolders != null)
+			{
+				RootFolders.AddRange(GameFolders);
+			}
 
 			// Find all the plugin source directories
 			foreach(string RootFolder in RootFolders)
@@ -1711,59 +1760,101 @@ namespace UnrealBuildTool
 			return SourceFiles;
 		}
 
+		/// <summary>
+		/// The cached rules assembly for engine modules and targets.
+		/// </summary>
+		private static RulesAssembly EngineRulesAssembly;
+
 		/// Map of assembly names we've already compiled and loaded to their Assembly and list of game folders.  This is used to prevent
 		/// trying to recompile the same assembly when ping-ponging between different types of targets
 		private static Dictionary<string, RulesAssembly> LoadedAssemblyMap = new Dictionary<string, RulesAssembly>(StringComparer.InvariantCultureIgnoreCase);
 
+		/// <summary>
+		/// Creates the engine rules assembly
+		/// </summary>
+		/// <param name="ForeignPlugins">List of plugins to include in this assembly</param>
+		/// <returns>New rules assembly</returns>
+		public static RulesAssembly CreateEngineRulesAssembly()
+		{
+			if(EngineRulesAssembly == null)
+			{
+				// Find all the rules files
+				List<string> ModuleFileNames = FindAllRulesSourceFiles(RulesFileType.Module, null, null, null, true);
+				List<string> TargetFileNames = FindAllRulesSourceFiles(RulesFileType.Target, null, null, null, true);
+
+				// Create a path to the assembly that we'll either load or compile
+				string AssemblyFileName = Path.Combine(BuildConfiguration.BaseIntermediatePath, "BuildRules", "UE4Rules.dll");
+				EngineRulesAssembly = new RulesAssembly(ModuleFileNames, TargetFileNames, AssemblyFileName, null);
+			}
+			return EngineRulesAssembly;
+		}
+
 		/// <summary>
 		/// Creates a rules assembly with the given parameters.
 		/// </summary>
 		/// <param name="ProjectFileName">The project file to create rules for. Null for the engine.</param>
 		/// <param name="ForeignPlugins">List of foreign plugin folders to include in the assembly. May be null.</param>
-		public static RulesAssembly CreateRulesAssembly(string ProjectFileName, List<string> ForeignPlugins)
+		public static RulesAssembly CreateProjectRulesAssembly(string ProjectFileName)
 		{
 			// Normalize the project filename, so we can look it up in the cache
-			string FullProjectFileName = String.IsNullOrEmpty(ProjectFileName)? "" : Path.GetFullPath(ProjectFileName);
+			string FullProjectFileName = Path.GetFullPath(ProjectFileName);
 
 			// Check if there's an existing assembly for this project
-			RulesAssembly ExistingAssembly;
-			if(LoadedAssemblyMap.TryGetValue(FullProjectFileName, out ExistingAssembly))
+			RulesAssembly ProjectRulesAssembly;
+			if(!LoadedAssemblyMap.TryGetValue(FullProjectFileName, out ProjectRulesAssembly))
 			{
-				return ExistingAssembly;
-			}
-
-			// Get the assembly name
-			string AssemblyName = ((FullProjectFileName == null)? "UE4" : Path.GetFileNameWithoutExtension(FullProjectFileName)) + "ModuleRules";
+				// Create the engine rules assembly
+				RulesAssembly Parent = CreateEngineRulesAssembly();
 
-			// Get the folders to search through
-			List<string> GameFolders = new List<string>();
-			List<string> AdditionalSearchPaths = new List<string>();
-			if (FullProjectFileName.Length > 0)
-			{
+				// Get the folders to search through
+				List<string> GameFolders = new List<string>();
 				string ProjectDirectoryName = Path.GetDirectoryName(FullProjectFileName);
 				GameFolders.Add(ProjectDirectoryName);
 
 				// Add the games project's intermediate source folder
+				List<string> AdditionalSearchPaths = new List<string>();
 				var ProjectIntermediateSourceDirectory = Path.Combine(ProjectDirectoryName, "Intermediate", "Source");
 				if (Directory.Exists(ProjectIntermediateSourceDirectory))
 				{
 					AdditionalSearchPaths.Add(ProjectIntermediateSourceDirectory);
 				}
-			}
-
-			var ModuleFileNames = FindAllRulesSourceFiles(RulesFileType.Module, GameFolders, ForeignPlugins, AdditionalSearchPaths);
 
-			// Create a path to the assembly that we'll either load or compile
-			string BaseIntermediatePath = UnrealBuildTool.HasUProjectFile() ?
-				Path.Combine(UnrealBuildTool.GetUProjectPath(), BuildConfiguration.BaseIntermediateFolder) : BuildConfiguration.BaseIntermediatePath;
+				// Find all the rules source files
+				List<string> ModuleFileNames = FindAllRulesSourceFiles(RulesFileType.Module, GameFolders, null, AdditionalSearchPaths, false);
+				List<string> TargetFileNames = FindAllRulesSourceFiles(RulesFileType.Target, GameFolders, null, AdditionalSearchPaths, false);
 
-			string OutputAssemblyPath = Path.GetFullPath(Path.Combine(BaseIntermediatePath, "BuildRules", AssemblyName + ".dll"));
+				// Compile the assembly
+				string AssemblyFileName = Path.Combine(ProjectDirectoryName, BuildConfiguration.BaseIntermediateFolder, "BuildRules", Path.GetFileNameWithoutExtension(ProjectFileName) + "ModuleRules.dll");
+				ProjectRulesAssembly = new RulesAssembly(ModuleFileNames, TargetFileNames, AssemblyFileName, Parent);
+				LoadedAssemblyMap.Add(FullProjectFileName, ProjectRulesAssembly);
+			}
+			return ProjectRulesAssembly;
+		}
 
-			var TargetFileNames = FindAllRulesSourceFiles(RulesFileType.Target, GameFolders, ForeignPlugins, AdditionalSearchPaths);
+		/// <summary>
+		/// Creates a rules assembly with the given parameters.
+		/// </summary>
+		/// <param name="ProjectFileName">The project file to create rules for. Null for the engine.</param>
+		/// <param name="ForeignPlugins">List of foreign plugin folders to include in the assembly. May be null.</param>
+		public static RulesAssembly CreatePluginRulesAssembly(string PluginFileName, RulesAssembly Parent)
+		{
+			// Normalize the project filename, so we can look it up in the cache
+			string FullPluginFileName = Path.GetFullPath(PluginFileName);
 
-			RulesAssembly Rules = new RulesAssembly(ModuleFileNames, TargetFileNames, OutputAssemblyPath);
-			LoadedAssemblyMap.Add(FullProjectFileName, Rules);
-			return Rules;
+			// Check if there's an existing assembly for this project
+			RulesAssembly PluginRulesAssembly;
+			if(!LoadedAssemblyMap.TryGetValue(FullPluginFileName, out PluginRulesAssembly))
+			{
+				// Find all the rules source files
+				List<string> ModuleFileNames = FindAllRulesSourceFiles(RulesFileType.Module, null, new List<string>{ PluginFileName }, null, false);
+				List<string> TargetFileNames = FindAllRulesSourceFiles(RulesFileType.Target, null, new List<string>{ PluginFileName }, null, false);
+
+				// Compile the assembly
+				string AssemblyFileName = Path.Combine(Path.GetDirectoryName(PluginFileName), BuildConfiguration.BaseIntermediateFolder, "BuildRules", Path.GetFileNameWithoutExtension(PluginFileName) + "ModuleRules.dll");
+				PluginRulesAssembly = new RulesAssembly(ModuleFileNames, TargetFileNames, AssemblyFileName, Parent);
+				LoadedAssemblyMap.Add(FullPluginFileName, PluginRulesAssembly);
+			}
+			return PluginRulesAssembly;
 		}
 
 		/// <summary>
diff --git a/Engine/Source/Programs/UnrealBuildTool/System/ThirdPartyHeaderFinder.cs b/Engine/Source/Programs/UnrealBuildTool/System/ThirdPartyHeaderFinder.cs
--- a/Engine/Source/Programs/UnrealBuildTool/System/ThirdPartyHeaderFinder.cs
+++ b/Engine/Source/Programs/UnrealBuildTool/System/ThirdPartyHeaderFinder.cs
@@ -91,7 +91,7 @@ namespace UnrealBuildTool
 			ThirdPartyHeaders = new List<Header>();
 
 			// Create a rules assembly for the engine
-			RulesAssembly EngineRulesAssembly = RulesCompiler.CreateRulesAssembly(null, null);
+			RulesAssembly EngineRulesAssembly = RulesCompiler.CreateEngineRulesAssembly();
 
 			// Find all modules referenced by the current target
 			List<string> ModuleFileNames = RulesCompiler.FindAllRulesSourceFiles( RulesCompiler.RulesFileType.Module, GameFolders:null, ForeignPlugins:null, AdditionalSearchPaths:null );
diff --git a/Engine/Source/Programs/UnrealBuildTool/System/XcodeProjectFileGenerator.cs b/Engine/Source/Programs/UnrealBuildTool/System/XcodeProjectFileGenerator.cs
--- a/Engine/Source/Programs/UnrealBuildTool/System/XcodeProjectFileGenerator.cs
+++ b/Engine/Source/Programs/UnrealBuildTool/System/XcodeProjectFileGenerator.cs
@@ -1087,7 +1087,7 @@ namespace UnrealBuildTool
 				ProjectTarget TargetProject = new global::UnrealBuildTool.ProjectTarget();
 				var Target = new TargetInfo(UnrealTargetPlatform.Mac, UnrealTargetConfiguration.Development);
 				TargetProject.TargetFilePath = "";
-				RulesAssembly EngineRulesAssembly = RulesCompiler.CreateRulesAssembly(null, null);
+				RulesAssembly EngineRulesAssembly = RulesCompiler.CreateEngineRulesAssembly();
 				TargetProject.TargetRules = EngineRulesAssembly.CreateTargetRules("UE4Game", Target, false, out TargetProject.TargetFilePath);
 				TargetProjects.Add(TargetProject);
 			}
