Move the AddDefaultIncludePathsToModuleRules() function into UEBuildTarget, since it's not related to the RulesCompiler.

diff --git a/Engine/Source/Programs/UnrealBuildTool/Configuration/UEBuildTarget.cs b/Engine/Source/Programs/UnrealBuildTool/Configuration/UEBuildTarget.cs
--- a/Engine/Source/Programs/UnrealBuildTool/Configuration/UEBuildTarget.cs
+++ b/Engine/Source/Programs/UnrealBuildTool/Configuration/UEBuildTarget.cs
@@ -3416,7 +3416,7 @@ namespace UnrealBuildTool
 				if (RulesObject.Type != ModuleRules.ModuleType.External && ModuleName != "Core")
 				{
 					// Add the default include paths to the module rules, if they exist.
-					RulesCompiler.AddDefaultIncludePathsToModuleRules(ModuleName, ModuleFileName, ModuleFileRelativeToEngineDirectory, IsGameModule, RulesObject);
+					AddDefaultIncludePathsToModuleRules(ModuleName, ModuleFileName, ModuleFileRelativeToEngineDirectory, IsGameModule, RulesObject);
 
 					// Add the path to the generated headers 
 					if(GeneratedCodeDirectory != null)
@@ -3538,6 +3538,149 @@ namespace UnrealBuildTool
 			}
 		}
 
+		/// <summary>
+		/// Add the standard default include paths to the given modulerules object
+		/// </summary>
+		/// <param name="InModuleName">The name of the module</param>
+		/// <param name="InModuleFilename">The filename to the module rules file (Build.cs)</param>
+		/// <param name="InModuleFileRelativeToEngineDirectory">The module file relative to the engine directory</param>
+		/// <param name="IsGameModule">true if it is a game module, false if not</param>
+		/// <param name="RulesObject">The module rules object itself</param>
+		public static void AddDefaultIncludePathsToModuleRules(string InModuleName, string InModuleFilename, string InModuleFileRelativeToEngineDirectory, bool IsGameModule, ModuleRules RulesObject)
+		{
+			// Grab the absolute path of the Engine/Source folder for use later
+			string AbsEngineSourceDirectory = Path.Combine(ProjectFileGenerator.RootRelativePath, "Engine/Source");
+			AbsEngineSourceDirectory = Path.GetFullPath(AbsEngineSourceDirectory);
+			AbsEngineSourceDirectory = AbsEngineSourceDirectory.Replace("\\", "/");
+
+			// Find the module path relative to the Engine/Source folder
+			string ModuleDirectoryRelativeToEngineSourceDirectory = Utils.MakePathRelativeTo(InModuleFilename, Path.Combine(ProjectFileGenerator.RootRelativePath, "Engine/Source"));
+			ModuleDirectoryRelativeToEngineSourceDirectory = ModuleDirectoryRelativeToEngineSourceDirectory.Replace("\\", "/");
+			// Remove the build.cs file from the directory if present
+			if (ModuleDirectoryRelativeToEngineSourceDirectory.EndsWith("Build.cs", StringComparison.InvariantCultureIgnoreCase))
+			{
+				Int32 LastSlashIdx = ModuleDirectoryRelativeToEngineSourceDirectory.LastIndexOf("/");
+				if (LastSlashIdx != -1)
+				{
+					ModuleDirectoryRelativeToEngineSourceDirectory = ModuleDirectoryRelativeToEngineSourceDirectory.Substring(0, LastSlashIdx + 1);
+				}
+			}
+			else
+			{
+				throw new BuildException("Invalid module filename '{0}'.", InModuleFilename);
+			}
+
+			// Determine the 'Game/Source' folder
+			[REDACTED_EMAIL]: This currently requires following our standard format for folder layout.
+			//				Also, it assumes the module itself is not named Source...
+			string GameSourceIncludePath = (IsGameModule == true) ? ModuleDirectoryRelativeToEngineSourceDirectory : "";
+			if (string.IsNullOrEmpty(GameSourceIncludePath) == false)
+			{
+				Int32 SourceIdx = GameSourceIncludePath.IndexOf("/Source/");
+				if (SourceIdx != -1)
+				{
+					GameSourceIncludePath = GameSourceIncludePath.Substring(0, SourceIdx + 8);
+					RulesObject.PublicIncludePaths.Add(GameSourceIncludePath);
+				}
+			}
+
+			// Setup the directories for Classes, Public, and Intermediate			
+			string ClassesDirectory = Path.Combine(ModuleDirectoryRelativeToEngineSourceDirectory, "Classes/");	// @todo uht: Deprecate eventually.  Or force it to be manually specified...
+			string PublicDirectory = Path.Combine(ModuleDirectoryRelativeToEngineSourceDirectory, "Public/");
+
+			if(!Utils.IsFileUnderDirectory(InModuleFilename, AbsEngineSourceDirectory))
+			{
+				// This will be either the format 
+				//		../<Game>/Source/<Module>
+				// or
+				//		c:/PATH/<Game>/Source/<Module>
+                string SourceDirName = "Source";
+                Int32 SourceSlashIdx = ModuleDirectoryRelativeToEngineSourceDirectory.IndexOf(SourceDirName);
+				string SourceDirectoryPath = null;
+				if(SourceSlashIdx != -1)
+				{
+					try
+					{
+                        SourceDirectoryPath = Path.GetFullPath(ModuleDirectoryRelativeToEngineSourceDirectory.Substring(0, SourceSlashIdx + SourceDirName.Length));
+					}
+					catch(Exception Exc)
+					{
+						throw new BuildException(Exc, "Failed to resolve module source directory for private include paths for module {0}.", InModuleName);
+					}
+				}
+				else
+				{
+					//@todo. throw a build exception here?
+				}
+
+				// Resolve private include paths against the module source root so they are simpler and don't have to be engine root relative.
+				if(SourceDirectoryPath != null)
+				{
+					List<string> ResolvedPrivatePaths = new List<string>();
+					foreach(var PrivatePath in RulesObject.PrivateIncludePaths)
+					{
+						try
+						{
+							if(!Path.IsPathRooted(PrivatePath))
+							{
+								ResolvedPrivatePaths.Add(Path.Combine(SourceDirectoryPath, PrivatePath));
+							}
+							else
+							{
+								ResolvedPrivatePaths.Add(PrivatePath);
+							}
+						}
+						catch(Exception Exc)
+						{
+							throw new BuildException(Exc, "Failed to resolve private include path {0}.", PrivatePath);
+						}
+					}
+					RulesObject.PrivateIncludePaths = ResolvedPrivatePaths;
+				}
+			}
+
+			string IncludePath_Classes = "";
+			string IncludePath_Public = "";
+
+			bool bModulePathIsRooted = Path.IsPathRooted(ModuleDirectoryRelativeToEngineSourceDirectory);
+
+			string ClassesFolderName = (bModulePathIsRooted == false) ? Path.Combine(AbsEngineSourceDirectory, ClassesDirectory) : ClassesDirectory;
+			if (Directory.Exists(ClassesFolderName) == true)
+			{
+				IncludePath_Classes = ClassesDirectory;
+			}
+
+			string PublicFolderName = (bModulePathIsRooted == false) ? Path.Combine(AbsEngineSourceDirectory, PublicDirectory) : PublicDirectory;
+			if (Directory.Exists(PublicFolderName) == true)
+			{
+				IncludePath_Public = PublicDirectory;
+			}
+
+			// Add them if they are required...
+			if (IncludePath_Classes.Length > 0)
+			{
+				RulesObject.PublicIncludePaths.Add(IncludePath_Classes);
+			}
+			if (IncludePath_Public.Length > 0)
+			{
+				RulesObject.PublicIncludePaths.Add(IncludePath_Public);
+
+				// Add subdirectories of Public if present
+				DirectoryInfo PublicInfo = new DirectoryInfo(IncludePath_Public);
+				DirectoryInfo[] PublicSubDirs = PublicInfo.GetDirectories("*", SearchOption.AllDirectories);
+				if (PublicSubDirs.Length > 0)
+				{
+					foreach (DirectoryInfo SubDir in PublicSubDirs)
+					{
+						string PartialDir = SubDir.FullName.Replace(PublicInfo.FullName, "");
+						string NewDir = IncludePath_Public + PartialDir;
+						NewDir = Utils.CleanDirectorySeparators(NewDir, '/');
+						RulesObject.PublicIncludePaths.Add(NewDir);
+					}
+				}
+			}
+		}
+
 		/** Finds a module given its name.  Throws an exception if the module couldn't be found. */
 		public UEBuildModule GetModuleByName(string Name)
 		{
diff --git a/Engine/Source/Programs/UnrealBuildTool/System/RulesCompiler.cs b/Engine/Source/Programs/UnrealBuildTool/System/RulesCompiler.cs
--- a/Engine/Source/Programs/UnrealBuildTool/System/RulesCompiler.cs
+++ b/Engine/Source/Programs/UnrealBuildTool/System/RulesCompiler.cs
@@ -1874,148 +1874,5 @@ namespace UnrealBuildTool
 			}
 			return null;
 		}
-
-		/// <summary>
-		/// Add the standard default include paths to the given modulerules object
-		/// </summary>
-		/// <param name="InModuleName">The name of the module</param>
-		/// <param name="InModuleFilename">The filename to the module rules file (Build.cs)</param>
-		/// <param name="InModuleFileRelativeToEngineDirectory">The module file relative to the engine directory</param>
-		/// <param name="IsGameModule">true if it is a game module, false if not</param>
-		/// <param name="RulesObject">The module rules object itself</param>
-		public static void AddDefaultIncludePathsToModuleRules(string InModuleName, string InModuleFilename, string InModuleFileRelativeToEngineDirectory, bool IsGameModule, ModuleRules RulesObject)
-		{
-			// Grab the absolute path of the Engine/Source folder for use later
-			string AbsEngineSourceDirectory = Path.Combine(ProjectFileGenerator.RootRelativePath, "Engine/Source");
-			AbsEngineSourceDirectory = Path.GetFullPath(AbsEngineSourceDirectory);
-			AbsEngineSourceDirectory = AbsEngineSourceDirectory.Replace("\\", "/");
-
-			// Find the module path relative to the Engine/Source folder
-			string ModuleDirectoryRelativeToEngineSourceDirectory = Utils.MakePathRelativeTo(InModuleFilename, Path.Combine(ProjectFileGenerator.RootRelativePath, "Engine/Source"));
-			ModuleDirectoryRelativeToEngineSourceDirectory = ModuleDirectoryRelativeToEngineSourceDirectory.Replace("\\", "/");
-			// Remove the build.cs file from the directory if present
-			if (ModuleDirectoryRelativeToEngineSourceDirectory.EndsWith("Build.cs", StringComparison.InvariantCultureIgnoreCase))
-			{
-				Int32 LastSlashIdx = ModuleDirectoryRelativeToEngineSourceDirectory.LastIndexOf("/");
-				if (LastSlashIdx != -1)
-				{
-					ModuleDirectoryRelativeToEngineSourceDirectory = ModuleDirectoryRelativeToEngineSourceDirectory.Substring(0, LastSlashIdx + 1);
-				}
-			}
-			else
-			{
-				throw new BuildException("Invalid module filename '{0}'.", InModuleFilename);
-			}
-
-			// Determine the 'Game/Source' folder
-			[REDACTED_EMAIL]: This currently requires following our standard format for folder layout.
-			//				Also, it assumes the module itself is not named Source...
-			string GameSourceIncludePath = (IsGameModule == true) ? ModuleDirectoryRelativeToEngineSourceDirectory : "";
-			if (string.IsNullOrEmpty(GameSourceIncludePath) == false)
-			{
-				Int32 SourceIdx = GameSourceIncludePath.IndexOf("/Source/");
-				if (SourceIdx != -1)
-				{
-					GameSourceIncludePath = GameSourceIncludePath.Substring(0, SourceIdx + 8);
-					RulesObject.PublicIncludePaths.Add(GameSourceIncludePath);
-				}
-			}
-
-			// Setup the directories for Classes, Public, and Intermediate			
-			string ClassesDirectory = Path.Combine(ModuleDirectoryRelativeToEngineSourceDirectory, "Classes/");	// @todo uht: Deprecate eventually.  Or force it to be manually specified...
-			string PublicDirectory = Path.Combine(ModuleDirectoryRelativeToEngineSourceDirectory, "Public/");
-
-			if(!Utils.IsFileUnderDirectory(InModuleFilename, AbsEngineSourceDirectory))
-			{
-				// This will be either the format 
-				//		../<Game>/Source/<Module>
-				// or
-				//		c:/PATH/<Game>/Source/<Module>
-                string SourceDirName = "Source";
-                Int32 SourceSlashIdx = ModuleDirectoryRelativeToEngineSourceDirectory.IndexOf(SourceDirName);
-				string SourceDirectoryPath = null;
-				if(SourceSlashIdx != -1)
-				{
-					try
-					{
-                        SourceDirectoryPath = Path.GetFullPath(ModuleDirectoryRelativeToEngineSourceDirectory.Substring(0, SourceSlashIdx + SourceDirName.Length));
-					}
-					catch(Exception Exc)
-					{
-						throw new BuildException(Exc, "Failed to resolve module source directory for private include paths for module {0}.", InModuleName);
-					}
-				}
-				else
-				{
-					//@todo. throw a build exception here?
-				}
-
-				// Resolve private include paths against the module source root so they are simpler and don't have to be engine root relative.
-				if(SourceDirectoryPath != null)
-				{
-					List<string> ResolvedPrivatePaths = new List<string>();
-					foreach(var PrivatePath in RulesObject.PrivateIncludePaths)
-					{
-						try
-						{
-							if(!Path.IsPathRooted(PrivatePath))
-							{
-								ResolvedPrivatePaths.Add(Path.Combine(SourceDirectoryPath, PrivatePath));
-							}
-							else
-							{
-								ResolvedPrivatePaths.Add(PrivatePath);
-							}
-						}
-						catch(Exception Exc)
-						{
-							throw new BuildException(Exc, "Failed to resolve private include path {0}.", PrivatePath);
-						}
-					}
-					RulesObject.PrivateIncludePaths = ResolvedPrivatePaths;
-				}
-			}
-
-			string IncludePath_Classes = "";
-			string IncludePath_Public = "";
-
-			bool bModulePathIsRooted = Path.IsPathRooted(ModuleDirectoryRelativeToEngineSourceDirectory);
-
-			string ClassesFolderName = (bModulePathIsRooted == false) ? Path.Combine(AbsEngineSourceDirectory, ClassesDirectory) : ClassesDirectory;
-			if (Directory.Exists(ClassesFolderName) == true)
-			{
-				IncludePath_Classes = ClassesDirectory;
-			}
-
-			string PublicFolderName = (bModulePathIsRooted == false) ? Path.Combine(AbsEngineSourceDirectory, PublicDirectory) : PublicDirectory;
-			if (Directory.Exists(PublicFolderName) == true)
-			{
-				IncludePath_Public = PublicDirectory;
-			}
-
-			// Add them if they are required...
-			if (IncludePath_Classes.Length > 0)
-			{
-				RulesObject.PublicIncludePaths.Add(IncludePath_Classes);
-			}
-			if (IncludePath_Public.Length > 0)
-			{
-				RulesObject.PublicIncludePaths.Add(IncludePath_Public);
-
-				// Add subdirectories of Public if present
-				DirectoryInfo PublicInfo = new DirectoryInfo(IncludePath_Public);
-				DirectoryInfo[] PublicSubDirs = PublicInfo.GetDirectories("*", SearchOption.AllDirectories);
-				if (PublicSubDirs.Length > 0)
-				{
-					foreach (DirectoryInfo SubDir in PublicSubDirs)
-					{
-						string PartialDir = SubDir.FullName.Replace(PublicInfo.FullName, "");
-						string NewDir = IncludePath_Public + PartialDir;
-						NewDir = Utils.CleanDirectorySeparators(NewDir, '/');
-						RulesObject.PublicIncludePaths.Add(NewDir);
-					}
-				}
-			}
-		}
 	}
 }
