Fix failure to get module filename for engine modules.

diff --git a/Engine/Source/Programs/UnrealBuildTool/System/RulesCompiler.cs b/Engine/Source/Programs/UnrealBuildTool/System/RulesCompiler.cs
--- a/Engine/Source/Programs/UnrealBuildTool/System/RulesCompiler.cs
+++ b/Engine/Source/Programs/UnrealBuildTool/System/RulesCompiler.cs
@@ -1864,9 +1864,13 @@ namespace UnrealBuildTool
 		/// <returns>The filename that declared the given type, or null</returns>
 		public static string GetFileNameFromType(Type ExistingType)
 		{
+			string FileName;
+			if(EngineRulesAssembly != null && EngineRulesAssembly.TryGetFileNameFromType(ExistingType, out FileName))
+			{
+				return FileName;
+			}
 			foreach(RulesAssembly RulesAssembly in LoadedAssemblyMap.Values)
 			{
-				string FileName;
 				if(RulesAssembly.TryGetFileNameFromType(ExistingType, out FileName))
 				{
 					return FileName;
