Allow rules assemblies to be empty, and not require any module source files to be loaded. The target may just use source files from the parent assembly.

diff --git a/Engine/Source/Programs/UnrealBuildTool/System/RulesCompiler.cs b/Engine/Source/Programs/UnrealBuildTool/System/RulesCompiler.cs
--- a/Engine/Source/Programs/UnrealBuildTool/System/RulesCompiler.cs
+++ b/Engine/Source/Programs/UnrealBuildTool/System/RulesCompiler.cs
@@ -1103,16 +1103,15 @@ namespace UnrealBuildTool
 			Parent = InParent;
 
 			// Find all the source files
-			var AssemblySourceFiles = new List<string>();
+			List<string> AssemblySourceFiles = new List<string>();
 			AssemblySourceFiles.AddRange( ModuleFileNames );
-			if( AssemblySourceFiles.Count == 0 )
-			{
-				throw new BuildException("No module rules source files were found in any of the module base directories!");
-			}
 			AssemblySourceFiles.AddRange( TargetFileNames );
 
 			// Compile the assembly
-			CompiledAssembly = DynamicCompilation.CompileAndLoadAssembly( AssemblyFileName, AssemblySourceFiles );
+			if(AssemblySourceFiles.Count > 0)
+			{
+				CompiledAssembly = DynamicCompilation.CompileAndLoadAssembly( AssemblyFileName, AssemblySourceFiles );
+			}
 
 			// Setup the module map
 			foreach( var CurModuleFileName in ModuleFileNames )
@@ -1137,14 +1136,6 @@ namespace UnrealBuildTool
 			}
 		}
 
-		/// <summary>
-		/// The name of the compiled assembly
-		/// </summary>
-		public string Name
-		{
-			get { return Path.GetFileNameWithoutExtension( CompiledAssembly.Location ); }
-		}
-
 		/// <summary>
 		/// Tries to get the filename that declared the given type
 		/// </summary>
@@ -1234,8 +1225,6 @@ namespace UnrealBuildTool
 		/// <returns>Compiled module rule info</returns>
 		public ModuleRules CreateModuleRules(string ModuleName, TargetInfo Target, out string ModuleFileName)
 		{
-			var AssemblyFileName = Path.GetFileNameWithoutExtension(CompiledAssembly.Location);
-
 			// Currently, we expect the user's rules object type name to be the same as the module name
 			var ModuleTypeName = ModuleName;
 
@@ -1266,6 +1255,7 @@ namespace UnrealBuildTool
 				RulesObjectType = CompiledAssembly.GetType("UnrealBuildTool.Rules." + ModuleName);
 			}
 
+			var AssemblyFileName = Path.GetFileNameWithoutExtension(CompiledAssembly.Location);
 			if (RulesObjectType == null)
 			{
 				throw new BuildException("Expecting to find a type to be declared in a module rules named '{0}' in {1}.  This type must derive from the 'ModuleRules' type defined by Unreal Build Tool.", ModuleTypeName, CompiledAssembly.FullName);
