Add a class to represent file and directory names stored in a canonical form. Allows strongly typing full paths and performing faster operations on them, rather than needing to repeatedly normalize them.

diff --git a/Engine/Source/Programs/UnrealBuildTool/System/FileSystemNames.cs b/Engine/Source/Programs/UnrealBuildTool/System/FileSystemNames.cs
new file mode 100644
--- /dev/null
+++ b/Engine/Source/Programs/UnrealBuildTool/System/FileSystemNames.cs
@@ -0,0 +1,313 @@
+﻿using System;
+using System.Collections.Generic;
+using System.IO;
+using System.Linq;
+using System.Text;
+using System.Threading.Tasks;
+
+namespace UnrealBuildTool
+{
+	/// <summary>
+	/// Base class for file system objects (files or directories).
+	/// </summary>
+	public abstract class FullFileSystemName
+	{
+		/// <summary>
+		/// The path to this object. Stored as an absolute path, with O/S preferred separator characters, and no trailing slash for directories.
+		/// </summary>
+		public readonly string FullName;
+
+		/// <summary>
+		/// The canonical full name for this object.
+		/// </summary>
+		public readonly string CanonicalName;
+
+		/// <summary>
+		/// Constructs a filesystem object for the given path.
+		/// </summary>
+		public FullFileSystemName(string InPath)
+		{
+			FullName = Path.GetFullPath(InPath);
+			CanonicalName = FullName.ToLowerInvariant();
+		}
+
+		/// <summary>
+		/// Direct constructor for a path
+		/// </summary>
+		protected FullFileSystemName(string InFullName, string InCanonicalName)
+		{
+			FullName = InFullName;
+			CanonicalName = InCanonicalName;
+		}
+
+		/// <summary>
+		/// Determines if the given object is at or under the given directory
+		/// </summary>
+		/// <param name="Directory"></param>
+		/// <returns></returns>
+		public bool IsUnderDirectory(FullDirectoryName Other)
+		{
+			return CanonicalName.StartsWith(Other.CanonicalName) && (CanonicalName.Length == Other.CanonicalName.Length || CanonicalName[Other.CanonicalName.Length] == Path.DirectorySeparatorChar);
+		}
+
+		/// <summary>
+		/// Creates a relative path from the given base directory
+		/// </summary>
+		/// <param name="Directory">The directory to create a relative path from</param>
+		/// <returns>A relative path from the given directory</returns>
+		public string MakeRelativeTo(FullDirectoryName Directory)
+		{
+			// Find how much of the path is common between the two paths. This length does not include a trailing directory separator character.
+			int CommonDirectoryLength = -1;
+			for (int Idx = 0;;Idx++)
+			{
+				if(Idx == CanonicalName.Length)
+				{
+					// The two paths are identical. Just return the "." character.
+					if(Idx == Directory.CanonicalName.Length)
+					{
+						return ".";
+					}
+
+					// Check if we're finishing on a complete directory name
+					if(Directory.CanonicalName[Idx] == Path.DirectorySeparatorChar)
+					{
+						CommonDirectoryLength = Idx;
+					}
+					break;
+				}
+				else if(Idx == Directory.CanonicalName.Length)
+				{
+					// Check whether the end of the directory name coincides with a boundary for the current name.
+					if(CanonicalName[Idx] == Path.DirectorySeparatorChar)
+					{
+						CommonDirectoryLength = Idx;
+					}
+					break;
+				}
+				else
+				{
+					// Check the two paths match, and bail if they don't. Increase the common directory length if we've reached a separator.
+					if (CanonicalName[Idx] != Directory.CanonicalName[Idx])
+					{
+						break;
+					}
+					if (CanonicalName[Idx] == Path.DirectorySeparatorChar)
+					{
+						CommonDirectoryLength = Idx;
+					}
+				}
+			}
+
+			// If there's no relative path, just return the absolute path
+			if(CommonDirectoryLength == -1)
+			{
+				return FullName;
+			}
+
+			// Append all the '..' separators to get back to the common directory, then the rest of the string to reach the target item
+			StringBuilder Result = new StringBuilder();
+			for(int Idx = CommonDirectoryLength + 1; Idx < Directory.CanonicalName.Length; Idx++)
+			{
+				// Move up a directory
+				Result.Append("..");
+				Result.Append(Path.DirectorySeparatorChar);
+
+				// Scan to the next directory separator
+				while(Idx < Directory.CanonicalName.Length && Directory.CanonicalName[Idx] != Path.DirectorySeparatorChar)
+				{
+					Idx++;
+				}
+			}
+			if(CommonDirectoryLength + 1 < FullName.Length)
+			{
+				Result.Append(FullName, CommonDirectoryLength + 1, FullName.Length - CommonDirectoryLength - 1);
+			}
+			return Result.ToString();
+		}
+
+		/// <summary>
+		/// Returns a string representation of this filesystem object
+		/// </summary>
+		/// <returns>Full path to the object</returns>
+		public override string ToString()
+		{
+			return FullName;
+		}
+	}
+
+	/// <summary>
+	/// Representation of a directory in the filesystem
+	/// </summary>
+	public class FullDirectoryName : FullFileSystemName
+	{
+		/// <summary>
+		/// Default constructor.
+		/// </summary>
+		/// <param name="InPath">Path to this directory.</param>
+		public FullDirectoryName(string InPath)
+			: base(InPath)
+		{
+		}
+
+		/// <summary>
+		/// Constructor for creating a directory object directly from two strings.
+		/// </summary>
+		/// <param name="InFullName"></param>
+		/// <param name="InCanonicalName"></param>
+		protected FullDirectoryName(string InFullName, string InCanonicalName)
+			: base(InFullName, InCanonicalName)
+		{
+		}
+
+		/// <summary>
+		/// Gets the directory containing this object
+		/// </summary>
+		/// <returns>A new directory object representing the directory containing this object</returns>
+		public FullDirectoryName GetParentDirectory()
+		{
+			if (IsRootDirectory())
+			{
+				return null;
+			}
+
+			int ParentLength = CanonicalName.LastIndexOf(Path.DirectorySeparatorChar);
+			if (ParentLength == 2 && CanonicalName[1] == ':')
+			{
+				ParentLength++;
+			}
+
+			return new FullDirectoryName(FullName.Substring(0, ParentLength), CanonicalName.Substring(0, ParentLength));
+		}
+
+		/// <summary>
+		/// Gets the parent directory for a file
+		/// </summary>
+		/// <param name="File">The file to get directory for</param>
+		/// <returns>The full directory name containing the given file</returns>
+		public static FullDirectoryName GetParentDirectory(FullFileName File)
+		{
+			int ParentLength = File.CanonicalName.LastIndexOf(Path.DirectorySeparatorChar);
+			return new FullDirectoryName(File.FullName.Substring(0, ParentLength), File.CanonicalName.Substring(0, ParentLength));
+		}
+
+		/// <summary>
+		/// Determines whether this path represents a root directory in the filesystem
+		/// </summary>
+		/// <returns>True if this path is a root directory, false otherwise</returns>
+		public bool IsRootDirectory()
+		{
+			return CanonicalName[CanonicalName.Length - 1] == Path.DirectorySeparatorChar;
+		}
+
+
+		/// <summary>
+		/// Compares two filesystem object names for equality. Uses the canonical name representation, not the display name representation.
+		/// </summary>
+		/// <param name="A">First object to compare.</param>
+		/// <param name="B">Second object to compare.</param>
+		/// <returns>True if the names represent the same object, false otherwise</returns>
+		public static bool operator ==(FullDirectoryName A, FullDirectoryName B)
+		{
+			return A.CanonicalName == B.CanonicalName;
+		}
+
+		/// <summary>
+		/// Compares two filesystem object names for inequality. Uses the canonical name representation, not the display name representation.
+		/// </summary>
+		/// <param name="A">First object to compare.</param>
+		/// <param name="B">Second object to compare.</param>
+		/// <returns>False if the names represent the same object, true otherwise</returns>
+		public static bool operator !=(FullDirectoryName A, FullDirectoryName B)
+		{
+			return A.CanonicalName != B.CanonicalName;
+		}
+
+		/// <summary>
+		/// Compares against another object for equality.
+		/// </summary>
+		/// <param name="A">First name to compare.</param>
+		/// <param name="B">Second name to compare.</param>
+		/// <returns>True if the names represent the same object, false otherwise</returns>
+		public override bool Equals(object Obj)
+		{
+			return (Obj is FullDirectoryName) && ((FullDirectoryName)Obj) == this;
+		}
+
+		/// <summary>
+		/// Returns a hash code for this object
+		/// </summary>
+		/// <returns></returns>
+		public override int GetHashCode()
+		{
+			return CanonicalName.GetHashCode();
+		}
+	}
+
+	/// <summary>
+	/// Representation of a file in the filesystem
+	/// </summary>
+	public class FullFileName : FullFileSystemName
+	{
+		/// <summary>
+		/// Default constructor.
+		/// </summary>
+		/// <param name="InPath">Path to this file</param>
+		public FullFileName(string InPath)
+			: base(InPath)
+		{
+		}
+
+		/// <summary>
+		/// Gets the directory containing this file
+		/// </summary>
+		/// <returns>A new directory object representing the directory containing this object</returns>
+		public FullDirectoryName GetDirectory()
+		{
+			return FullDirectoryName.GetParentDirectory(this);
+		}
+
+
+		/// <summary>
+		/// Compares two filesystem object names for equality. Uses the canonical name representation, not the display name representation.
+		/// </summary>
+		/// <param name="A">First object to compare.</param>
+		/// <param name="B">Second object to compare.</param>
+		/// <returns>True if the names represent the same object, false otherwise</returns>
+		public static bool operator ==(FullFileName A, FullFileName B)
+		{
+			return A.CanonicalName == B.CanonicalName;
+		}
+
+		/// <summary>
+		/// Compares two filesystem object names for inequality. Uses the canonical name representation, not the display name representation.
+		/// </summary>
+		/// <param name="A">First object to compare.</param>
+		/// <param name="B">Second object to compare.</param>
+		/// <returns>False if the names represent the same object, true otherwise</returns>
+		public static bool operator !=(FullFileName A, FullFileName B)
+		{
+			return A.CanonicalName != B.CanonicalName;
+		}
+
+		/// <summary>
+		/// Compares against another object for equality.
+		/// </summary>
+		/// <param name="A">First name to compare.</param>
+		/// <param name="B">Second name to compare.</param>
+		/// <returns>True if the names represent the same object, false otherwise</returns>
+		public override bool Equals(object Obj)
+		{
+			return (Obj is FullFileName) && ((FullFileName)Obj) == this;
+		}
+
+		/// <summary>
+		/// Returns a hash code for this object
+		/// </summary>
+		/// <returns></returns>
+		public override int GetHashCode()
+		{
+			return CanonicalName.GetHashCode();
+		}
+	}
+}
diff --git a/Engine/Source/Programs/UnrealBuildTool/UnrealBuildTool.csproj b/Engine/Source/Programs/UnrealBuildTool/UnrealBuildTool.csproj
--- a/Engine/Source/Programs/UnrealBuildTool/UnrealBuildTool.csproj
+++ b/Engine/Source/Programs/UnrealBuildTool/UnrealBuildTool.csproj
@@ -124,6 +124,7 @@
     <Compile Include="Mac\UEBuildMac.cs" />
     <Compile Include="Mac\MacProjectGenerator.cs" />
     <Compile Include="Mac\UEDeployMac.cs" />
+    <Compile Include="System\FileSystemNames.cs" />
     <Compile Include="System\TargetReceipt.cs" />
     <Compile Include="System\VersionManifest.cs" />
     <Compile Include="ToolChain\AppleToolChain.cs" />
