UE-20590 - fix for building IOS UE4Game from PC
#ios

diff --git a/Engine/Source/Programs/UnrealBuildTool/IOS/IOSToolChain.cs b/Engine/Source/Programs/UnrealBuildTool/IOS/IOSToolChain.cs
--- a/Engine/Source/Programs/UnrealBuildTool/IOS/IOSToolChain.cs
+++ b/Engine/Source/Programs/UnrealBuildTool/IOS/IOSToolChain.cs
@@ -1165,7 +1165,10 @@ namespace UnrealBuildTool
 								throw new BuildException("The target was not in a /Source subdirectory");
 							}
 						}
-						PathToApp += "\\" + AppName + ".uproject";
+						if (AppName != "UE4Game")
+						{
+							PathToApp += "\\" + AppName + ".uproject";
+						}
 					}
 
 					if (bUseDangerouslyFastMode)
