PR1483 - upload libraries in UBT so static link libraries get properly discovered on iOS
fix provided by pluranium
#ios

diff --git a/Engine/Source/Programs/UnrealBuildTool/IOS/IOSToolChain.cs b/Engine/Source/Programs/UnrealBuildTool/IOS/IOSToolChain.cs
--- a/Engine/Source/Programs/UnrealBuildTool/IOS/IOSToolChain.cs
+++ b/Engine/Source/Programs/UnrealBuildTool/IOS/IOSToolChain.cs
@@ -600,7 +600,7 @@ namespace UnrealBuildTool
 					{
 						// add it to the prerequisites to make sure it's built first (this should be the case of non-system libraries)
 						FileItem LibFile = FileItem.GetItemByPath(Path.GetFullPath(AdditionalLibrary));
-						FileItem RemoteLibFile = LocalToRemoteFileItem(LibFile, false);
+						FileItem RemoteLibFile = LocalToRemoteFileItem(LibFile, true);
 						LinkAction.PrerequisiteItems.Add(RemoteLibFile);
 
 						// and add to the commandline
