UEB-394 - improved error message when UAT/UBT fails to execute due to global mutex. Removed name of mutex as it was not useful, and gave an idea of how to find the offending instance.

diff --git a/Engine/Source/Programs/AutomationTool/AutomationUtils/Utils.cs b/Engine/Source/Programs/AutomationTool/AutomationUtils/Utils.cs
--- a/Engine/Source/Programs/AutomationTool/AutomationUtils/Utils.cs
+++ b/Engine/Source/Programs/AutomationTool/AutomationUtils/Utils.cs
@@ -698,11 +698,7 @@ namespace AutomationTool
 			{
 				if (!bCreatedMutex)
 				{
-                    throw new AutomationException("Another instance of {0} is already running.", EntryAssemblyLocation);
-				}
-				else
-				{
-					Log.TraceVerbose("No other instance of {0} is running.", EntryAssemblyLocation);
+                    throw new AutomationException("A conflicting instance of AutomationTool is already running. Curent location: {0}. A process manager may be used to determine the conflicting process and what tool may have launched it", EntryAssemblyLocation);
 				}
 
 				ExitCode Result = Main(Param);
diff --git a/Engine/Source/Programs/UnrealBuildTool/System/UnrealBuildTool.cs b/Engine/Source/Programs/UnrealBuildTool/System/UnrealBuildTool.cs
--- a/Engine/Source/Programs/UnrealBuildTool/System/UnrealBuildTool.cs
+++ b/Engine/Source/Programs/UnrealBuildTool/System/UnrealBuildTool.cs
@@ -967,7 +967,14 @@ namespace UnrealBuildTool
                     SingleInstanceMutex = new Mutex(true, MutexName, out bCreatedMutex);
                     if (!bCreatedMutex)
                     {
-                        throw new BuildException("Mutex {0} already set, indicating that a conflicting instance of {1} is already running.", MutexName, Assembly.GetEntryAssembly().GetOriginalLocation());
+                        if (bAutoSDKOnly || bValidatePlatforms)
+                        {
+                            throw new BuildException("A conflicting instance of UnrealBuildTool is already running. Either -autosdkonly or -validateplatform was passed, which allows only a single instance to be run globally. Therefore, the conflicting instance may be in another location or the current location: {0}. A process manager may be used to determine the conflicting process and what tool may have launched it.", Assembly.GetEntryAssembly().GetOriginalLocation());
+                        }
+                        else
+                        {
+                            throw new BuildException("A conflicting instance of UnrealBuildTool is already running at this location: {0}. A process manager may be used to determine the conflicting process and what tool may have launched it.", Assembly.GetEntryAssembly().GetOriginalLocation());
+                        }
                     }
                 }
 
