UE-18651 - properly write the .suo file in the correct directory for VS2015
#ubt

diff --git a/Engine/Source/Programs/UnrealBuildTool/System/VCProjectFileGenerator.cs b/Engine/Source/Programs/UnrealBuildTool/System/VCProjectFileGenerator.cs
--- a/Engine/Source/Programs/UnrealBuildTool/System/VCProjectFileGenerator.cs
+++ b/Engine/Source/Programs/UnrealBuildTool/System/VCProjectFileGenerator.cs
@@ -864,6 +864,11 @@ namespace UnrealBuildTool
 
 				// Check it doesn't exist before overwriting it. Since these files store the user's preferences, it'd be bad form to overwrite them.
 				string SolutionOptionsFileName = Path.Combine(MasterProjectRelativePath, Path.ChangeExtension(SolutionFileName, SolutionOptionsExtension));
+				if (ProjectFileFormat == VCProjectFileFormat.VisualStudio2015)
+				{
+					SolutionOptionsFileName = Path.Combine(MasterProjectRelativePath, ".vs", Path.GetFileNameWithoutExtension(SolutionFileName), "v14", ".suo");
+					Directory.CreateDirectory(Path.Combine(MasterProjectRelativePath, ".vs", Path.GetFileNameWithoutExtension(SolutionFileName), "v14"));
+				}
 				if(!File.Exists(SolutionOptionsFileName))
 				{
 					VCSolutionOptions Options = new VCSolutionOptions();
