Fix detection of automation projects when running UBT with -ProjectFiles.

diff --git a/Engine/Source/Programs/UnrealBuildTool/System/ProjectFileGenerator.cs b/Engine/Source/Programs/UnrealBuildTool/System/ProjectFileGenerator.cs
--- a/Engine/Source/Programs/UnrealBuildTool/System/ProjectFileGenerator.cs
+++ b/Engine/Source/Programs/UnrealBuildTool/System/ProjectFileGenerator.cs
@@ -262,7 +262,7 @@ namespace UnrealBuildTool
 			}
 
 			// Find all the automation modules .csproj files to add
-			var ModuleFiles = RulesCompiler.FindAllRulesSourceFiles(RulesCompiler.RulesFileType.AutomationModule, BuildFolders, ForeignPlugins:null, AdditionalSearchPaths:null);
+			var ModuleFiles = RulesCompiler.FindAllRulesSourceFiles(RulesCompiler.RulesFileType.AutomationModule, null, ForeignPlugins:null, AdditionalSearchPaths: BuildFolders );
 			foreach (var ProjectFile in ModuleFiles)
 			{
 				FileInfo Info = new FileInfo(Path.Combine(ProjectFile));
