Use FileReference and DirectoryReference classes instead of strings for paths. Massively reduces the amount of paranoid calls to GetFullPath(), CleanDirectorySeparators() et al., and improves performance for really common UBT operations like making relative paths or checking if files are under a directory. Strong typing ensures that intention of parameters is explicit.

diff --git a/Engine/Source/Programs/AutomationTool/AutomationUtils/ProjectUtils.cs b/Engine/Source/Programs/AutomationTool/AutomationUtils/ProjectUtils.cs
--- a/Engine/Source/Programs/AutomationTool/AutomationUtils/ProjectUtils.cs
+++ b/Engine/Source/Programs/AutomationTool/AutomationUtils/ProjectUtils.cs
@@ -205,14 +205,14 @@ namespace AutomationTool
 
 			// Read the project descriptor, and find all the plugins available to this project
 			ProjectDescriptor Project = ProjectDescriptor.FromFile(RawProjectPath);
-			List<PluginInfo> AvailablePlugins = Plugins.ReadAvailablePlugins(BuildConfiguration.RelativeEnginePath, RawProjectPath);
+			List<PluginInfo> AvailablePlugins = Plugins.ReadAvailablePlugins(new DirectoryReference(BuildConfiguration.RelativeEnginePath), new FileReference(RawProjectPath));
 
 			// check the target platforms for any differences in build settings or additional plugins
 			bool RetVal = false;
 			foreach (UnrealTargetPlatform TargetPlatformType in TargetPlatforms)
 			{
 				IUEBuildPlatform BuildPlat = UEBuildPlatform.GetBuildPlatform(TargetPlatformType, true);
-				if (!GlobalCommandLine.Rocket && BuildPlat != null && !(BuildPlat as UEBuildPlatform).HasDefaultBuildConfig(TargetPlatformType, Path.GetDirectoryName(RawProjectPath)))
+				if (!GlobalCommandLine.Rocket && BuildPlat != null && !(BuildPlat as UEBuildPlatform).HasDefaultBuildConfig(TargetPlatformType, new DirectoryReference(Path.GetDirectoryName(RawProjectPath))))
 				{
 					RetVal = true;
 					break;
@@ -399,24 +399,24 @@ namespace AutomationTool
 		private static void DetectTargetsForProject(ProjectProperties Properties, List<string> ExtraSearchPaths = null)
 		{
 			Properties.Targets = new Dictionary<TargetRules.TargetType, SingleTargetProperties>();
-			string TargetsDllFilename;
+			FileReference TargetsDllFilename;
 			string FullProjectPath = null;
 
-			var GameFolders = new List<string>();
-			var RulesFolder = GetRulesAssemblyFolder();
+			var GameFolders = new List<DirectoryReference>();
+			var RulesFolder = new DirectoryReference(GetRulesAssemblyFolder());
 			if (!String.IsNullOrEmpty(Properties.RawProjectPath))
 			{
 				CommandUtils.LogVerbose("Looking for targets for project {0}", Properties.RawProjectPath);
 
-				TargetsDllFilename = CommandUtils.CombinePaths(RulesFolder, String.Format("UATRules{0}.dll", Properties.RawProjectPath.GetHashCode()));
+				TargetsDllFilename = FileReference.Combine(RulesFolder, String.Format("UATRules{0}.dll", Properties.RawProjectPath.GetHashCode()));
 
 				FullProjectPath = CommandUtils.GetDirectoryName(Properties.RawProjectPath);
-				GameFolders.Add(FullProjectPath);
+				GameFolders.Add(new DirectoryReference(FullProjectPath));
 				CommandUtils.LogVerbose("Searching for target rule files in {0}", FullProjectPath);
 			}
 			else
 			{
-				TargetsDllFilename = CommandUtils.CombinePaths(RulesFolder, String.Format("UATRules{0}.dll", "_BaseEngine_"));
+				TargetsDllFilename = FileReference.Combine(RulesFolder, String.Format("UATRules{0}.dll", "_BaseEngine_"));
 			}
 
 			// the UBT code assumes a certain CWD, but artists don't have this CWD.
@@ -427,7 +427,8 @@ namespace AutomationTool
 				CommandUtils.PushDir(SourceDir);
 				DirPushed = true;
 			}
-			var TargetScripts = RulesCompiler.FindAllRulesSourceFiles(RulesCompiler.RulesFileType.Target, GameFolders: GameFolders, ForeignPlugins: null, AdditionalSearchPaths: ExtraSearchPaths);
+			var ExtraSearchDirectories = (ExtraSearchPaths == null)? null : ExtraSearchPaths.Select(x => new DirectoryReference(x)).ToList();
+			var TargetScripts = RulesCompiler.FindAllRulesSourceFiles(RulesCompiler.RulesFileType.Target, GameFolders: GameFolders, ForeignPlugins: null, AdditionalSearchPaths: ExtraSearchDirectories);
 			if (DirPushed)
 			{
 				CommandUtils.PopDir();
@@ -436,13 +437,12 @@ namespace AutomationTool
 			if (!CommandUtils.IsNullOrEmpty(TargetScripts))
 			{
 				// We only care about project target script so filter out any scripts not in the project folder, or take them all if we are just doing engine stuff
-				var ProjectTargetScripts = new List<string>();
-				foreach (var Filename in TargetScripts)
+				var ProjectTargetScripts = new List<FileReference>();
+				foreach (var TargetScript in TargetScripts)
 				{
-					var FullScriptPath = CommandUtils.CombinePaths(Path.GetFullPath(Filename));
-					if (FullProjectPath == null || FullScriptPath.StartsWith(FullProjectPath, StringComparison.InvariantCultureIgnoreCase))
+					if (FullProjectPath == null || TargetScript.IsUnderDirectory(new DirectoryReference(FullProjectPath)))
 					{
-						ProjectTargetScripts.Add(FullScriptPath);
+						ProjectTargetScripts.Add(TargetScript);
 					}
 				}
 				TargetScripts = ProjectTargetScripts;
@@ -463,9 +463,9 @@ namespace AutomationTool
 					Log.TraceVerbose("Targets DLL {0} is up to date.", TargetsDllFilename);
 					DoNotCompile = true;
 				}
-				if (!DoNotCompile && CommandUtils.FileExists_NoExceptions(TargetsDllFilename))
+				if (!DoNotCompile && CommandUtils.FileExists_NoExceptions(TargetsDllFilename.FullName))
 				{
-					if (!CommandUtils.DeleteFile_NoExceptions(TargetsDllFilename, true))
+					if (!CommandUtils.DeleteFile_NoExceptions(TargetsDllFilename.FullName, true))
 					{
 						DoNotCompile = true;
 						CommandUtils.LogVerbose("Could not delete {0} assuming it is up to date and reusable for a recursive UAT call.", TargetsDllFilename);
@@ -483,7 +483,7 @@ namespace AutomationTool
 		/// <param name="TargetsDllFilename"></param>
 		/// <param name="DoNotCompile"></param>
 		/// <param name="TargetScripts"></param>
-		private static void CompileAndLoadTargetsAssembly(ProjectProperties Properties, string TargetsDllFilename, bool DoNotCompile, List<string> TargetScripts)
+		private static void CompileAndLoadTargetsAssembly(ProjectProperties Properties, FileReference TargetsDllFilename, bool DoNotCompile, List<FileReference> TargetScripts)
 		{
 			CommandUtils.LogVerbose("Compiling targets DLL: {0}", TargetsDllFilename);
 
@@ -535,15 +535,15 @@ namespace AutomationTool
 		/// <param name="TargetsDllFilename"></param>
 		/// <param name="TargetScripts"></param>
 		/// <returns>True if the generated assembly is out of date.</returns>
-		private static bool CheckIfScriptAssemblyIsOutOfDate(string TargetsDllFilename, List<string> TargetScripts)
+		private static bool CheckIfScriptAssemblyIsOutOfDate(FileReference TargetsDllFilename, List<FileReference> TargetScripts)
 		{
 			var bOutOfDate = false;
-			var AssemblyInfo = new FileInfo(TargetsDllFilename);
+			var AssemblyInfo = new FileInfo(TargetsDllFilename.FullName);
 			if (AssemblyInfo.Exists)
 			{
 				foreach (var ScriptFilename in TargetScripts)
 				{
-					var ScriptInfo = new FileInfo(ScriptFilename);
+					var ScriptInfo = new FileInfo(ScriptFilename.FullName);
 					if (ScriptInfo.Exists && ScriptInfo.LastWriteTimeUtc > AssemblyInfo.LastWriteTimeUtc)
 					{
 						bOutOfDate = true;
@@ -610,7 +610,7 @@ namespace AutomationTool
                 GameName = InfoEntry.GameName;
 
                 //not sure what the heck this path is relative to
-                FilePath = Path.GetFullPath(CommandUtils.CombinePaths(CommandUtils.CmdEnv.LocalRoot, "Engine", "Binaries", InfoEntry.FilePath));
+                FilePath = InfoEntry.FilePath.FullName;
 
                 if (!CommandUtils.FileExists_NoExceptions(FilePath))
                 {
diff --git a/Engine/Source/Programs/AutomationTool/AutomationUtils/ScriptCompiler.cs b/Engine/Source/Programs/AutomationTool/AutomationUtils/ScriptCompiler.cs
--- a/Engine/Source/Programs/AutomationTool/AutomationUtils/ScriptCompiler.cs
+++ b/Engine/Source/Programs/AutomationTool/AutomationUtils/ScriptCompiler.cs
@@ -116,46 +116,38 @@ namespace AutomationTool
 			var OldCWD = Environment.CurrentDirectory;
 			var UnrealBuildToolCWD = CommandUtils.CombinePaths(CommandUtils.CmdEnv.LocalRoot, "Engine", "Source");
 
-			// Convert script folders to be relative to UnrealBuildTool's expected CWD
-			var RemappedAdditionalScriptFolders = new List<string>();
-			foreach (var CurFolder in AdditionalScriptsFolders)
-			{
-				RemappedAdditionalScriptFolders.Add(UnrealBuildTool.Utils.MakePathRelativeTo(CurFolder, UnrealBuildToolCWD));
-			}
-
 			Environment.CurrentDirectory = UnrealBuildToolCWD;
 
 			// Configure the rules compiler
 			// Get all game folders and convert them to build subfolders.
-			List<string> AllGameFolders;
+			List<DirectoryReference> AllGameFolders;
 			if(ScriptsForProjectFileName == null)
 			{
 				AllGameFolders = UnrealBuildTool.UEBuildTarget.DiscoverAllGameFolders();
 			}
 			else
 			{
-				AllGameFolders = new List<string>{ Path.GetDirectoryName(ScriptsForProjectFileName) };
+				AllGameFolders = new List<DirectoryReference>{ new DirectoryReference(Path.GetDirectoryName(ScriptsForProjectFileName)) };
 			}
-			var BuildFolders = new List<string>(AllGameFolders.Count);
+			var BuildFolders = new List<DirectoryReference>(AllGameFolders.Count);
 			foreach (var Folder in AllGameFolders)
 			{
-				var GameBuildFolder = CommandUtils.CombinePaths(Folder, "Build");
-				if (Directory.Exists(GameBuildFolder))
+				var GameBuildFolder = DirectoryReference.Combine(Folder, "Build");
+				if (GameBuildFolder.Exists())
 				{
 					BuildFolders.Add(GameBuildFolder);
 				}
 			}
-			RemappedAdditionalScriptFolders.AddRange(BuildFolders);
 
 			Log.TraceVerbose("Discovering game folders.");
 
-			var DiscoveredModules = UnrealBuildTool.RulesCompiler.FindAllRulesSourceFiles(UnrealBuildTool.RulesCompiler.RulesFileType.AutomationModule, GameFolders: AllGameFolders, ForeignPlugins: null, AdditionalSearchPaths: RemappedAdditionalScriptFolders);
+			var DiscoveredModules = UnrealBuildTool.RulesCompiler.FindAllRulesSourceFiles(UnrealBuildTool.RulesCompiler.RulesFileType.AutomationModule, GameFolders: AllGameFolders, ForeignPlugins: null, AdditionalSearchPaths: AdditionalScriptsFolders.Select(x => new DirectoryReference(x)).ToList());
 			var ModulesToCompile = new List<string>(DiscoveredModules.Count);
 			foreach (var ModuleFilename in DiscoveredModules)
 			{
-				if (HostPlatform.Current.IsScriptModuleSupported(CommandUtils.GetFilenameWithoutAnyExtensions(ModuleFilename)))
+				if (HostPlatform.Current.IsScriptModuleSupported(ModuleFilename.GetFileNameWithoutAnyExtensions()))
 				{
-					ModulesToCompile.Add(ModuleFilename);
+					ModulesToCompile.Add(ModuleFilename.FullName);
 				}
 				else
 				{
@@ -169,7 +161,6 @@ namespace AutomationTool
 				string Modules = string.Join(";", ModulesToCompile.ToArray());
 				var UATProj = CommandUtils.CombinePaths(CommandUtils.CmdEnv.LocalRoot, @"Engine\Source\Programs\AutomationTool\Scripts\UAT.proj");
 				var CmdLine = String.Format("\"{0}\" /p:Modules=\"{1}\" /p:Configuration={2} /verbosity:minimal /nologo", UATProj, Modules, BuildConfig);
-                Log.TraceVerbose("Building Automation projects in parallel...");
                 // supress the run command because it can be long and intimidating, making the logs around this code harder to read.
                 var Result = CommandUtils.Run(CommandUtils.CmdEnv.MsBuildExe, CmdLine, Options: CommandUtils.ERunOptions.Default | CommandUtils.ERunOptions.NoLoggingOfRunCommand | CommandUtils.ERunOptions.LoggingOfRunDuration);
 				if (Result.ExitCode != 0)
diff --git a/Engine/Source/Programs/AutomationTool/HTML5/HTML5Platform.Automation.cs b/Engine/Source/Programs/AutomationTool/HTML5/HTML5Platform.Automation.cs
--- a/Engine/Source/Programs/AutomationTool/HTML5/HTML5Platform.Automation.cs
+++ b/Engine/Source/Programs/AutomationTool/HTML5/HTML5Platform.Automation.cs
@@ -437,7 +437,7 @@ public class HTML5Platform : Platform
 		// Check HTML5LaunchHelper source for command line args
 
 		var LowerBrowserPath = BrowserPath.ToLower();
-		var ProfileDirectory = Path.Combine(Utils.GetUserSettingDirectory(), "UE4_HTML5", "user");
+		var ProfileDirectory = Path.Combine(Utils.GetUserSettingDirectory().FullName, "UE4_HTML5", "user");
 
 		string BrowserCommandline = url; 
 
diff --git a/Engine/Source/Programs/AutomationTool/Scripts/BuildPluginCommand.Automation.cs b/Engine/Source/Programs/AutomationTool/Scripts/BuildPluginCommand.Automation.cs
--- a/Engine/Source/Programs/AutomationTool/Scripts/BuildPluginCommand.Automation.cs
+++ b/Engine/Source/Programs/AutomationTool/Scripts/BuildPluginCommand.Automation.cs
@@ -25,7 +25,7 @@ class BuildPlugin : BuildCommand
 		}
 
 		// Read the plugin
-		PluginDescriptor Plugin = PluginDescriptor.FromFile(PluginFileName);
+		PluginDescriptor Plugin = PluginDescriptor.FromFile(new FileReference(PluginFileName));
 
 		// Clean the intermediate build directory
 		string IntermediateBuildDirectory = Path.Combine(Path.GetDirectoryName(PluginFileName), "Intermediate", "Build");
@@ -142,7 +142,7 @@ class BuildPlugin : BuildCommand
 
 		// Get the output plugin filename
 		string TargetPluginFileName = CommandUtils.MakeRerootedFilePath(Path.GetFullPath(PluginFileName), Path.GetDirectoryName(Path.GetFullPath(PluginFileName)), PackageDirectory);
-		PluginDescriptor NewDescriptor = PluginDescriptor.FromFile(TargetPluginFileName);
+		PluginDescriptor NewDescriptor = PluginDescriptor.FromFile(new FileReference(TargetPluginFileName));
 		NewDescriptor.bEnabledByDefault = true;
 		NewDescriptor.bInstalled = true;
 		NewDescriptor.Save(TargetPluginFileName);
diff --git a/Engine/Source/Programs/AutomationTool/Scripts/CopyBuildToStagingDirectory.Automation.cs b/Engine/Source/Programs/AutomationTool/Scripts/CopyBuildToStagingDirectory.Automation.cs
--- a/Engine/Source/Programs/AutomationTool/Scripts/CopyBuildToStagingDirectory.Automation.cs
+++ b/Engine/Source/Programs/AutomationTool/Scripts/CopyBuildToStagingDirectory.Automation.cs
@@ -377,14 +377,14 @@ public partial class Project : CommandUtils
 			ProjectDescriptor Project = ProjectDescriptor.FromFile(SC.RawProjectPath);
 			LogLog("Searching for plugins with CurrentWorkingDir: " + Directory.GetCurrentDirectory());
 			LogLog("Searching for plugins in: " + SC.RawProjectPath);
-			List<PluginInfo> AvailablePlugins = Plugins.ReadAvailablePlugins(CombinePaths(SC.LocalRoot, "Engine"), SC.RawProjectPath);
+			List<PluginInfo> AvailablePlugins = Plugins.ReadAvailablePlugins(new DirectoryReference(CombinePaths(SC.LocalRoot, "Engine")), new FileReference(SC.RawProjectPath));
 			foreach (PluginInfo Plugin in AvailablePlugins)
 			{
-				LogLog("Considering Plugin for Stage: " + Plugin.FileName);
+				LogLog("Considering Plugin for Stage: " + Plugin.File.FullName);
 				if (UProjectInfo.IsPluginEnabledForProject(Plugin, Project, SC.StageTargetPlatform.PlatformType))
 				{
-					LogLog("EnabledPlugin: " + Plugin.FileName);
-					SC.StageFiles(StagedFileType.UFS, Plugin.Directory, "*.uplugin", false, null, null, true, !Params.UsePak(SC.StageTargetPlatform), null, true, false);
+					LogLog("EnabledPlugin: " + Plugin.File.FullName);
+					SC.StageFiles(StagedFileType.UFS, Plugin.Directory.FullName, "*.uplugin", false, null, null, true, !Params.UsePak(SC.StageTargetPlatform), null, true, false);
 				}
 			}
 
@@ -1577,7 +1577,7 @@ public partial class Project : CommandUtils
 						}
 
 						// Convert the paths to absolute
-						Receipt.ExpandPathVariables(EngineDir, Path.GetDirectoryName(Params.RawProjectPath));
+						Receipt.ExpandPathVariables(new DirectoryReference(EngineDir), new DirectoryReference(Path.GetDirectoryName(Params.RawProjectPath)));
 						TargetsToStage.Add(new StageTarget{ Receipt = Receipt, RequireFilesExist = bRequireStagedFilesToExist });
 					}
 				}
diff --git a/Engine/Source/Programs/AutomationTool/Scripts/RocketBuild.Automation.cs b/Engine/Source/Programs/AutomationTool/Scripts/RocketBuild.Automation.cs
--- a/Engine/Source/Programs/AutomationTool/Scripts/RocketBuild.Automation.cs
+++ b/Engine/Source/Programs/AutomationTool/Scripts/RocketBuild.Automation.cs
@@ -989,7 +989,7 @@ namespace Rocket
 					}
 
 					// Convert the paths to absolute
-					Receipt.ExpandPathVariables(EnginePath, EnginePath);
+					Receipt.ExpandPathVariables(new DirectoryReference(EnginePath), new DirectoryReference(EnginePath));
 
 					foreach (var RuntimeDependency in Receipt.RuntimeDependencies)
 					{
diff --git a/Engine/Source/Programs/UnrealBuildTool/Android/AndroidProjectGenerator.cs b/Engine/Source/Programs/UnrealBuildTool/Android/AndroidProjectGenerator.cs
--- a/Engine/Source/Programs/UnrealBuildTool/Android/AndroidProjectGenerator.cs
+++ b/Engine/Source/Programs/UnrealBuildTool/Android/AndroidProjectGenerator.cs
@@ -198,7 +198,7 @@ namespace UnrealBuildTool
 		 *	
 		 *	@return	string				The custom path lines for the project file; Empty string if it doesn't require one
 		 */
-		public override string GetVisualStudioPathsEntries(UnrealTargetPlatform InPlatform, UnrealTargetConfiguration InConfiguration, TargetRules.TargetType TargetType, string TargetRulesPath, string ProjectFilePath, string NMakeOutputPath)
+		public override string GetVisualStudioPathsEntries(UnrealTargetPlatform InPlatform, UnrealTargetConfiguration InConfiguration, TargetRules.TargetType TargetType, FileReference TargetRulesPath, FileReference ProjectFilePath, FileReference NMakeOutputPath)
 		{
 			if (!IsNsightInstalled())
 			{
@@ -212,18 +212,18 @@ namespace UnrealBuildTool
 			//       environment variable
 
 			//@todo android: clean up debug path generation
-			string GameName = Path.GetFileNameWithoutExtension(TargetRulesPath);
+			string GameName = TargetRulesPath.GetFileNameWithoutExtension();
 			GameName = Path.GetFileNameWithoutExtension(GameName);
 
 
 			// intermediate path for Engine or Game's intermediate
 			string IntermediateDirectoryPath;
-			IntermediateDirectoryPath = Path.GetDirectoryName(NMakeOutputPath) + "/../../Intermediate/Android/APK";
+			IntermediateDirectoryPath = Path.GetDirectoryName(NMakeOutputPath.FullName) + "/../../Intermediate/Android/APK";
 
 			// string for <OverrideAPKPath>
 			string APKPath = Path.Combine(
-				Path.GetDirectoryName(NMakeOutputPath),
-				Path.GetFileNameWithoutExtension(NMakeOutputPath) + "-armv7-es2.apk");
+				Path.GetDirectoryName(NMakeOutputPath.FullName),
+				Path.GetFileNameWithoutExtension(NMakeOutputPath.FullName) + "-armv7-es2.apk");
 
 			// string for <BuildXmlPath> and <AndroidManifestPath>
 			string BuildXmlPath = IntermediateDirectoryPath;
diff --git a/Engine/Source/Programs/UnrealBuildTool/Android/AndroidToolChain.cs b/Engine/Source/Programs/UnrealBuildTool/Android/AndroidToolChain.cs
--- a/Engine/Source/Programs/UnrealBuildTool/Android/AndroidToolChain.cs
+++ b/Engine/Source/Programs/UnrealBuildTool/Android/AndroidToolChain.cs
@@ -788,7 +788,7 @@ namespace UnrealBuildTool
 			}
 		}
 
-		static void GenerateEmptyLinkFunctionsForRemovedModules(List<FileItem> SourceFiles, string ModuleName, string OutputDirectory)
+		static void GenerateEmptyLinkFunctionsForRemovedModules(List<FileItem> SourceFiles, string ModuleName, DirectoryReference OutputDirectory)
 		{
 			// Only add to UELinkerFixups module
 			if (!ModuleName.Equals("Launch"))
@@ -797,10 +797,10 @@ namespace UnrealBuildTool
 			}
 
 			string LinkerExceptionsName = "../UELinkerExceptions";
-			string LinkerExceptionsCPPFilename = Path.Combine(OutputDirectory, LinkerExceptionsName + ".cpp");
+			FileReference LinkerExceptionsCPPFilename = FileReference.Combine(OutputDirectory, LinkerExceptionsName + ".cpp");
 
 			// Create the cpp filename
-			if (!File.Exists(LinkerExceptionsCPPFilename))
+			if (!LinkerExceptionsCPPFilename.Exists())
 			{
 				// Create a dummy file in case it doesn't exist yet so that the module does not complain it's not there
 				ResponseFile.Create(LinkerExceptionsCPPFilename, new List<string>());
@@ -827,9 +827,9 @@ namespace UnrealBuildTool
 
 			// Determine if the file changed. Write it if it either doesn't exist or the contents are different.
 			bool bShouldWriteFile = true;
-			if (File.Exists(LinkerExceptionsCPPFilename))
+			if (LinkerExceptionsCPPFilename.Exists())
 			{
-				string[] ExistingExceptionText = File.ReadAllLines(LinkerExceptionsCPPFilename);
+				string[] ExistingExceptionText = File.ReadAllLines(LinkerExceptionsCPPFilename.FullName);
 				string JoinedNewContents = string.Join("", Result.ToArray());
 				string JoinedOldContents = string.Join("", ExistingExceptionText);
 				bShouldWriteFile = (JoinedNewContents != JoinedOldContents);
@@ -841,7 +841,7 @@ namespace UnrealBuildTool
 				ResponseFile.Create(LinkerExceptionsCPPFilename, Result);
 			}
 
-			SourceFiles.Add(FileItem.GetItemByPath(LinkerExceptionsCPPFilename));
+			SourceFiles.Add(FileItem.GetItemByFileReference(LinkerExceptionsCPPFilename));
 		}
 
 		// cache the location of NDK tools
@@ -1017,8 +1017,8 @@ namespace UnrealBuildTool
 						if (CompileEnvironment.Config.PrecompiledHeaderAction == PrecompiledHeaderAction.Create)
 						{
 							// Add the precompiled header file to the produced item list.
-							FileItem PrecompiledHeaderFile = FileItem.GetItemByPath(
-								Path.Combine(
+							FileItem PrecompiledHeaderFile = FileItem.GetItemByFileReference(
+								FileReference.Combine(
 									CompileEnvironment.Config.OutputDirectory,
 									Path.GetFileName(InlineArchName(SourceFile.AbsolutePath, Arch, GPUArchitecture) + PCHExtension)
 									)
@@ -1042,10 +1042,10 @@ namespace UnrealBuildTool
 							var ObjectFileExtension = UEBuildPlatform.BuildPlatformDictionary[UnrealTargetPlatform.Android].GetBinaryExtension(UEBuildBinaryType.Object);
 
 							// Add the object file to the produced item list.
-							FileItem ObjectFile = FileItem.GetItemByPath(
-								InlineArchName( Path.Combine(
+							FileItem ObjectFile = FileItem.GetItemByFileReference(
+								FileReference.Combine(
 									CompileEnvironment.Config.OutputDirectory,
-									Path.GetFileName(SourceFile.AbsolutePath) + ObjectFileExtension), Arch, GPUArchitecture
+									InlineArchName(Path.GetFileName(SourceFile.AbsolutePath) + ObjectFileExtension, Arch, GPUArchitecture)
 									)
 								);
 							CompileAction.ProducedItems.Add(ObjectFile);
@@ -1073,10 +1073,10 @@ namespace UnrealBuildTool
 						}
 
 						// Create the response file
-						string ResponseFileName = CompileAction.ProducedItems[0].AbsolutePath + ".response";
-						string ResponseArgument = string.Format("@\"{0}\"", ResponseFile.Create(ResponseFileName, new List<string> { AllArguments }));
+						FileReference ResponseFileName = CompileAction.ProducedItems[0].Reference + ".response";
+						string ResponseArgument = string.Format("@\"{0}\"", ResponseFile.Create(ResponseFileName, new List<string> { AllArguments }).FullName);
 
-						CompileAction.WorkingDirectory = Path.GetFullPath(".");
+						CompileAction.WorkingDirectory = UnrealBuildTool.EngineSourceDirectory.FullName;
 						CompileAction.CommandPath = ClangPath;
 						CompileAction.CommandArguments = ResponseArgument;
 						CompileAction.StatusDescription = string.Format("{0} [{1}-{2}]", Path.GetFileName(SourceFile.AbsolutePath), Arch.Replace("-", ""), GPUArchitecture.Replace("-", ""));
@@ -1134,14 +1134,14 @@ namespace UnrealBuildTool
 
 					// Android will have an array of outputs
 					if (LinkEnvironment.Config.OutputFilePaths.Count < OutputPathIndex ||
-						!Path.GetFileNameWithoutExtension(LinkEnvironment.Config.OutputFilePaths[OutputPathIndex]).EndsWith(Arch + GPUArchitecture))
+						!LinkEnvironment.Config.OutputFilePaths[OutputPathIndex].GetFileNameWithoutExtension().EndsWith(Arch + GPUArchitecture))
 					{
 						throw new BuildException("The OutputFilePaths array didn't match the Arches array in AndroidToolChain.LinkAllFiles");
 					}
 
 					// Create an action that invokes the linker.
 					Action LinkAction = new Action(ActionType.Link);
-					LinkAction.WorkingDirectory = Path.GetFullPath(".");
+					LinkAction.WorkingDirectory = UnrealBuildTool.EngineSourceDirectory.FullName;
 
 					if (LinkEnvironment.Config.bIsBuildingLibrary)
 					{
@@ -1161,14 +1161,14 @@ namespace UnrealBuildTool
 
 					string LinkerPath = LinkAction.WorkingDirectory;
 
-					LinkAction.WorkingDirectory = LinkEnvironment.Config.IntermediateDirectory;
+					LinkAction.WorkingDirectory = LinkEnvironment.Config.IntermediateDirectory.FullName;
 
 					// Get link arguments.
 					LinkAction.CommandArguments = LinkEnvironment.Config.bIsBuildingLibrary ? GetArArguments(LinkEnvironment) : GetLinkArguments(LinkEnvironment, Arch);
 
 					// Add the output file as a production of the link action.
 					FileItem OutputFile;
-					OutputFile = FileItem.GetItemByPath(LinkEnvironment.Config.OutputFilePaths[OutputPathIndex]);
+					OutputFile = FileItem.GetItemByFileReference(LinkEnvironment.Config.OutputFilePaths[OutputPathIndex]);
 					Outputs.Add(OutputFile);
 					LinkAction.ProducedItems.Add(OutputFile);
 					LinkAction.StatusDescription = string.Format("{0}", Path.GetFileName(OutputFile.AbsolutePath));
@@ -1194,7 +1194,7 @@ namespace UnrealBuildTool
 						{
 							string AbsolutePath = InputFile.AbsolutePath.Replace("\\", "/");
 
-							AbsolutePath = AbsolutePath.Replace(LinkEnvironment.Config.IntermediateDirectory.Replace("\\", "/"), "");
+							AbsolutePath = AbsolutePath.Replace(LinkEnvironment.Config.IntermediateDirectory.FullName.Replace("\\", "/"), "");
 							AbsolutePath = AbsolutePath.TrimStart(new char[] { '/' });
 
 							InputFileNames.Add(string.Format("\"{0}\"", AbsolutePath));
@@ -1202,7 +1202,7 @@ namespace UnrealBuildTool
 						}
 					}
 
-					string ResponseFileName = GetResponseFileName(LinkEnvironment, OutputFile);
+					FileReference ResponseFileName = GetResponseFileName(LinkEnvironment, OutputFile);
 					LinkAction.CommandArguments += string.Format(" @\"{0}\"", ResponseFile.Create(ResponseFileName, InputFileNames));
 
 					// libs don't link in other libs
@@ -1262,15 +1262,15 @@ namespace UnrealBuildTool
 			// the binary will have all of the .so's in the output files, we need to trim down to the shared apk (which is what needs to go into the manifest)
 			if (Binary.Config.Type != UEBuildBinaryType.StaticLibrary)
 			{
-				foreach (string BinaryPath in Binary.Config.OutputFilePaths)
+				foreach (FileReference BinaryPath in Binary.Config.OutputFilePaths)
 				{
-					string ApkFile = Path.ChangeExtension(BinaryPath, ".apk");
+					FileReference ApkFile = BinaryPath.ChangeExtension(".apk");
 					Receipt.AddBuildProduct(ApkFile, BuildProductType.Executable);
 				}
 			}
 		}
 
-		public override void CompileCSharpProject(CSharpEnvironment CompileEnvironment, string ProjectFileName, string DestinationFile)
+		public override void CompileCSharpProject(CSharpEnvironment CompileEnvironment, FileReference ProjectFileName, FileReference DestinationFile)
 		{
             throw new BuildException("Android cannot compile C# files");
 		}
diff --git a/Engine/Source/Programs/UnrealBuildTool/Android/UEBuildAndroid.cs b/Engine/Source/Programs/UnrealBuildTool/Android/UEBuildAndroid.cs
--- a/Engine/Source/Programs/UnrealBuildTool/Android/UEBuildAndroid.cs
+++ b/Engine/Source/Programs/UnrealBuildTool/Android/UEBuildAndroid.cs
@@ -58,7 +58,7 @@ namespace UnrealBuildTool
         {
             string NDKPath = Environment.GetEnvironmentVariable("NDKROOT");
             {
-                var configCacheIni = new ConfigCacheIni("Engine", null);
+                var configCacheIni = new ConfigCacheIni("Engine", (DirectoryReference)null);
                 var AndroidEnv = new Dictionary<string, string>();
 
                 Dictionary<string, string> EnvVarNames = new Dictionary<string,string> { 
@@ -165,10 +165,8 @@ namespace UnrealBuildTool
 			{
 				bool bRegisterBuildPlatform = true;
 
-				string EngineSourcePath = Path.Combine(ProjectFileGenerator.EngineRelativePath, "Source");
-				string AndroidTargetPlatformFile = Path.Combine(EngineSourcePath, "Developer", "Android", "AndroidTargetPlatform", "AndroidTargetPlatform.Build.cs");
-
-				if (File.Exists(AndroidTargetPlatformFile) == false)
+				FileReference AndroidTargetPlatformFile = FileReference.Combine(UnrealBuildTool.EngineSourceDirectory, "Developer", "Android", "AndroidTargetPlatform", "AndroidTargetPlatform.Build.cs");
+				if (AndroidTargetPlatformFile.Exists() == false)
 				{
 					bRegisterBuildPlatform = false;
 				}
@@ -250,7 +248,7 @@ namespace UnrealBuildTool
 			BuildConfiguration.bUseSharedPCHs = false;
 		}
 			
-		public override bool HasDefaultBuildConfig(UnrealTargetPlatform Platform, string ProjectPath)
+		public override bool HasDefaultBuildConfig(UnrealTargetPlatform Platform, DirectoryReference ProjectPath)
 		{
 			string[] BoolKeys = new string[] {
 				"bBuildForArmV7", "bBuildForArm64", "bBuildForX86", "bBuildForX8664", 
@@ -456,18 +454,18 @@ namespace UnrealBuildTool
 			};
 		}
 
-		public override List<string> FinalizeBinaryPaths(string BinaryName)
+		public override List<FileReference> FinalizeBinaryPaths(FileReference BinaryName)
 		{
 			var Architectures = AndroidToolChain.GetAllArchitectures();
 			var GPUArchitectures = AndroidToolChain.GetAllGPUArchitectures();
 
 			// make multiple output binaries
-			List<string> AllBinaries = new List<string>();
+			List<FileReference> AllBinaries = new List<FileReference>();
 			foreach (string Architecture in Architectures)
 			{
 				foreach (string GPUArchitecture in GPUArchitectures)
 				{
-					AllBinaries.Add(AndroidToolChain.InlineArchName(BinaryName, Architecture, GPUArchitecture));
+					AllBinaries.Add(new FileReference(AndroidToolChain.InlineArchName(BinaryName.FullName, Architecture, GPUArchitecture)));
 				}
 			}
 
diff --git a/Engine/Source/Programs/UnrealBuildTool/Android/UEDeployAndroid.cs b/Engine/Source/Programs/UnrealBuildTool/Android/UEDeployAndroid.cs
--- a/Engine/Source/Programs/UnrealBuildTool/Android/UEDeployAndroid.cs
+++ b/Engine/Source/Programs/UnrealBuildTool/Android/UEDeployAndroid.cs
@@ -1782,14 +1782,14 @@ namespace UnrealBuildTool.Android
 		public override bool PrepTargetForDeployment(UEBuildTarget InTarget)
 		{
 			// we need to strip architecture from any of the output paths
-			string BaseSoName = AndroidToolChain.RemoveArchName(InTarget.OutputPaths[0]);
+			string BaseSoName = AndroidToolChain.RemoveArchName(InTarget.OutputPaths[0].FullName);
 
 			// make an apk at the end of compiling, so that we can run without packaging (debugger, cook on the fly, etc)
-			MakeApk(InTarget.AppName, InTarget.ProjectDirectory, BaseSoName, BuildConfiguration.RelativeEnginePath, bForDistribution: false, CookFlavor: "", 
+			MakeApk(InTarget.AppName, InTarget.ProjectDirectory.FullName, BaseSoName, BuildConfiguration.RelativeEnginePath, bForDistribution: false, CookFlavor: "", 
 				bMakeSeparateApks:ShouldMakeSeparateApks(), bIncrementalPackage:true, bDisallowPackagingDataInApk:false);
 
 			// if we made any non-standard .apk files, the generated debugger settings may be wrong
-			if (ShouldMakeSeparateApks() && (InTarget.OutputPaths.Count > 1 || !InTarget.OutputPaths[0].Contains("-armv7-es2")))
+			if (ShouldMakeSeparateApks() && (InTarget.OutputPaths.Count > 1 || !InTarget.OutputPaths[0].FullName.Contains("-armv7-es2")))
 			{
 				Console.WriteLine("================================================================================================================================");
 				Console.WriteLine("Non-default apk(s) have been made: If you are debugging, you will need to manually select one to run in the debugger properties!");
diff --git a/Engine/Source/Programs/UnrealBuildTool/Configuration/EngineConfiguration.cs b/Engine/Source/Programs/UnrealBuildTool/Configuration/EngineConfiguration.cs
--- a/Engine/Source/Programs/UnrealBuildTool/Configuration/EngineConfiguration.cs
+++ b/Engine/Source/Programs/UnrealBuildTool/Configuration/EngineConfiguration.cs
@@ -63,6 +63,17 @@ namespace UnrealBuildTool
 		/// <param name="Platform">Target platform</param>
 		/// <param name="BaseIniName">Ini name (Engine, Editor, etc)</param>
 		public ConfigCacheIni(string BaseIniName, string ProjectDirectory, string EngineDirectory = null)
+		{
+			Init(UnrealTargetPlatform.Unknown, BaseIniName, (ProjectDirectory == null)? null : new DirectoryReference(ProjectDirectory), (EngineDirectory == null)? null : new DirectoryReference(EngineDirectory));
+		}
+
+		/// <summary>
+		/// Constructor. Parses ini hierarchy for the specified project.  No Platform settings.
+		/// </summary>
+		/// <param name="ProjectDirectory">Project path</param>
+		/// <param name="Platform">Target platform</param>
+		/// <param name="BaseIniName">Ini name (Engine, Editor, etc)</param>
+		public ConfigCacheIni(string BaseIniName, DirectoryReference ProjectDirectory, DirectoryReference EngineDirectory = null)
 		{
 			Init(UnrealTargetPlatform.Unknown, BaseIniName, ProjectDirectory, EngineDirectory);
 		}
@@ -74,21 +85,32 @@ namespace UnrealBuildTool
 		/// <param name="Platform">Target platform</param>
 		/// <param name="BaseIniName">Ini name (Engine, Editor, etc)</param>
 		public ConfigCacheIni(UnrealTargetPlatform Platform, string BaseIniName, string ProjectDirectory, string EngineDirectory = null)
+		{
+			Init(Platform, BaseIniName, (ProjectDirectory == null)? null : new DirectoryReference(ProjectDirectory), (EngineDirectory == null)? null : new DirectoryReference(EngineDirectory));
+		}
+
+		/// <summary>
+		/// Constructor. Parses ini hierarchy for the specified platform and project.
+		/// </summary>
+		/// <param name="ProjectDirectory">Project path</param>
+		/// <param name="Platform">Target platform</param>
+		/// <param name="BaseIniName">Ini name (Engine, Editor, etc)</param>
+		public ConfigCacheIni(UnrealTargetPlatform Platform, string BaseIniName, DirectoryReference ProjectDirectory, DirectoryReference EngineDirectory = null)
 		{
 			Init(Platform, BaseIniName, ProjectDirectory, EngineDirectory);
 		}
 
-		private void Init(UnrealTargetPlatform Platform, string BaseIniName, string ProjectDirectory, string EngineDirectory)
+		private void Init(UnrealTargetPlatform Platform, string BaseIniName, DirectoryReference ProjectDirectory, DirectoryReference EngineDirectory)
 		{
 			Sections = new Dictionary<string, IniSection>(StringComparer.InvariantCultureIgnoreCase);
-			if (String.IsNullOrEmpty(EngineDirectory))
+			if (EngineDirectory == null)
 			{
-				EngineDirectory = BuildConfiguration.RelativeEnginePath;
+				EngineDirectory = UnrealBuildTool.EngineDirectory;
 			}
 
 			foreach (var IniFileName in EnumerateCrossPlatformIniFileNames(ProjectDirectory, EngineDirectory, Platform, BaseIniName))
 			{
-				if (File.Exists(IniFileName))
+				if (IniFileName.Exists())
 				{
 					ParseIniFile(IniFileName);
 				}
@@ -376,12 +398,12 @@ namespace UnrealBuildTool
 		/// <summary>
 		/// Loads and parses ini file.
 		/// </summary>
-		public void ParseIniFile(string Filename)
+		public void ParseIniFile(FileReference Filename)
 		{
 			String[] IniLines = null;
 			try
 			{
-				IniLines = File.ReadAllLines(Filename);
+				IniLines = File.ReadAllLines(Filename.FullName);
 			}
 			catch (Exception ex)
 			{
@@ -453,7 +475,7 @@ namespace UnrealBuildTool
 		/// <summary>
 		/// Splits a line into key and value
 		/// </summary>
-		private static void ParseKeyValuePair(string TrimmedLine, string Filename, int LineIndex, out string Key, out string Value)
+		private static void ParseKeyValuePair(string TrimmedLine, FileReference Filename, int LineIndex, out string Key, out string Value)
 		{			
 			var AssignIndex = TrimmedLine.IndexOf('=');
 			if (AssignIndex < 0)
@@ -528,7 +550,7 @@ namespace UnrealBuildTool
 		/// <summary>
 		/// Finds an existing section or adds a new one
 		/// </summary>
-		private IniSection FindOrAddSection(string TrimmedLine, string Filename, int LineIndex)
+		private IniSection FindOrAddSection(string TrimmedLine, FileReference Filename, int LineIndex)
 		{
 			var SectionEndIndex = TrimmedLine.IndexOf(']');
 			if (SectionEndIndex != (TrimmedLine.Length - 1))
@@ -552,70 +574,70 @@ namespace UnrealBuildTool
 		/// <summary>
 		/// Returns a list of INI filenames for the given project
 		/// </summary>
-		private static IEnumerable<string> EnumerateCrossPlatformIniFileNames(string ProjectDirectory, string EngineDirectory, UnrealTargetPlatform Platform, string BaseIniName)
+		private static IEnumerable<FileReference> EnumerateCrossPlatformIniFileNames(DirectoryReference ProjectDirectory, DirectoryReference EngineDirectory, UnrealTargetPlatform Platform, string BaseIniName)
 		{
 			// Engine/Config/Base.ini (included in every ini type, required)
-			yield return Path.Combine(EngineDirectory, "Config", "Base.ini");
+			yield return FileReference.Combine(EngineDirectory, "Config", "Base.ini");
 
 			// Engine/Config/Base* ini
-			yield return Path.Combine(EngineDirectory, "Config", "Base" + BaseIniName + ".ini");
+			yield return FileReference.Combine(EngineDirectory, "Config", "Base" + BaseIniName + ".ini");
 
 			// Engine/Config/NotForLicensees/Base* ini
-			yield return Path.Combine(EngineDirectory, "Config", "NotForLicensees", "Base" + BaseIniName + ".ini");
+			yield return FileReference.Combine(EngineDirectory, "Config", "NotForLicensees", "Base" + BaseIniName + ".ini");
 
 			// NOTE: 4.7: See comment in GetSourceIniHierarchyFilenames()
 			// Engine/Config/NoRedist/Base* ini
 			// yield return Path.Combine(EngineDirectory, "Config", "NoRedist", "Base" + BaseIniName + ".ini");
 
-			if(!String.IsNullOrEmpty(ProjectDirectory))
+			if(ProjectDirectory != null)
 			{
 				// Game/Config/Default* ini
-				yield return Path.Combine(ProjectDirectory, "Config", "Default" + BaseIniName + ".ini");
+				yield return FileReference.Combine(ProjectDirectory, "Config", "Default" + BaseIniName + ".ini");
 			
 				// Game/Config/NotForLicensees/Default* ini
-				yield return Path.Combine(ProjectDirectory, "Config", "NotForLicensees", "Default" + BaseIniName + ".ini");
+				yield return FileReference.Combine(ProjectDirectory, "Config", "NotForLicensees", "Default" + BaseIniName + ".ini");
 
 				// Game/Config/NoRedist/Default* ini
-				yield return Path.Combine(ProjectDirectory, "Config", "NoRedist", "Default" + BaseIniName + ".ini");
+				yield return FileReference.Combine(ProjectDirectory, "Config", "NoRedist", "Default" + BaseIniName + ".ini");
 			}
 			
 			string PlatformName = GetIniPlatformName(Platform);			
 			if (Platform != UnrealTargetPlatform.Unknown)
 			{
 				// Engine/Config/Platform/Platform* ini
-				yield return Path.Combine(EngineDirectory, "Config", PlatformName, PlatformName + BaseIniName + ".ini");
+				yield return FileReference.Combine(EngineDirectory, "Config", PlatformName, PlatformName + BaseIniName + ".ini");
 
-				if(!String.IsNullOrEmpty(ProjectDirectory))
+				if(ProjectDirectory != null)
 				{
 					// Game/Config/Platform/Platform* ini
-					yield return Path.Combine(ProjectDirectory, "Config", PlatformName, PlatformName + BaseIniName + ".ini");
+					yield return FileReference.Combine(ProjectDirectory, "Config", PlatformName, PlatformName + BaseIniName + ".ini");
 				}
 			}
 
-			string UserSettingsFolder = Utils.GetUserSettingDirectory(); // Match FPlatformProcess::UserSettingsDir()
-			string PersonalFolder = null; // Match FPlatformProcess::UserDir()
+			DirectoryReference UserSettingsFolder = Utils.GetUserSettingDirectory(); // Match FPlatformProcess::UserSettingsDir()
+			DirectoryReference PersonalFolder = null; // Match FPlatformProcess::UserDir()
 			if (BuildHostPlatform.Current.Platform == UnrealTargetPlatform.Mac)
 			{
-				PersonalFolder = Path.Combine(Environment.GetFolderPath(Environment.SpecialFolder.Personal), "Documents");
+				PersonalFolder = new DirectoryReference(Path.Combine(Environment.GetFolderPath(Environment.SpecialFolder.Personal), "Documents"));
 			}
 			else if (Environment.OSVersion.Platform == PlatformID.Unix)
 			{
-				PersonalFolder = Path.Combine(Environment.GetFolderPath(Environment.SpecialFolder.Personal), "Documents");
+				PersonalFolder = new DirectoryReference(Path.Combine(Environment.GetFolderPath(Environment.SpecialFolder.Personal), "Documents"));
 			}
 			else
 			{
-				PersonalFolder = Environment.GetFolderPath(Environment.SpecialFolder.Personal);
+				PersonalFolder = new DirectoryReference(Environment.GetFolderPath(Environment.SpecialFolder.Personal));
 			}
 
 			// <AppData>/UE4/EngineConfig/User* ini
-			yield return Path.Combine(UserSettingsFolder, "Unreal Engine", "Engine", "Config", "User" + BaseIniName + ".ini");
+			yield return FileReference.Combine(UserSettingsFolder, "Unreal Engine", "Engine", "Config", "User" + BaseIniName + ".ini");
 			// <Documents>/UE4/EngineConfig/User* ini
-			yield return Path.Combine(PersonalFolder, "Unreal Engine", "Engine", "Config", "User" + BaseIniName + ".ini");
+			yield return FileReference.Combine(PersonalFolder, "Unreal Engine", "Engine", "Config", "User" + BaseIniName + ".ini");
 			
 			// Game/Config/User* ini
-            if (!String.IsNullOrEmpty(ProjectDirectory))
+            if (ProjectDirectory != null)
             {
-				yield return Path.Combine(ProjectDirectory, "Config", "User" + BaseIniName + ".ini");
+				yield return FileReference.Combine(ProjectDirectory, "Config", "User" + BaseIniName + ".ini");
 			}
 		}
 
diff --git a/Engine/Source/Programs/UnrealBuildTool/Configuration/UEBuildBinary.cs b/Engine/Source/Programs/UnrealBuildTool/Configuration/UEBuildBinary.cs
--- a/Engine/Source/Programs/UnrealBuildTool/Configuration/UEBuildBinary.cs
+++ b/Engine/Source/Programs/UnrealBuildTool/Configuration/UEBuildBinary.cs
@@ -39,12 +39,12 @@ namespace UnrealBuildTool
 		/// <summary>
 		/// The output file path. This must be set before a binary can be built using it.
 		/// </summary>
-		public List<string> OutputFilePaths = new List<string>();
+		public List<FileReference> OutputFilePaths = new List<FileReference>();
 
 		/// <summary>
 		/// Returns the OutputFilePath if there is only one entry in OutputFilePaths
 		/// </summary>
-		public string OutputFilePath
+		public FileReference OutputFilePath
 		{
 			get
 			{
@@ -59,12 +59,12 @@ namespace UnrealBuildTool
 		/// <summary>
 		/// Original output filepath. This is the original binary name before hot-reload suffix has been appended to it.
 		/// </summary>
-		public List<string> OriginalOutputFilePaths;
+		public List<FileReference> OriginalOutputFilePaths;
 
 		/// <summary>
 		/// Returns the OriginalOutputFilePath if there is only one entry in OriginalOutputFilePaths
 		/// </summary>
-		public string OriginalOutputFilePath
+		public FileReference OriginalOutputFilePath
 		{
 			get
 			{
@@ -79,7 +79,7 @@ namespace UnrealBuildTool
 		/// <summary>
 		/// The intermediate directory for this binary. Modules should create separate intermediate directories below this. Must be set before a binary can be built using it.
 		/// </summary>
-		public string IntermediateDirectory;
+		public DirectoryReference IntermediateDirectory;
 
 		/// <summary>
 		/// If true, build exports lib
@@ -120,7 +120,7 @@ namespace UnrealBuildTool
 		/// <summary>
 		/// The projectfile path
 		/// </summary>
-		public string ProjectFilePath = "";
+		public FileReference ProjectFilePath;
 		
 		/// <summary>
 		/// List of modules to link together into this executable
@@ -139,15 +139,15 @@ namespace UnrealBuildTool
 		/// <param name="InModuleNames"></param>
 		public UEBuildBinaryConfiguration(
 				UEBuildBinaryType InType,
-				IEnumerable<string> InOutputFilePaths = null,
-				string InIntermediateDirectory = null,
+				IEnumerable<FileReference> InOutputFilePaths = null,
+				DirectoryReference InIntermediateDirectory = null,
 				bool bInAllowExports = false,
 				bool bInCreateImportLibrarySeparately = false,
 				bool bInIncludeDependentLibrariesInLibrary = false,
 				bool bInAllowCompilation = true,
 				bool bInHasModuleRules = true,
                 bool bInIsCrossTarget = false,
-				string InProjectFilePath = "",
+				FileReference InProjectFilePath = null,
 				IEnumerable<string> InModuleNames = null
 			)
 		{
@@ -310,7 +310,7 @@ namespace UnrealBuildTool
 
 			// Add the primary build products
 			string DebugExtension = UEBuildPlatform.GetBuildPlatform(Target.Platform).GetDebugInfoExtension(Config.Type);
-			foreach (string OutputFilePath in Config.OutputFilePaths)
+			foreach (FileReference OutputFilePath in Config.OutputFilePaths)
 			{
 				AddBuildProductAndDebugFile(OutputFilePath, Type, DebugExtension, Receipt, ToolChain);
 			}
@@ -318,7 +318,7 @@ namespace UnrealBuildTool
 			// Add the console app, if there is one
 			if (Config.Type == UEBuildBinaryType.Executable && Config.bBuildAdditionalConsoleApp)
 			{
-				foreach (string OutputFilePath in Config.OutputFilePaths)
+				foreach (FileReference OutputFilePath in Config.OutputFilePaths)
 				{
 					AddBuildProductAndDebugFile(GetAdditionalConsoleAppPath(OutputFilePath), Type, DebugExtension, Receipt, ToolChain);
 				}
@@ -335,13 +335,13 @@ namespace UnrealBuildTool
 		/// <param name="OutputFile">Build product to add</param>
 		/// <param name="DebugExtension">Extension for the matching debug file (may be null).</param>
 		/// <param name="Receipt">Receipt to add to</param>
-		static void AddBuildProductAndDebugFile(string OutputFile, BuildProductType OutputType, string DebugExtension, TargetReceipt Receipt, IUEToolChain ToolChain)
+		static void AddBuildProductAndDebugFile(FileReference OutputFile, BuildProductType OutputType, string DebugExtension, TargetReceipt Receipt, IUEToolChain ToolChain)
 		{
 			Receipt.AddBuildProduct(OutputFile, OutputType);
 
 			if(!String.IsNullOrEmpty(DebugExtension) && ToolChain.ShouldAddDebugFileToReceipt(OutputFile, OutputType))
 			{
-				Receipt.AddBuildProduct(Path.ChangeExtension(OutputFile, DebugExtension), BuildProductType.SymbolFile);
+				Receipt.AddBuildProduct(OutputFile.ChangeExtension(DebugExtension), BuildProductType.SymbolFile);
 			}
 		}
 
@@ -350,9 +350,9 @@ namespace UnrealBuildTool
 		/// </summary>
 		/// <param name="BinaryPath">Full path to the binary exe.</param>
 		/// <returns></returns>
-		public static string GetAdditionalConsoleAppPath(string BinaryPath)
+		public static FileReference GetAdditionalConsoleAppPath(FileReference BinaryPath)
 		{
-			return Path.Combine(Path.GetDirectoryName(BinaryPath), Path.GetFileNameWithoutExtension(BinaryPath) + "-Cmd" + Path.GetExtension(BinaryPath));
+			return FileReference.Combine(BinaryPath.Directory, BinaryPath.GetFileNameWithoutExtension() + "-Cmd" + BinaryPath.GetExtension());
 		}
 
 		/**
@@ -387,7 +387,7 @@ namespace UnrealBuildTool
 			{
 				foreach (var OutputFilePath in Config.OutputFilePaths)
 				{
-					var OutputDistributionLevel = UEBuildModule.GetModuleDistributionLevelBasedOnLocation(OutputFilePath);
+					var OutputDistributionLevel = UEBuildModule.GetModuleDistributionLevelBasedOnLocation(OutputFilePath.FullName);
 
 					// Throw exception if output path is not appropriate
 					if (OutputDistributionLevel < DistributionLevel)
@@ -572,9 +572,9 @@ namespace UnrealBuildTool
 		/// <param name="DependentLinkEnvironment">The link environment of the dependency</param>
 		public override void SetupDependentLinkEnvironment(LinkEnvironment DependentLinkEnvironment)
 		{
-			foreach (string OutputFilePath in Config.OutputFilePaths)
+			foreach (FileReference OutputFilePath in Config.OutputFilePaths)
 			{
-				string LibraryFileName;
+				FileReference LibraryFileName;
 				if (Config.Type == UEBuildBinaryType.StaticLibrary
 					|| DependentLinkEnvironment.Config.Target.Platform == CPPTargetPlatform.Mac
 					|| DependentLinkEnvironment.Config.Target.Platform == CPPTargetPlatform.Linux)
@@ -583,16 +583,16 @@ namespace UnrealBuildTool
 				}
 				else
 				{
-					LibraryFileName = Path.Combine(Config.IntermediateDirectory, Path.GetFileNameWithoutExtension(OutputFilePath) + ".lib");
+					LibraryFileName = FileReference.Combine(Config.IntermediateDirectory, OutputFilePath.GetFileNameWithoutExtension() + ".lib");
 				}
-				DependentLinkEnvironment.Config.AdditionalLibraries.Add(LibraryFileName);
+				DependentLinkEnvironment.Config.AdditionalLibraries.Add(LibraryFileName.FullName);
 			}
 			
 			// If we're linking against static library containing the launch module on windows, we need to add the compiled resource separately. We can't link it through the static library.
 			if(Config.Type == UEBuildBinaryType.StaticLibrary && ModuleNames.Contains("Launch") && (Target.Platform == UnrealTargetPlatform.Win32 || Target.Platform == UnrealTargetPlatform.Win64))
 			{
-				string ResourceFilePath = Path.Combine(Config.IntermediateDirectory, "Launch", "PCLaunch.rc.res");
-				DependentLinkEnvironment.InputFiles.Add(FileItem.GetItemByPath(ResourceFilePath));
+				FileReference ResourceFileRef = FileReference.Combine(Config.IntermediateDirectory, "Launch", "PCLaunch.rc.res");
+				DependentLinkEnvironment.InputFiles.Add(FileItem.GetItemByFileReference(ResourceFileRef));
 			}
 		}
 
@@ -622,7 +622,7 @@ namespace UnrealBuildTool
 			foreach (var ModuleName in ModuleNames)
 			{
 				UEBuildModule Module = Target.GetModuleByName(ModuleName);
-				if (!Utils.IsFileUnderDirectory(Module.ModuleDirectory, BuildConfiguration.RelativeEnginePath))
+				if (!Module.ModuleDirectory.IsUnderDirectory(UnrealBuildTool.EngineDirectory))
 				{
 					GameModules.Add(Module);
 				}
@@ -650,8 +650,8 @@ namespace UnrealBuildTool
 			// Add the compiled resource file if we're building a static library containing the launch module on Windows
 			if(Config.Type == UEBuildBinaryType.StaticLibrary && ModuleNames.Contains("Launch") && (Target.Platform == UnrealTargetPlatform.Win32 || Target.Platform == UnrealTargetPlatform.Win64))
 			{
-				string ResourceFilePath = Path.Combine(Config.IntermediateDirectory, "Launch", "PCLaunch.rc.res");
-				Receipt.AddBuildProduct(ResourceFilePath, BuildProductType.StaticLibrary);
+				FileReference ResourceFilePath = FileReference.Combine(Config.IntermediateDirectory, "Launch", "PCLaunch.rc.res");
+				Receipt.AddBuildProduct(ResourceFilePath.FullName, BuildProductType.StaticLibrary);
 			}
 			
 			// Add runtime dependencies for all the modules in this binary, and build up a list of all the referenced modules
@@ -685,7 +685,7 @@ namespace UnrealBuildTool
 		/// <returns>Returns the OutputFilePath for this binary</returns>
 		public override string ToString()
 		{
-			return Config.OutputFilePath;
+			return Config.OutputFilePath.FullName;
 		}
 
 		private LinkEnvironment SetupBinaryLinkEnvironment(LinkEnvironment LinkEnvironment, CPPEnvironment CompileEnvironment)
@@ -705,8 +705,8 @@ namespace UnrealBuildTool
 			// @todo: This should be in some Windows code somewhere...
 			// Set the original file name macro; used in PCLaunch.rc to set the binary metadata fields.
 			var OriginalFilename = (Config.OriginalOutputFilePaths != null) ?
-				Path.GetFileName(Config.OriginalOutputFilePaths[0]) :
-				Path.GetFileName(Config.OutputFilePaths[0]);
+				Config.OriginalOutputFilePaths[0].GetFileName() :
+				Config.OutputFilePaths[0].GetFileName();
 			BinaryCompileEnvironment.Config.Definitions.Add("ORIGINAL_FILE_NAME=\"" + OriginalFilename + "\"");
 
 			foreach (var ModuleName in ModuleNames)
@@ -766,7 +766,7 @@ namespace UnrealBuildTool
 			BinaryLinkEnvironment.Config.IntermediateDirectory = Config.IntermediateDirectory;
 
 			// Put the non-executable output files (PDB, import library, etc) in the same directory as the production
-			BinaryLinkEnvironment.Config.OutputDirectory = Path.GetDirectoryName(Config.OutputFilePaths[0]);
+			BinaryLinkEnvironment.Config.OutputDirectory = Config.OutputFilePaths[0].Directory;
 
 			// Setup link output type
 			BinaryLinkEnvironment.Config.bIsBuildingDLL = IsBuildingDll(Config.Type);
@@ -849,12 +849,12 @@ namespace UnrealBuildTool
 			var ModuleName = ModuleNames.First(Name => Name.CompareTo(BuildConfiguration.UCAModuleToAnalyze) == 0);
 			var ModuleCPP = (UEBuildModuleCPP)Target.GetModuleByName(ModuleName);
 			var ModulePrivatePCH = ModuleCPP.ProcessedDependencies.UniquePCHHeaderFile;
-			var IntermediatePath = Path.Combine(Target.ProjectIntermediateDirectory, ModuleName);
+			var IntermediatePath = Path.Combine(Target.ProjectIntermediateDirectory.FullName, ModuleName);
 			var OutputFileName = Target.OutputPath;
-			var OutputFile = FileItem.GetItemByPath(OutputFileName);
+			var OutputFile = FileItem.GetItemByFileReference(OutputFileName);
 
 			Action LinkAction = new Action(ActionType.Compile);
-			LinkAction.WorkingDirectory = Path.GetFullPath(".");
+			LinkAction.WorkingDirectory = UnrealBuildTool.EngineSourceDirectory.FullName;
 			LinkAction.CommandPath = System.IO.Path.Combine(LinkAction.WorkingDirectory, @"..", @"Binaries", @"Win32", @"UnrealCodeAnalyzer.exe");
 			LinkAction.ProducedItems.Add(OutputFile);
 			LinkAction.PrerequisiteItems.AddRange(BinaryLinkEnvironment.InputFiles);
@@ -909,7 +909,7 @@ namespace UnrealBuildTool
 			UEToolChain.GetPlatformToolChain(CPPTargetPlatform.Win64).CompileCSharpProject(
 				ProjectCSharpEnviroment, Config.ProjectFilePath, Config.OutputFilePath);
 
-			return new FileItem[] { FileItem.GetItemByPath(Config.OutputFilePath) };
+			return new FileItem[] { FileItem.GetItemByFileReference(Config.OutputFilePath) };
 		}
 	};
 }
diff --git a/Engine/Source/Programs/UnrealBuildTool/Configuration/UEBuildClient.cs b/Engine/Source/Programs/UnrealBuildTool/Configuration/UEBuildClient.cs
--- a/Engine/Source/Programs/UnrealBuildTool/Configuration/UEBuildClient.cs
+++ b/Engine/Source/Programs/UnrealBuildTool/Configuration/UEBuildClient.cs
@@ -17,7 +17,7 @@ namespace UnrealBuildTool
 		{
 		}
 
-		public UEBuildClient(TargetDescriptor InDesc, TargetRules InRulesObject, RulesAssembly InRulesAssembly, string InTargetCsFilename)
+		public UEBuildClient(TargetDescriptor InDesc, TargetRules InRulesObject, RulesAssembly InRulesAssembly, FileReference InTargetCsFilename)
             : base(InDesc, InRulesObject, InRulesAssembly, "UE4Client", InTargetCsFilename)
         {
             if (ShouldCompileMonolithic())
@@ -28,7 +28,7 @@ namespace UnrealBuildTool
                     // We want the output to go into the <GAME>\Binaries folder
                     if (!InRulesObject.bOutputToEngineBinaries)
                     {
-						OutputPaths = OutputPaths.Select(Path => Path.Replace("Engine\\Binaries", InDesc.TargetName + "\\Binaries")).ToList();
+						OutputPaths = OutputPaths.Select(Path => new FileReference(Path.FullName.Replace("Engine\\Binaries", InDesc.TargetName + "\\Binaries"))).ToList();
                     }
                 }
             }
diff --git a/Engine/Source/Programs/UnrealBuildTool/Configuration/UEBuildEditor.cs b/Engine/Source/Programs/UnrealBuildTool/Configuration/UEBuildEditor.cs
--- a/Engine/Source/Programs/UnrealBuildTool/Configuration/UEBuildEditor.cs
+++ b/Engine/Source/Programs/UnrealBuildTool/Configuration/UEBuildEditor.cs
@@ -16,7 +16,7 @@ namespace UnrealBuildTool
 		{
 		}
 
-		public UEBuildEditor(TargetDescriptor InDesc, TargetRules InRulesObject, RulesAssembly InRulesAssembly, string InTargetCsFilename)
+		public UEBuildEditor(TargetDescriptor InDesc, TargetRules InRulesObject, RulesAssembly InRulesAssembly, FileReference InTargetCsFilename)
 			: base(InDesc, InRulesObject, InRulesAssembly, "UE4Editor", InTargetCsFilename)
 		{
 		}
diff --git a/Engine/Source/Programs/UnrealBuildTool/Configuration/UEBuildGame.cs b/Engine/Source/Programs/UnrealBuildTool/Configuration/UEBuildGame.cs
--- a/Engine/Source/Programs/UnrealBuildTool/Configuration/UEBuildGame.cs
+++ b/Engine/Source/Programs/UnrealBuildTool/Configuration/UEBuildGame.cs
@@ -17,7 +17,7 @@ namespace UnrealBuildTool
 		{
 		}
 
-		public UEBuildGame(TargetDescriptor InDesc, TargetRules InRulesObject, RulesAssembly InRulesAssembly, string InTargetCsFilename)
+		public UEBuildGame(TargetDescriptor InDesc, TargetRules InRulesObject, RulesAssembly InRulesAssembly, FileReference InTargetCsFilename)
 			: base(InDesc, InRulesObject, InRulesAssembly, "UE4", InTargetCsFilename)
 		{
 			if (ShouldCompileMonolithic())
@@ -28,7 +28,7 @@ namespace UnrealBuildTool
 					// We want the output to go into the <GAME>\Binaries folder
 					if (!InRulesObject.bOutputToEngineBinaries)
 					{
-						OutputPaths = OutputPaths.Select(Path => Path.Replace("Engine\\Binaries", InDesc.TargetName + "\\Binaries")).ToList();
+						OutputPaths = OutputPaths.Select(Path => new FileReference(Path.FullName.Replace("Engine\\Binaries", InDesc.TargetName + "\\Binaries"))).ToList();
 					}
 				}
 			}
diff --git a/Engine/Source/Programs/UnrealBuildTool/Configuration/UEBuildModule.cs b/Engine/Source/Programs/UnrealBuildTool/Configuration/UEBuildModule.cs
--- a/Engine/Source/Programs/UnrealBuildTool/Configuration/UEBuildModule.cs
+++ b/Engine/Source/Programs/UnrealBuildTool/Configuration/UEBuildModule.cs
@@ -78,8 +78,9 @@ namespace UnrealBuildTool
 			}
 		}
 
-		public static UEBuildModuleType GetEngineModuleTypeBasedOnLocation(string ModuleName, UEBuildModuleType ModuleType, string ModuleFileRelativeToEngineDirectory)
+		public static UEBuildModuleType GetEngineModuleTypeBasedOnLocation(string ModuleName, UEBuildModuleType ModuleType, FileReference ModuleFileName)
 		{
+			string ModuleFileRelativeToEngineDirectory = ModuleFileName.MakeRelativeTo(UnrealBuildTool.EngineDirectory);
 			if (ModuleFileRelativeToEngineDirectory.IndexOf(UEBuildModule.RuntimeFolder, StringComparison.InvariantCultureIgnoreCase) >= 0)
 			{
 				ModuleType = UEBuildModuleType.Runtime;
@@ -138,7 +139,7 @@ namespace UnrealBuildTool
 		private void SetupModuleDistributionLevel()
 		{
 			List<string> PathsToCheck = new List<string>();
-			PathsToCheck.Add(ModuleDirectory);
+			PathsToCheck.Add(ModuleDirectory.FullName);
 			PathsToCheck.AddRange(PublicIncludePaths);
 			PathsToCheck.AddRange(PrivateIncludePaths);
 			// Not sure if these two are necessary as paths will usually be in basic includes too
@@ -176,7 +177,7 @@ namespace UnrealBuildTool
 		public UEBuildModuleDistribution DistributionLevel;
 
 		/** Path to the module directory */
-		public readonly string ModuleDirectory;
+		public readonly DirectoryReference ModuleDirectory;
 
 		/** Is this module allowed to be redistributed. */
 		private readonly bool? IsRedistributableOverride;
@@ -243,7 +244,7 @@ namespace UnrealBuildTool
 			UEBuildTarget InTarget,
 			string InName,
 			UEBuildModuleType InType,
-			string InModuleDirectory,
+			DirectoryReference InModuleDirectory,
 			ModuleRules InRules,
 			string InBuildCsFilename
 			)
@@ -311,14 +312,14 @@ namespace UnrealBuildTool
 					// of something else, which cause a fall-through to the code below and set up an empty _API define.
 					if(bIncludePathsOnly)
 					{
-						Log.TraceVerbose( "{0}: Include paths only for {1} (no binary)", Path.GetFileNameWithoutExtension( SourceBinary.Config.OutputFilePaths[0] ), Name );
+						Log.TraceVerbose( "{0}: Include paths only for {1} (no binary)", SourceBinary.Config.OutputFilePaths[0].GetFileNameWithoutExtension(), Name );
 						Definitions.Add( Name.ToUpperInvariant() + "_API=" );
 					}
 				}
 				else
 				{
-					string BinaryPath = Binary.Config.OutputFilePaths[0];
-					string SourceBinaryPath = SourceBinary.Config.OutputFilePaths[0];
+					FileReference BinaryPath = Binary.Config.OutputFilePaths[0];
+					FileReference SourceBinaryPath = SourceBinary.Config.OutputFilePaths[0];
 
 					if (ProjectFileGenerator.bGenerateProjectFiles || (Binary.Config.Type == UEBuildBinaryType.StaticLibrary))
 					{
@@ -330,12 +331,12 @@ namespace UnrealBuildTool
 					{
 						if( Binary.Config.bAllowExports )
 						{
-							Log.TraceVerbose( "{0}: Exporting {1} from {2}", Path.GetFileNameWithoutExtension( SourceBinaryPath ), Name, Path.GetFileNameWithoutExtension( BinaryPath ) );
+							Log.TraceVerbose( "{0}: Exporting {1} from {2}", SourceBinaryPath.GetFileNameWithoutExtension(), Name, BinaryPath.GetFileNameWithoutExtension() );
 							Definitions.Add( Name.ToUpperInvariant() + "_API=DLLEXPORT" );
 						}
 						else
 						{
-							Log.TraceVerbose( "{0}: Not importing/exporting {1} (binary: {2})", Path.GetFileNameWithoutExtension( SourceBinaryPath ), Name, Path.GetFileNameWithoutExtension( BinaryPath ) );
+							Log.TraceVerbose( "{0}: Not importing/exporting {1} (binary: {2})", SourceBinaryPath.GetFileNameWithoutExtension(), Name, BinaryPath.GetFileNameWithoutExtension() );
 							Definitions.Add( Name.ToUpperInvariant() + "_API=" );
 						}
 					}
@@ -346,17 +347,17 @@ namespace UnrealBuildTool
 						//		this, we need to suppress warnings at compile time.
 						if( bIncludePathsOnly )
 						{
-							Log.TraceVerbose( "{0}: Include paths only for {1} (binary: {2})", Path.GetFileNameWithoutExtension( SourceBinaryPath ), Name, Path.GetFileNameWithoutExtension( BinaryPath ) );
+							Log.TraceVerbose( "{0}: Include paths only for {1} (binary: {2})", SourceBinaryPath.GetFileNameWithoutExtension(), Name, BinaryPath.GetFileNameWithoutExtension() );
 							Definitions.Add( Name.ToUpperInvariant() + "_API=" );
 						}
 						else if (Binary.Config.bAllowExports)
 						{
-							Log.TraceVerbose( "{0}: Importing {1} from {2}", Path.GetFileNameWithoutExtension( SourceBinaryPath ), Name, Path.GetFileNameWithoutExtension( BinaryPath ) );
+							Log.TraceVerbose( "{0}: Importing {1} from {2}", SourceBinaryPath.GetFileNameWithoutExtension(), Name, BinaryPath.GetFileNameWithoutExtension() );
 							Definitions.Add( Name.ToUpperInvariant() + "_API=DLLIMPORT" );
 						}
 						else
 						{
-							Log.TraceVerbose( "{0}: Not importing/exporting {1} (binary: {2})", Path.GetFileNameWithoutExtension( SourceBinaryPath ), Name, Path.GetFileNameWithoutExtension( BinaryPath ));
+							Log.TraceVerbose( "{0}: Not importing/exporting {1} (binary: {2})", SourceBinaryPath.GetFileNameWithoutExtension(), Name, BinaryPath.GetFileNameWithoutExtension());
 							Definitions.Add( Name.ToUpperInvariant() + "_API=" );
 						}
 					}
@@ -379,7 +380,7 @@ namespace UnrealBuildTool
 				}
 
 				// Add the module's directory to the include path, so we can root #includes to it
-				IncludePaths.Add(Utils.CleanDirectorySeparators(Utils.MakePathRelativeTo(ModuleDirectory, Path.Combine(ProjectFileGenerator.RootRelativePath, "Engine/Source")), '/'));
+				IncludePaths.Add(Utils.CleanDirectorySeparators(ModuleDirectory.MakeRelativeTo(UnrealBuildTool.EngineSourceDirectory), '/'));
 
 				// Add the additional frameworks so that the compiler can know about their #include paths
 				AdditionalFrameworks.AddRange(PublicAdditionalFrameworks);
@@ -655,7 +656,7 @@ namespace UnrealBuildTool
 			UEBuildTarget InTarget,
 			UEBuildModuleType InType,
 			string InName,
-			string InModuleDirectory,
+			DirectoryReference InModuleDirectory,
 			ModuleRules InRules,
 			string InBuildCsFilename
 			)
@@ -776,7 +777,7 @@ namespace UnrealBuildTool
 		public readonly SourceFilesClass SourceFilesFound = new SourceFilesClass();
 
 		/** The directory for this module's generated code */
-		public readonly string GeneratedCodeDirectory;
+		public readonly DirectoryReference GeneratedCodeDirectory;
 
 		/** The preprocessor definitions used to compile this module's private implementation. */
 		HashSet<string> Definitions;
@@ -839,8 +840,8 @@ namespace UnrealBuildTool
 			UEBuildTarget InTarget,
 			string InName,
 			UEBuildModuleType InType,
-			string InModuleDirectory,
-			string InGeneratedCodeDirectory,
+			DirectoryReference InModuleDirectory,
+			DirectoryReference InGeneratedCodeDirectory,
 			IntelliSenseGatherer InIntelliSenseGatherer,
 			IEnumerable<FileItem> InSourceFiles,
 			ModuleRules InRules,
@@ -917,14 +918,14 @@ namespace UnrealBuildTool
 				// Add default PCLaunch.rc file if this module has no own resource file specified
 				if (SourceFilesToBuild.RCFiles.Count <= 0)
 				{
-					string DefRC = Utils.CleanDirectorySeparators( Path.GetFullPath(Path.Combine(Directory.GetCurrentDirectory(), "Runtime/Launch/Resources/Windows/PCLaunch.rc")) );
-					FileItem Item = FileItem.GetItemByFullPath(DefRC);
+					FileReference DefRC = FileReference.Combine(UnrealBuildTool.EngineSourceDirectory, "Runtime", "Launch", "Resources", "Windows", "PCLaunch.rc");
+					FileItem Item = FileItem.GetItemByFileReference(DefRC);
 					SourceFilesToBuild.RCFiles.Add(Item);
 				}
 
 				// Always compile in the API version resource separately. This is required for the module manager to detect compatible API versions.
-				string ModuleVersionRC = Utils.CleanDirectorySeparators( Path.GetFullPath(Path.Combine(Directory.GetCurrentDirectory(), "Runtime/Core/Resources/Windows/ModuleVersionResource.rc.inl")) );
-				FileItem ModuleVersionItem = FileItem.GetItemByFullPath(ModuleVersionRC);
+				FileReference ModuleVersionRC = FileReference.Combine(UnrealBuildTool.EngineSourceDirectory, "Runtime", "Core", "Resources", "Windows", "ModuleVersionResource.rc.inl");
+				FileItem ModuleVersionItem = FileItem.GetItemByFileReference(ModuleVersionRC);
 				if( !SourceFilesToBuild.RCFiles.Contains(ModuleVersionItem) )
 				{
 					SourceFilesToBuild.RCFiles.Add(ModuleVersionItem);
@@ -945,8 +946,8 @@ namespace UnrealBuildTool
 
 
 			// Check to see if this is an Engine module (including program or plugin modules).  That is, the module is located under the "Engine" folder
-			var IsPluginModule = Utils.IsFileUnderDirectory( this.ModuleDirectory, Path.Combine(Target.ProjectDirectory, "Plugins" ) );
-			var IsGameModule = !IsPluginModule && !Utils.IsFileUnderDirectory(this.ModuleDirectory, Path.Combine(ProjectFileGenerator.EngineRelativePath));
+			var IsPluginModule = ModuleDirectory.IsUnderDirectory( DirectoryReference.Combine(Target.ProjectDirectory, "Plugins") );
+			var IsGameModule = !IsPluginModule && !ModuleDirectory.IsUnderDirectory(UnrealBuildTool.EngineDirectory);
 
 			// Should we force a precompiled header to be generated for this module?  Usually, we only bother with a
 			// precompiled header if there are at least several source files in the module (after combining them for unity
@@ -1038,17 +1039,17 @@ namespace UnrealBuildTool
 
 				if( bUseSharedPCHFiles )
 				{
-					string SharingPCHHeaderFilePath = null;
+					FileReference SharingPCHHeaderFilePath = null;
 					bool bIsASharedPCHModule = bUseSharedPCHFiles && GlobalCompileEnvironment.SharedPCHHeaderFiles.Any( PCH => PCH.Module == this );
 					if( bIsASharedPCHModule )
 					{
-						SharingPCHHeaderFilePath = Path.GetFullPath( Path.Combine( ProjectFileGenerator.RootRelativePath, "Engine", "Source", Rules.SharedPCHHeaderFile ) );
+						SharingPCHHeaderFilePath = FileReference.Combine( UnrealBuildTool.EngineSourceDirectory, Rules.SharedPCHHeaderFile );
 					}
 
 					// We can't use a shared PCH file when compiling a module
 					// with exports, because the shared PCH can only have imports in it to work correctly.
 					bool bAllowSharedPCH = (Rules.PCHUsage == ModuleRules.PCHUsageMode.NoSharedPCHs)? false : true;
-					bool bCanModuleUseOwnSharedPCH = bAllowSharedPCH && bIsASharedPCHModule && !Binary.Config.bAllowExports && ProcessedDependencies.UniquePCHHeaderFile.AbsolutePath.Equals( SharingPCHHeaderFilePath, StringComparison.InvariantCultureIgnoreCase );
+					bool bCanModuleUseOwnSharedPCH = bAllowSharedPCH && bIsASharedPCHModule && !Binary.Config.bAllowExports && ProcessedDependencies.UniquePCHHeaderFile.Reference == SharingPCHHeaderFilePath;
 					if( bAllowSharedPCH && ( !bIsASharedPCHModule || bCanModuleUseOwnSharedPCH ) )
 					{
 						// Figure out which shared PCH tier we're in
@@ -1134,7 +1135,7 @@ namespace UnrealBuildTool
 						foreach (var CPPFile in SourceFilesToBuild.CPPFiles)
 						{
 							CPPFile.PCHHeaderNameInCode              = PCHToUse.AbsolutePath;
-							CPPFile.PrecompiledHeaderIncludeFilename = PCHToUse.AbsolutePath;
+							CPPFile.PrecompiledHeaderIncludeFilename = PCHToUse.Reference;
 						}
 					}
 
@@ -1177,7 +1178,7 @@ namespace UnrealBuildTool
 
 					var ModulePCHCompileEnvironment = ModuleCompileEnvironment.DeepCopy();
 					ModulePCHCompileEnvironment.Config.PrecompiledHeaderAction          = PrecompiledHeaderAction.Include;
-					ModulePCHCompileEnvironment.Config.PrecompiledHeaderIncludeFilename = ModulePCHEnvironment.PrecompiledHeaderIncludeFilename.AbsolutePath;
+					ModulePCHCompileEnvironment.Config.PrecompiledHeaderIncludeFilename = ModulePCHEnvironment.PrecompiledHeaderIncludeFilename.Reference;
 					ModulePCHCompileEnvironment.Config.PCHHeaderNameInCode              = ModulePCHEnvironment.PCHHeaderNameInCode;
 
 					if( SharedPCHHeaderFile != null )
@@ -1232,7 +1233,7 @@ namespace UnrealBuildTool
 									CPPFilesToBuild[0].PCHHeaderNameInCode,
 									ModulePCHEnvironment.PrecompiledHeaderIncludeFilename,
 									SharedPCHCompileEnvironment,
-									Path.Combine( CompileEnvironment.Config.OutputDirectory, "SharedPCHs" ),
+									DirectoryReference.Combine( CompileEnvironment.Config.OutputDirectory, "SharedPCHs" ),
 									"Shared", 
 									false );
 							}
@@ -1435,7 +1436,7 @@ namespace UnrealBuildTool
 				if( bFoundAProblemWithPCHs )
 				{
 					// Map from pch header string to the source files that use that PCH
-					var UsageMapPCH = new Dictionary<string, List<FileItem>>( StringComparer.InvariantCultureIgnoreCase );
+					var UsageMapPCH = new Dictionary<FileReference, List<FileItem>>();
 					foreach( var CPPFile in SourceFilesToBuild.CPPFiles )
 					{
 						// Create a new entry if not in the pch usage map
@@ -1460,7 +1461,7 @@ namespace UnrealBuildTool
 					if( UsageMapPCH.Count > 1 )
 					{
 						// Keep track of the PCH file that is most used within this module
-						string MostFilesAreIncludingPCH = string.Empty;
+						FileReference MostFilesAreIncludingPCH = null;
 						int MostFilesIncluded = 0;
 						foreach( var CurPCH in UsageMapPCH.Where( PCH => PCH.Value.Count > MostFilesIncluded ) )
 						{
@@ -1520,14 +1521,14 @@ namespace UnrealBuildTool
 
 			// Resolve the PCH header to an absolute path.
 			// Check NullOrEmpty here because if the file could not be resolved we need to throw an exception
-			if (!string.IsNullOrEmpty(FirstInclude.IncludeResolvedName) &&
+			if (FirstInclude.IncludeResolvedNameIfSuccessful != null &&
 				// ignore any preexisting resolve cache if we are not configured to use it.
 				BuildConfiguration.bUseIncludeDependencyResolveCache &&
 				// if we are testing the resolve cache, we force UBT to resolve every time to look for conflicts
 				!BuildConfiguration.bTestIncludeDependencyResolveCache)
 			{
-				CPPFile.PrecompiledHeaderIncludeFilename = FirstInclude.IncludeResolvedName;
-				return FileItem.GetItemByFullPath(CPPFile.PrecompiledHeaderIncludeFilename);
+				CPPFile.PrecompiledHeaderIncludeFilename = FirstInclude.IncludeResolvedNameIfSuccessful;
+				return FileItem.GetItemByFileReference(CPPFile.PrecompiledHeaderIncludeFilename);
 			}
 
 			// search the include paths to resolve the file.
@@ -1537,8 +1538,8 @@ namespace UnrealBuildTool
 				throw new BuildException("The first include statement in source file '{0}' is trying to include the file '{1}' as the precompiled header, but that file could not be located in any of the module's include search paths.", CPPFile.AbsolutePath, CPPFile.PCHHeaderNameInCode);
 			}
 
-			CPPEnvironment.IncludeDependencyCache[Target].CacheResolvedIncludeFullPath(CPPFile, 0, PrecompiledHeaderIncludeFile.AbsolutePath);
-			CPPFile.PrecompiledHeaderIncludeFilename = PrecompiledHeaderIncludeFile.AbsolutePath;
+			CPPEnvironment.IncludeDependencyCache[Target].CacheResolvedIncludeFullPath(CPPFile, 0, PrecompiledHeaderIncludeFile.Reference);
+			CPPFile.PrecompiledHeaderIncludeFilename = PrecompiledHeaderIncludeFile.Reference;
 
 			return PrecompiledHeaderIncludeFile;
 		}
@@ -1564,12 +1565,12 @@ namespace UnrealBuildTool
 			Result.Config.bEnableExceptions                      = Rules.bEnableExceptions;
 			Result.Config.bEnableShadowVariableWarning           = Rules.bEnableShadowVariableWarnings;
 			Result.Config.bUseStaticCRT							 = (Target.Rules != null && Target.Rules.bUseStaticCRT);
-			Result.Config.OutputDirectory                        = Path.Combine(Binary.Config.IntermediateDirectory, Name);
+			Result.Config.OutputDirectory                        = DirectoryReference.Combine(Binary.Config.IntermediateDirectory, Name);
 
 			// Switch the optimization flag if we're building a game module. Also pass the definition for building in DebugGame along (see ModuleManager.h for notes).
 			if (Target.Configuration == UnrealTargetConfiguration.DebugGame)
 			{
-				if(!Utils.IsFileUnderDirectory(ModuleDirectory, BuildConfiguration.RelativeEnginePath))
+				if(!ModuleDirectory.IsUnderDirectory(UnrealBuildTool.EngineDirectory))
 				{
 					Result.Config.Target.Configuration = CPPTargetConfiguration.Debug;
 					Result.Config.Definitions.Add("UE_BUILD_DEVELOPMENT_WITH_DEBUGGAME=1");
@@ -1621,7 +1622,7 @@ namespace UnrealBuildTool
 		{
 			if (UHTModuleInfoCache == null)
 			{
-				IEnumerable<string> HeaderFilenames = Directory.GetFiles( ModuleDirectory, "*.h", SearchOption.AllDirectories );
+				IEnumerable<string> HeaderFilenames = Directory.GetFiles( ModuleDirectory.FullName, "*.h", SearchOption.AllDirectories );
 				UHTModuleInfo       Info            = ExternalExecution.CreateUHTModuleInfo(HeaderFilenames, Target, Name, ModuleDirectory, Type);
 				UHTModuleInfoCache = new UHTModuleInfoCacheType(Info.PublicUObjectHeaders.Concat(Info.PublicUObjectClassesHeaders).Concat(Info.PrivateUObjectHeaders).Select(x => x.AbsolutePath).ToList(), Info);
 			}
@@ -1759,8 +1760,8 @@ namespace UnrealBuildTool
 			UEBuildTarget InTarget,
 			string InName,
 			UEBuildModuleType InType,
-			string InModuleDirectory,
-			string InGeneratedCodeDirectory,
+			DirectoryReference InModuleDirectory,
+			DirectoryReference InGeneratedCodeDirectory,
 			IntelliSenseGatherer InIntelliSenseGatherer,
 			IEnumerable<FileItem> InSourceFiles,
 			ModuleRules InRules,
@@ -1890,7 +1891,7 @@ namespace UnrealBuildTool
 		/// <param name="ModuleName">Name of the module this PCH is being generated for</param>
 		/// <param name="bAllowDLLExports">True if we should allow DLLEXPORT definitions for this PCH</param>
 		/// <returns>the compilation output result of the created pch.</returns>
-		public static CPPOutput GeneratePCHCreationAction(UEBuildTarget Target, string PCHHeaderNameInCode, FileItem PrecompiledHeaderIncludeFilename, CPPEnvironment ProjectCPPEnvironment, string OutputDirectory, string ModuleName, bool bAllowDLLExports )
+		public static CPPOutput GeneratePCHCreationAction(UEBuildTarget Target, string PCHHeaderNameInCode, FileItem PrecompiledHeaderIncludeFilename, CPPEnvironment ProjectCPPEnvironment, DirectoryReference OutputDirectory, string ModuleName, bool bAllowDLLExports )
 		{
 			// Find the header file to be precompiled. Don't skip external headers
 			if (PrecompiledHeaderIncludeFilename.bExists)
@@ -1900,7 +1901,7 @@ namespace UnrealBuildTool
 				string PCHGuardDefine = Path.GetFileNameWithoutExtension(PrecompiledHeaderIncludeFilename.AbsolutePath).ToUpper();
 				string LocalPCHHeaderNameInCode = ToolChain.ConvertPath(PrecompiledHeaderIncludeFilename.AbsolutePath);
 				string TmpPCHHeaderContents = String.Format("#ifndef __AUTO_{0}_H__\n#define __AUTO_{0}_H__\n//Last Write: {2}\n#include \"{1}\"\n#endif//__AUTO_{0}_H__", PCHGuardDefine, LocalPCHHeaderNameInCode, PrecompiledHeaderIncludeFilename.LastWriteTime);
-				string DummyPath = Path.Combine(
+				FileReference DummyPath = FileReference.Combine(
 					ProjectCPPEnvironment.Config.OutputDirectory,
 					Path.GetFileName(PrecompiledHeaderIncludeFilename.AbsolutePath));
 				FileItem DummyPCH = FileItem.CreateIntermediateTextFile(DummyPath, TmpPCHHeaderContents);
@@ -1908,7 +1909,7 @@ namespace UnrealBuildTool
 				// Create a new C++ environment that is used to create the PCH.
 				var ProjectPCHEnvironment = ProjectCPPEnvironment.DeepCopy();
 				ProjectPCHEnvironment.Config.PrecompiledHeaderAction = PrecompiledHeaderAction.Create;
-				ProjectPCHEnvironment.Config.PrecompiledHeaderIncludeFilename = PrecompiledHeaderIncludeFilename.AbsolutePath;
+				ProjectPCHEnvironment.Config.PrecompiledHeaderIncludeFilename = PrecompiledHeaderIncludeFilename.Reference;
 				ProjectPCHEnvironment.Config.PCHHeaderNameInCode = PCHHeaderNameInCode;
 				ProjectPCHEnvironment.Config.OutputDirectory = OutputDirectory;
 
diff --git a/Engine/Source/Programs/UnrealBuildTool/Configuration/UEBuildPlatform.cs b/Engine/Source/Programs/UnrealBuildTool/Configuration/UEBuildPlatform.cs
--- a/Engine/Source/Programs/UnrealBuildTool/Configuration/UEBuildPlatform.cs
+++ b/Engine/Source/Programs/UnrealBuildTool/Configuration/UEBuildPlatform.cs
@@ -181,7 +181,7 @@ namespace UnrealBuildTool
 		 * For platforms that need to output multiple files per binary (ie Android "fat" binaries)
 		 * this will emit multiple paths. By default, it simply makes an array from the input
 		 */
-		List<string> FinalizeBinaryPaths(string BinaryName);
+		List<FileReference> FinalizeBinaryPaths(FileReference BinaryName);
 
 		/**
 		 * Setup the configuration environment for building
@@ -275,7 +275,7 @@ namespace UnrealBuildTool
 		 *
 		 * return true if the project uses the default build config
 		 */
-		bool HasDefaultBuildConfig(UnrealTargetPlatform InPlatform, string InProjectPath);
+		bool HasDefaultBuildConfig(UnrealTargetPlatform InPlatform, DirectoryReference InProjectPath);
 	}
 
 	public abstract partial class UEBuildPlatform : IUEBuildPlatform
@@ -664,7 +664,7 @@ namespace UnrealBuildTool
 		/**
 		 * Allow the platform to override the NMake output name
 		 */
-		public virtual string ModifyNMakeOutput(string ExeName)
+		public virtual FileReference ModifyNMakeOutput(FileReference ExeName)
 		{
 			// by default, use original
 			return ExeName;
@@ -687,9 +687,9 @@ namespace UnrealBuildTool
 		 * For platforms that need to output multiple files per binary (ie Android "fat" binaries)
 		 * this will emit multiple paths. By default, it simply makes an array from the input
 		 */
-		public virtual List<string> FinalizeBinaryPaths(string BinaryName)
+		public virtual List<FileReference> FinalizeBinaryPaths(FileReference BinaryName)
 		{
-			List<string> TempList = new List<string>() { BinaryName };
+			List<FileReference> TempList = new List<FileReference>() { BinaryName };
 			return TempList;
 		}
 
@@ -946,10 +946,10 @@ namespace UnrealBuildTool
 		{
 		}
 
-		protected static bool DoProjectSettingsMatchDefault(UnrealTargetPlatform Platform, string ProjectPath, string Section, string[] BoolKeys, string[] IntKeys, string[] StringKeys)
+		protected static bool DoProjectSettingsMatchDefault(UnrealTargetPlatform Platform, DirectoryReference ProjectDirectoryName, string Section, string[] BoolKeys, string[] IntKeys, string[] StringKeys)
 		{
-			ConfigCacheIni ProjIni = new ConfigCacheIni(Platform, "Engine", ProjectPath);
-			ConfigCacheIni DefaultIni = new ConfigCacheIni(Platform, "Engine", null);
+			ConfigCacheIni ProjIni = new ConfigCacheIni(Platform, "Engine", ProjectDirectoryName);
+			ConfigCacheIni DefaultIni = new ConfigCacheIni(Platform, "Engine", (DirectoryReference)null);
 
 			// look at all bool values
 			if (BoolKeys != null) foreach (string Key in BoolKeys)
@@ -999,7 +999,7 @@ namespace UnrealBuildTool
 		 *
 		 * return true if the project uses the default build config
 		 */
-		public virtual bool HasDefaultBuildConfig(UnrealTargetPlatform Platform, string ProjectPath)
+		public virtual bool HasDefaultBuildConfig(UnrealTargetPlatform Platform, DirectoryReference ProjectDirectoryName)
 		{
 			string[] BoolKeys = new string[] {
 				"bCompileApex", "bCompileBox2D", "bCompileICU", "bCompileSimplygon", 
@@ -1008,7 +1008,7 @@ namespace UnrealBuildTool
 				"bCompileForSize", "bCompileCEF3"
 			};
 
-			return DoProjectSettingsMatchDefault(Platform, ProjectPath, "/Script/BuildSettings.BuildSettings",
+			return DoProjectSettingsMatchDefault(Platform, ProjectDirectoryName, "/Script/BuildSettings.BuildSettings",
 				BoolKeys, null, null);
 		}
 	}
diff --git a/Engine/Source/Programs/UnrealBuildTool/Configuration/UEBuildServer.cs b/Engine/Source/Programs/UnrealBuildTool/Configuration/UEBuildServer.cs
--- a/Engine/Source/Programs/UnrealBuildTool/Configuration/UEBuildServer.cs
+++ b/Engine/Source/Programs/UnrealBuildTool/Configuration/UEBuildServer.cs
@@ -16,7 +16,7 @@ namespace UnrealBuildTool
 		{
 		}
 
-		public UEBuildServer(TargetDescriptor InDesc, TargetRules InRulesObject, RulesAssembly InRulesAssembly, string InTargetCsFilename)
+		public UEBuildServer(TargetDescriptor InDesc, TargetRules InRulesObject, RulesAssembly InRulesAssembly, FileReference InTargetCsFilename)
 			: base(InDesc, InRulesObject, InRulesAssembly, "UE4Server", InTargetCsFilename)
 		{
 		}
