Use FileReference and DirectoryReference classes instead of strings for paths. Massively reduces the amount of paranoid calls to GetFullPath(), CleanDirectorySeparators() et al., and improves performa

diff --git a/Engine/Source/Programs/UnrealBuildTool/Configuration/UEBuildTarget.cs b/Engine/Source/Programs/UnrealBuildTool/Configuration/UEBuildTarget.cs
--- a/Engine/Source/Programs/UnrealBuildTool/Configuration/UEBuildTarget.cs
+++ b/Engine/Source/Programs/UnrealBuildTool/Configuration/UEBuildTarget.cs
@@ -184,7 +184,7 @@ namespace UnrealBuildTool
 		public List<OnlyModule> OnlyModules;
 		public bool bPrecompile;
 		public bool bUsePrecompiled;
-		public List<string> ForeignPlugins;
+		public List<FileReference> ForeignPlugins;
 		public string ForceReceiptFileName;
 	}
 
@@ -236,7 +236,7 @@ namespace UnrealBuildTool
 			var Configuration = UnrealTargetConfiguration.Unknown;
 			string RemoteRoot = null;
 			var OnlyModules = new List<OnlyModule>();
-			List<string> ForeignPlugins = new List<string>();
+			List<FileReference> ForeignPlugins = new List<FileReference>();
 			string ForceReceiptFileName = null;
 
 			// If true, the recompile was launched by the editor.
@@ -297,7 +297,7 @@ namespace UnrealBuildTool
 									throw new BuildException("Expected plugin filename after -Plugin argument, but found nothing.");
 								}
 
-								ForeignPlugins.Add(Arguments[++ArgumentIndex]);
+								ForeignPlugins.Add(new FileReference(Arguments[++ArgumentIndex]));
 							}
 							break;
 
@@ -484,20 +484,13 @@ namespace UnrealBuildTool
 					TargetName = PossibleTargetName;
 
 					// If a project file was not specified see if we can find one
-					string CheckProjectFile = UProjectInfo.GetProjectForTarget(TargetName);
-					if (string.IsNullOrEmpty(CheckProjectFile) == false)
+					FileReference CheckProjectFile;
+					if(UProjectInfo.TryGetProjectForTarget(TargetName, out CheckProjectFile))
 					{
 						Log.TraceVerbose("Found project file for {0} - {1}", TargetName, CheckProjectFile);
 						if (UnrealBuildTool.HasUProjectFile() == false)
 						{
-							string NewProjectFilename = CheckProjectFile;
-							if (Path.IsPathRooted(NewProjectFilename) == false)
-							{
-								NewProjectFilename = Path.GetFullPath(NewProjectFilename);
-							}
-
- 							NewProjectFilename = NewProjectFilename.Replace("\\", "/");
-							UnrealBuildTool.SetProjectFile(NewProjectFilename);
+							UnrealBuildTool.SetProjectFile(CheckProjectFile);
 						}
 					}
 
@@ -555,13 +548,13 @@ namespace UnrealBuildTool
 			}
 			if(Desc.ForeignPlugins != null)
 			{
-				foreach(string ForeignPlugin in Desc.ForeignPlugins)
+				foreach(FileReference ForeignPlugin in Desc.ForeignPlugins)
 				{
 					RulesAssembly = RulesCompiler.CreatePluginRulesAssembly(ForeignPlugin, RulesAssembly);
 				}
 			}
 
-			string TargetFileName;
+			FileReference TargetFileName;
 			TargetRules RulesObject = RulesAssembly.CreateTargetRules(Desc.TargetName, new TargetInfo(Desc.Platform, Desc.Configuration), Desc.bIsEditorRecompile, out TargetFileName);
 			if (Desc.bIsEditorRecompile)
 			{
@@ -713,9 +706,9 @@ namespace UnrealBuildTool
 		 * Look for all folders with a uproject file, these are valid games
 		 * This is defined as a valid game
 		 */
-		public static List<string> DiscoverAllGameFolders()
+		public static List<DirectoryReference> DiscoverAllGameFolders()
 		{
-			List<string> AllGameFolders = new List<string>();
+			List<DirectoryReference> AllGameFolders = new List<DirectoryReference>();
 
 			// Add all the normal game folders. The UProjectInfo list is already filtered for projects specified on the command line.
 			List<UProjectInfo> GameProjects = UProjectInfo.FilterGameProjects(true, null);
@@ -757,19 +750,19 @@ namespace UnrealBuildTool
 		public TargetInfo TargetInfo;
 
 		/** Root directory for the active project. Typically contains the .uproject file, or the engine root. */
-		public string ProjectDirectory;
+		public DirectoryReference ProjectDirectory;
 
 		/** Default directory for intermediate files. Typically underneath ProjectDirectory. */
-		public string ProjectIntermediateDirectory;
+		public DirectoryReference ProjectIntermediateDirectory;
 
 		/** Directory for engine intermediates. For an agnostic editor/game executable, this will be under the engine directory. For monolithic executables this will be the same as the project intermediate directory. */
-		public string EngineIntermediateDirectory;
+		public DirectoryReference EngineIntermediateDirectory;
 
 		/** Output paths of final executable. */
-		public List<string> OutputPaths;
+		public List<FileReference> OutputPaths;
 
 		/** Returns the OutputPath is there is only one entry in OutputPaths */
-		public string OutputPath
+		public FileReference OutputPath
 		{
 			get
 			{
@@ -812,7 +805,7 @@ namespace UnrealBuildTool
 
 		/** Additional plugin filenames which are foreign to this target */
 		[NonSerialized]
-		public List<string> ForeignPlugins;
+		public List<FileReference> ForeignPlugins;
 
 		/** All application binaries; may include binaries not built by this target. */
 		[NonSerialized]
@@ -849,7 +842,7 @@ namespace UnrealBuildTool
 			public FlatModuleCsDataType(SerializationInfo Info, StreamingContext Context)
 			{
 				BuildCsFilename    = Info.GetString("bf");
-				ModuleSourceFolder = Info.GetString("mf");
+				ModuleSourceFolder = (DirectoryReference)Info.GetValue("mf", typeof(DirectoryReference));
 				UHTHeaderNames     = (List<string>)Info.GetValue("hn", typeof(List<string>));
 			}
 
@@ -866,7 +859,7 @@ namespace UnrealBuildTool
 			}
 
 			public string       BuildCsFilename;
-			public string       ModuleSourceFolder;
+			public DirectoryReference ModuleSourceFolder;
 			public List<string> UHTHeaderNames = new List<string>();
 		}
 
@@ -880,15 +873,15 @@ namespace UnrealBuildTool
 		private string ReceiptFileName;
 
 		/** Version manifests to be written to each output folder */
-		private KeyValuePair<string, VersionManifest>[] FileNameToVersionManifestPairs;
+		private KeyValuePair<FileReference, VersionManifest>[] FileReferenceToVersionManifestPairs;
 
 		/** Force output of the receipt to an additional filename */
 		[NonSerialized]
 		private string ForceReceiptFileName;
 
 		/** The name of the .Target.cs file, if the target was created with one */
-		private readonly string TargetCsFilenameField;
-		public string TargetCsFilename { get { return TargetCsFilenameField; } }
+		private readonly FileReference TargetCsFilenameField;
+		public FileReference TargetCsFilename { get { return TargetCsFilenameField; } }
 
 		/// <summary>
 		/// A list of the module filenames which were used to build this target.
@@ -926,10 +919,10 @@ namespace UnrealBuildTool
 			Platform                     = (UnrealTargetPlatform)Info.GetInt32("pl");
 			Configuration                = (UnrealTargetConfiguration)Info.GetInt32("co");
 			TargetInfo                   = (TargetInfo)Info.GetValue("ti", typeof(TargetInfo));
-			ProjectDirectory             = Info.GetString("pd");
-			ProjectIntermediateDirectory = Info.GetString("pi");
-			EngineIntermediateDirectory  = Info.GetString("ed");
-			OutputPaths                  = (List<string>)Info.GetValue("op", typeof(List<string>));
+			ProjectDirectory             = (DirectoryReference)Info.GetValue("pd", typeof(DirectoryReference));
+			ProjectIntermediateDirectory = (DirectoryReference)Info.GetValue("pi", typeof(DirectoryReference));
+			EngineIntermediateDirectory  = (DirectoryReference)Info.GetValue("ed", typeof(DirectoryReference));
+			OutputPaths                  = (List<FileReference>)Info.GetValue("op", typeof(List<FileReference>));
 			RemoteRoot                   = Info.GetString("rr");
 			bPrecompile                  = Info.GetBoolean("pc");
 			bUsePrecompiled              = Info.GetBoolean("up");
@@ -944,8 +937,8 @@ namespace UnrealBuildTool
 			}
 			Receipt                      = (TargetReceipt)Info.GetValue("re", typeof(TargetReceipt));
 			ReceiptFileName              = Info.GetString("rf");
-			FileNameToVersionManifestPairs = (KeyValuePair<string, VersionManifest>[])Info.GetValue("vm", typeof(KeyValuePair<string, VersionManifest>[]));
-			TargetCsFilenameField        = Info.GetString("tc");
+			FileReferenceToVersionManifestPairs = (KeyValuePair<FileReference, VersionManifest>[])Info.GetValue("vm", typeof(KeyValuePair<FileReference, VersionManifest>[]));
+			TargetCsFilenameField        = (FileReference)Info.GetValue("tc", typeof(FileReference));
 		}
 
 		public void GetObjectData(SerializationInfo Info, StreamingContext Context)
@@ -971,7 +964,7 @@ namespace UnrealBuildTool
 			Info.AddValue("fv", FlatModuleCsData.Values.ToArray());
 			Info.AddValue("re", Receipt);
 			Info.AddValue("rf", ReceiptFileName);
-			Info.AddValue("vm", FileNameToVersionManifestPairs);
+			Info.AddValue("vm", FileReferenceToVersionManifestPairs);
 			Info.AddValue("tc", TargetCsFilenameField);
 		}
 
@@ -982,7 +975,7 @@ namespace UnrealBuildTool
 		/// <param name="InRules">The target rules, as created by RulesCompiler.</param>
 		/// <param name="InPossibleAppName">The AppName for shared binaries of this target type, if used (null if there is none).</param>
 		/// <param name="InTargetCsFilename">The name of the target </param>
-		public UEBuildTarget(TargetDescriptor InDesc, TargetRules InRules, RulesAssembly InRulesAssembly, string InPossibleAppName, string InTargetCsFilename)
+		public UEBuildTarget(TargetDescriptor InDesc, TargetRules InRules, RulesAssembly InRulesAssembly, string InPossibleAppName, FileReference InTargetCsFilename)
 		{
 			AppName = InDesc.TargetName;
 			TargetName = InDesc.TargetName;
@@ -997,7 +990,7 @@ namespace UnrealBuildTool
 			ForeignPlugins = InDesc.ForeignPlugins;
 			ForceReceiptFileName = InDesc.ForceReceiptFileName;
 
-			Debug.Assert(InTargetCsFilename == null || InTargetCsFilename.EndsWith(".Target.cs", StringComparison.InvariantCultureIgnoreCase));
+			Debug.Assert(InTargetCsFilename == null || InTargetCsFilename.HasExtension(".Target.cs"));
 			TargetCsFilenameField = InTargetCsFilename;
 
 			{
@@ -1029,15 +1022,15 @@ namespace UnrealBuildTool
 			// Figure out what the project directory is. If we have a uproject file, use that. Otherwise use the engine directory.
 			if (UnrealBuildTool.HasUProjectFile())
 			{
-				ProjectDirectory = Path.GetFullPath(UnrealBuildTool.GetUProjectPath());
+				ProjectDirectory = UnrealBuildTool.GetUProjectPath();
 			}
 			else
 			{
-				ProjectDirectory = Path.GetFullPath(BuildConfiguration.RelativeEnginePath);
+				ProjectDirectory = UnrealBuildTool.EngineDirectory;
 			}
 
 			// Build the project intermediate directory
-			ProjectIntermediateDirectory = Path.GetFullPath(Path.Combine(ProjectDirectory, BuildConfiguration.PlatformIntermediateFolder, GetTargetName(), Configuration.ToString()));
+			ProjectIntermediateDirectory = DirectoryReference.Combine(ProjectDirectory, BuildConfiguration.PlatformIntermediateFolder, GetTargetName(), Configuration.ToString());
 
 			// Build the engine intermediate directory. If we're building agnostic engine binaries, we can use the engine intermediates folder. Otherwise we need to use the project intermediates directory.
 			if (!bUseSharedBuildEnvironment)
@@ -1046,29 +1039,29 @@ namespace UnrealBuildTool
 			}
 			else if(Configuration == UnrealTargetConfiguration.DebugGame)
 			{
-				EngineIntermediateDirectory = Path.GetFullPath(Path.Combine(BuildConfiguration.RelativeEnginePath, BuildConfiguration.PlatformIntermediateFolder, AppName, UnrealTargetConfiguration.Development.ToString()));
+				EngineIntermediateDirectory = DirectoryReference.Combine(UnrealBuildTool.EngineDirectory, BuildConfiguration.PlatformIntermediateFolder, AppName, UnrealTargetConfiguration.Development.ToString());
 			}
 			else
 			{
-				EngineIntermediateDirectory = Path.GetFullPath(Path.Combine(BuildConfiguration.RelativeEnginePath, BuildConfiguration.PlatformIntermediateFolder, AppName, Configuration.ToString()));
+				EngineIntermediateDirectory = DirectoryReference.Combine(UnrealBuildTool.EngineDirectory, BuildConfiguration.PlatformIntermediateFolder, AppName, Configuration.ToString());
 			}
 
 			// Get the receipt path for this target
-			ReceiptFileName = TargetReceipt.GetDefaultPath(ProjectDirectory, TargetName, Platform, Configuration, UEBuildPlatform.GetBuildPlatform(Platform).GetActiveArchitecture());
+			ReceiptFileName = TargetReceipt.GetDefaultPath(ProjectDirectory.FullName, TargetName, Platform, Configuration, UEBuildPlatform.GetBuildPlatform(Platform).GetActiveArchitecture());
 
 			RemoteRoot = InDesc.RemoteRoot;
 
 			OnlyModules = InDesc.OnlyModules;
 
 			// Construct the output paths for this target's executable
-			string OutputDirectory;
+			DirectoryReference OutputDirectory;
 			if((bCompileMonolithic || TargetType == TargetRules.TargetType.Program) && !Rules.bOutputToEngineBinaries)
 			{
 				OutputDirectory = ProjectDirectory;
 			}
 			else
 			{
-				OutputDirectory = Path.GetFullPath(BuildConfiguration.RelativeEnginePath);
+				OutputDirectory = UnrealBuildTool.EngineDirectory;
 			}
 			OutputPaths = MakeExecutablePaths(OutputDirectory, bCompileMonolithic? TargetName : AppName, Platform, Configuration, Rules.UndecoratedConfiguration, bCompileMonolithic && UnrealBuildTool.HasUProjectFile(), Rules.ExeBinariesSubFolder);
 
@@ -1206,7 +1199,7 @@ namespace UnrealBuildTool
 			{
 				ReceiptWithFullPaths = new TargetReceipt(Receipt);
 			}
-			ReceiptWithFullPaths.ExpandPathVariables(BuildConfiguration.RelativeEnginePath, ProjectDirectory);
+			ReceiptWithFullPaths.ExpandPathVariables(UnrealBuildTool.EngineDirectory, ProjectDirectory);
 
 			{
 				Log.TraceVerbose("Cleaning target {0} - AppName {1}", TargetName, AppName);
@@ -1268,11 +1261,11 @@ namespace UnrealBuildTool
 				// Delete the intermediate folder for each binary. This will catch all plugin intermediate folders, as well as any project and engine folders.
 				foreach(UEBuildBinary Binary in AppBinaries)
 				{
-					if(Directory.Exists(Binary.Config.IntermediateDirectory))
+					if(Binary.Config.IntermediateDirectory.Exists())
 					{
-						if (!UnrealBuildTool.RunningRocket() || !Utils.IsFileUnderDirectory(Binary.Config.IntermediateDirectory, BuildConfiguration.RelativeEnginePath))
+						if (!UnrealBuildTool.RunningRocket() || !Binary.Config.IntermediateDirectory.IsUnderDirectory(UnrealBuildTool.EngineDirectory))
 						{
-							CleanDirectory(Binary.Config.IntermediateDirectory);
+							CleanDirectory(Binary.Config.IntermediateDirectory.FullName);
 						}
 					}
 				}
@@ -1313,12 +1306,12 @@ namespace UnrealBuildTool
 				foreach (var ModuleName in ModuleList)
 				{
 					UEBuildModuleCPP Module = GetModuleByName(ModuleName) as UEBuildModuleCPP;
-					if (Module != null && Module.GeneratedCodeDirectory != null && Directory.Exists(Module.GeneratedCodeDirectory))
+					if (Module != null && Module.GeneratedCodeDirectory != null && Module.GeneratedCodeDirectory.Exists())
 					{
-						if(!UnrealBuildTool.IsEngineInstalled() || !Utils.IsFileUnderDirectory(Module.GeneratedCodeDirectory, BuildConfiguration.RelativeEnginePath))
+						if(!UnrealBuildTool.IsEngineInstalled() || !Module.GeneratedCodeDirectory.IsUnderDirectory(UnrealBuildTool.EngineDirectory))
 						{
 							Log.TraceVerbose("\t\tDeleting " + Module.GeneratedCodeDirectory);
-							CleanDirectory(Module.GeneratedCodeDirectory);
+							CleanDirectory(Module.GeneratedCodeDirectory.FullName);
 						}
 					}
 				}
@@ -1334,45 +1327,45 @@ namespace UnrealBuildTool
 				}
 
 				// Clean the intermediate directory
-				if( !String.IsNullOrEmpty( ProjectIntermediateDirectory ) && (!UnrealBuildTool.IsEngineInstalled() || !Utils.IsFileUnderDirectory(ProjectIntermediateDirectory, BuildConfiguration.RelativeEnginePath)))
+				if( ProjectIntermediateDirectory != null && (!UnrealBuildTool.IsEngineInstalled() || !ProjectIntermediateDirectory.IsUnderDirectory(UnrealBuildTool.EngineDirectory)))
 				{
-					if (Directory.Exists(ProjectIntermediateDirectory))
+					if (ProjectIntermediateDirectory.Exists())
 					{
-						CleanDirectory(ProjectIntermediateDirectory);
+						CleanDirectory(ProjectIntermediateDirectory.FullName);
 					}
 				}
 				if (!UnrealBuildTool.RunningRocket())
 				{
 					// This is always under Rocket installation folder
-					if (Directory.Exists(GlobalLinkEnvironment.Config.IntermediateDirectory))
+					if (GlobalLinkEnvironment.Config.IntermediateDirectory.Exists())
 					{
 						Log.TraceVerbose("\tDeleting " + GlobalLinkEnvironment.Config.IntermediateDirectory);
-						CleanDirectory(GlobalLinkEnvironment.Config.IntermediateDirectory);
+						CleanDirectory(GlobalLinkEnvironment.Config.IntermediateDirectory.FullName);
 					}
 				}
 				else if (ShouldCompileMonolithic())
 				{
 					// Only in monolithic, otherwise it's pointing to Rocket installation folder
-					if (Directory.Exists(GlobalLinkEnvironment.Config.OutputDirectory))
+					if (GlobalLinkEnvironment.Config.OutputDirectory.Exists())
 					{
 						Log.TraceVerbose("\tDeleting " + GlobalLinkEnvironment.Config.OutputDirectory);
-						CleanDirectory(GlobalCompileEnvironment.Config.OutputDirectory);
+						CleanDirectory(GlobalCompileEnvironment.Config.OutputDirectory.FullName);
 					}
 				}
 
 				// Delete the dependency caches
 				{
 					var FlatCPPIncludeDependencyCacheFilename = FlatCPPIncludeDependencyCache.GetDependencyCachePathForTarget(this);
-					if (File.Exists(FlatCPPIncludeDependencyCacheFilename))
+					if (FlatCPPIncludeDependencyCacheFilename.Exists())
 					{
 						Log.TraceVerbose("\tDeleting " + FlatCPPIncludeDependencyCacheFilename);
-						CleanFile(FlatCPPIncludeDependencyCacheFilename);
+						CleanFile(FlatCPPIncludeDependencyCacheFilename.FullName);
 					}
 					var DependencyCacheFilename = DependencyCache.GetDependencyCachePathForTarget(this);
-					if (File.Exists(DependencyCacheFilename))
+					if (DependencyCacheFilename.Exists())
 					{
 						Log.TraceVerbose("\tDeleting " + DependencyCacheFilename);
-						CleanFile(DependencyCacheFilename);
+						CleanFile(DependencyCacheFilename.FullName);
 					}
 				}
 
@@ -1391,10 +1384,10 @@ namespace UnrealBuildTool
 					// Normal makefile
 					{					
 						var UBTMakefilePath = UnrealBuildTool.GetUBTMakefilePath( TargetDescs );
-						if (File.Exists(UBTMakefilePath))
+						if (UBTMakefilePath.Exists())
 						{
 							Log.TraceVerbose("\tDeleting " + UBTMakefilePath);
-							CleanFile(UBTMakefilePath);
+							CleanFile(UBTMakefilePath.FullName);
 						}
 					}
 
@@ -1402,10 +1395,10 @@ namespace UnrealBuildTool
 					{					
 						UEBuildConfiguration.bHotReloadFromIDE = true;
 						var UBTMakefilePath = UnrealBuildTool.GetUBTMakefilePath( TargetDescs );
-						if (File.Exists(UBTMakefilePath))
+						if (UBTMakefilePath.Exists())
 						{
 							Log.TraceVerbose("\tDeleting " + UBTMakefilePath);
-							CleanFile(UBTMakefilePath);
+							CleanFile(UBTMakefilePath.FullName);
 						}
 						UEBuildConfiguration.bHotReloadFromIDE = false;
 					}
@@ -1414,10 +1407,10 @@ namespace UnrealBuildTool
 				// Delete the action history
 				{
 					var ActionHistoryFilename = ActionHistory.GeneratePathForTarget(this);
-					if (File.Exists(ActionHistoryFilename))
+					if (ActionHistoryFilename.Exists())
 					{
 						Log.TraceVerbose("\tDeleting " + ActionHistoryFilename);
-						CleanFile(ActionHistoryFilename);
+						CleanFile(ActionHistoryFilename.FullName);
 					}
 				}
 
@@ -1461,7 +1454,7 @@ namespace UnrealBuildTool
 		{
 			// Expand all the paths in the receipt; they'll currently use variables for the engine and project directories
 			TargetReceipt ReceiptWithFullPaths = new TargetReceipt(Receipt);
-			ReceiptWithFullPaths.ExpandPathVariables(BuildConfiguration.RelativeEnginePath, ProjectDirectory);
+			ReceiptWithFullPaths.ExpandPathVariables(UnrealBuildTool.EngineDirectory, ProjectDirectory);
 			IUEBuildPlatform BuildPlatform = UEBuildPlatform.GetBuildPlatform(Platform);
 			// Expand all the paths in the receipt; they'll currently use variables for the engine and project directories
 			TargetReceipt BuiltReceiptWithFullPaths;
@@ -1469,7 +1462,7 @@ namespace UnrealBuildTool
 			{
 				return;
 			}
-			BuiltReceiptWithFullPaths.ExpandPathVariables(BuildConfiguration.RelativeEnginePath, ProjectDirectory);
+			BuiltReceiptWithFullPaths.ExpandPathVariables(UnrealBuildTool.EngineDirectory, ProjectDirectory);
 
 			{
 				// Collect all files to build
@@ -1555,11 +1548,11 @@ namespace UnrealBuildTool
 			foreach(string ModuleName in ModuleNames)
 			{
 				// Create the module rules
-				string ModuleRulesFileName;
+				FileReference ModuleRulesFileName;
 				ModuleRules Rules = RulesAssembly.CreateModuleRules(ModuleName, TargetInfo, out ModuleRulesFileName);
 
 				// Add the rules file itself
-				FileNames.Add(ModuleRulesFileName);
+				FileNames.Add(ModuleRulesFileName.FullName);
 
 				// Get a list of all the library paths
 				List<string> LibraryPaths = new List<string>();
@@ -1655,7 +1648,7 @@ namespace UnrealBuildTool
 			string ManifestPath;
 			if (UnrealBuildTool.RunningRocket() && UnrealBuildTool.HasUProjectFile())
 			{
-				ManifestPath = Path.Combine(UnrealBuildTool.GetUProjectPath(), BuildConfiguration.BaseIntermediateFolder, "Manifest.xml");
+				ManifestPath = Path.Combine(UnrealBuildTool.GetUProjectPath().FullName, BuildConfiguration.BaseIntermediateFolder, "Manifest.xml");
 			}
 			else
 			{
@@ -1671,7 +1664,7 @@ namespace UnrealBuildTool
 
 			// Expand all the paths in the receipt; they'll currently use variables for the engine and project directories
 			TargetReceipt ReceiptWithFullPaths = new TargetReceipt(Receipt);
-			ReceiptWithFullPaths.ExpandPathVariables(BuildConfiguration.RelativeEnginePath, ProjectDirectory);
+			ReceiptWithFullPaths.ExpandPathVariables(UnrealBuildTool.EngineDirectory, ProjectDirectory);
 
 			foreach(BuildProduct BuildProduct in ReceiptWithFullPaths.BuildProducts)
 			{
@@ -1733,18 +1726,18 @@ namespace UnrealBuildTool
 				}
 				Receipt.Merge(BinaryReceipt);
 			}
-			Receipt.InsertStandardPathVariables(BuildConfiguration.RelativeEnginePath, ProjectDirectory);
+			Receipt.InsertStandardPathVariables(UnrealBuildTool.EngineDirectory, ProjectDirectory);
 
 			// Prepare all the version manifests
-			Dictionary<string, VersionManifest> FileNameToVersionManifest = new Dictionary<string, VersionManifest>(StringComparer.InvariantCultureIgnoreCase);
+			Dictionary<FileReference, VersionManifest> FileNameToVersionManifest = new Dictionary<FileReference, VersionManifest>();
 			if(!bCompileMonolithic)
 			{
 				// Create the receipts for each folder
 				foreach(UEBuildBinary Binary in AppBinaries)
 				{
-					string DirectoryName = Path.GetDirectoryName(Binary.Config.OutputFilePath);
-					bool bIsGameDirectory = !Utils.IsFileUnderDirectory(Binary.Config.OutputFilePath, BuildConfiguration.RelativeEnginePath);
-					string ManifestFileName = Path.Combine(DirectoryName, VersionManifest.GetStandardFileName(AppName, Platform, Configuration, UEBuildPlatform.GetBuildPlatform(Platform).GetActiveArchitecture(), bIsGameDirectory));
+					DirectoryReference DirectoryName = Binary.Config.OutputFilePath.Directory;
+					bool bIsGameDirectory = !DirectoryName.IsUnderDirectory(UnrealBuildTool.EngineDirectory);
+					FileReference ManifestFileName = FileReference.Combine(DirectoryName, VersionManifest.GetStandardFileName(AppName, Platform, Configuration, UEBuildPlatform.GetBuildPlatform(Platform).GetActiveArchitecture(), bIsGameDirectory));
 	
 					VersionManifest Manifest;
 					if(!FileNameToVersionManifest.TryGetValue(ManifestFileName, out Manifest))
@@ -1752,7 +1745,7 @@ namespace UnrealBuildTool
 						Manifest = new VersionManifest(Version.Changelist, BuildId);
 
 						VersionManifest ExistingManifest;
-						if(VersionManifest.TryRead(ManifestFileName, out ExistingManifest) && Version.Changelist == ExistingManifest.Changelist)
+						if(VersionManifest.TryRead(ManifestFileName.FullName, out ExistingManifest) && Version.Changelist == ExistingManifest.Changelist)
 						{
 							if(OnlyModules.Count > 0)
 							{
@@ -1771,11 +1764,11 @@ namespace UnrealBuildTool
 
 					foreach(string ModuleName in Binary.Config.ModuleNames)
 					{
-						Manifest.ModuleNameToFileName[ModuleName] = Path.GetFileName(Binary.Config.OutputFilePath);
+						Manifest.ModuleNameToFileName[ModuleName] = Binary.Config.OutputFilePath.GetFileName();
 					}
 				}
 			}
-			FileNameToVersionManifestPairs = FileNameToVersionManifest.ToArray();
+			FileReferenceToVersionManifestPairs = FileNameToVersionManifest.ToArray();
 		}
 
 		/** Writes out the version manifest */
@@ -1795,11 +1788,11 @@ namespace UnrealBuildTool
 					Receipt.Write(ForceReceiptFileName);
 				}
 			}
-			if(FileNameToVersionManifestPairs != null)
+			if(FileReferenceToVersionManifestPairs != null)
 			{
-				foreach(KeyValuePair<string, VersionManifest> FileNameToVersionManifest in FileNameToVersionManifestPairs)
+				foreach(KeyValuePair<FileReference, VersionManifest> FileNameToVersionManifest in FileReferenceToVersionManifestPairs)
 				{
-					FileNameToVersionManifest.Value.Write(FileNameToVersionManifest.Key);
+					FileNameToVersionManifest.Value.Write(FileNameToVersionManifest.Key.FullName);
 				}
 			}
 		}
@@ -1872,7 +1865,7 @@ namespace UnrealBuildTool
 					&& IsCurrentPlatform)
 				{
 					// The hardcoded engine directory needs to be a relative path to match the normal EngineDir format. Not doing so breaks the network file system (TTP#315861).
-					string OutputFilePath = ExecutableBinary.Config.OutputFilePath;
+					string OutputFilePath = ExecutableBinary.Config.OutputFilePath.FullName;
 					if (Platform == UnrealTargetPlatform.Mac && OutputFilePath.Contains(".app/Contents/MacOS"))
 					{
 						OutputFilePath = OutputFilePath.Substring(0, OutputFilePath.LastIndexOf(".app/Contents/MacOS") + 4);
@@ -1888,7 +1881,7 @@ namespace UnrealBuildTool
 				// Set the define for the project name. This allows the executable to locate the correct project file to use, which may not be the same as the game name or target.
 				if(UnrealBuildTool.HasUProjectFile())
 				{
-					string ProjectName = Path.GetFileNameWithoutExtension(UnrealBuildTool.GetUProjectFile());
+					string ProjectName = UnrealBuildTool.GetUProjectFile().GetFileNameWithoutExtension();
 					GlobalCompileEnvironment.Config.Definitions.Add(String.Format("UE_PROJECT_NAME={0}", ProjectName));
 				}
 			}
@@ -1951,7 +1944,7 @@ namespace UnrealBuildTool
 						// If we've got this far and there are no source files then it's likely we're running Rocket and ignoring
 						// engine files, so we don't need a .generated.cpp either
 						UEBuildModuleCPP.AutoGenerateCppInfoClass.BuildInfoClass BuildInfo = null;
-						UHTModuleInfo.Info.GeneratedCPPFilenameBase = Path.Combine( Module.GeneratedCodeDirectory, UHTModuleInfo.Info.ModuleName ) + ".generated";
+						UHTModuleInfo.Info.GeneratedCPPFilenameBase = Path.Combine( Module.GeneratedCodeDirectory.FullName, UHTModuleInfo.Info.ModuleName ) + ".generated";
 						if (Module.SourceFilesToBuild.Count != 0)
 						{
 							BuildInfo = new UEBuildModuleCPP.AutoGenerateCppInfoClass.BuildInfoClass( UHTModuleInfo.Info.GeneratedCPPFilenameBase + "*.cpp" );
@@ -1992,7 +1985,7 @@ namespace UnrealBuildTool
 				if( UObjectModules.Count > 0 )
 				{
 					// Execute the header tool
-					string ModuleInfoFileName = Path.GetFullPath( Path.Combine( ProjectIntermediateDirectory, "UnrealHeaderTool.manifest" ) );
+					FileReference ModuleInfoFileName = FileReference.Combine( ProjectIntermediateDirectory, "UnrealHeaderTool.manifest" );
 					ECompilationResult UHTResult = ECompilationResult.OtherCompilationError;
 					if (!ExternalExecution.ExecuteHeaderToolIfNecessary(this, GlobalCompileEnvironment, UObjectModules, ModuleInfoFileName, ref UHTResult))
 					{
@@ -2101,7 +2094,7 @@ namespace UnrealBuildTool
 			Log.TraceVerbose("Building {0} - {1} - {2} - {3}", AppName, TargetName, Platform, Configuration);
 
 			// Put the non-executable output files (PDB, import library, etc) in the intermediate directory.
-			GlobalLinkEnvironment.Config.IntermediateDirectory = Path.GetFullPath(GlobalCompileEnvironment.Config.OutputDirectory);
+			GlobalLinkEnvironment.Config.IntermediateDirectory = GlobalCompileEnvironment.Config.OutputDirectory;
 			GlobalLinkEnvironment.Config.OutputDirectory = GlobalLinkEnvironment.Config.IntermediateDirectory;
 
 			// By default, shadow source files for this target in the root OutputDirectory
@@ -2208,7 +2201,7 @@ namespace UnrealBuildTool
 				{
 					if (!FilteredBinaries.Contains(DLLBinary) && DLLBinary.FindOnlyModule(OtherThingsWeNeedToBuild) != null)
 					{
-						if (!File.Exists(DLLBinary.Config.OutputFilePath))
+						if (!DLLBinary.Config.OutputFilePath.Exists())
 						{
 							throw new BuildException("Module {0} is potentially needed for the {1} platform to work, but it isn't actually compiled. This either needs to be marked as platform specific or made a dynamic dependency of something compiled with the base editor.", DLLBinary.Config.OutputFilePath, UnrealBuildTool.GetOnlyPlatformSpecificFor().ToString());
 						}
@@ -2224,13 +2217,10 @@ namespace UnrealBuildTool
 			{
 				var FilteredBinaries = new List<UEBuildBinary>();
 
-				// Have to do absolute here as this could be a project that is under the root
-				string FullEnginePath = Path.GetFullPath(BuildConfiguration.RelativeEnginePath);
-
 				// We only want to build rocket projects...
 				foreach (var DLLBinary in AppBinaries)
 				{
-					if (!Utils.IsFileUnderDirectory(DLLBinary.Config.OutputFilePath, FullEnginePath))
+					if (!DLLBinary.Config.OutputFilePath.IsUnderDirectory(UnrealBuildTool.EngineDirectory))
 					{
 						FilteredBinaries.Add(DLLBinary);
 					}
@@ -2259,15 +2249,15 @@ namespace UnrealBuildTool
 			}
 		}
 
-		private static string AddModuleFilenameSuffix(string ModuleName, string FilePath, string Suffix)
+		private static FileReference AddModuleFilenameSuffix(string ModuleName, FileReference FilePath, string Suffix)
 		{			
-			var MatchPos = FilePath.LastIndexOf(ModuleName, StringComparison.InvariantCultureIgnoreCase);
+			var MatchPos = FilePath.FullName.LastIndexOf(ModuleName, StringComparison.InvariantCultureIgnoreCase);
 			if (MatchPos < 0)
 			{
 				throw new BuildException("Failed to find module name \"{0}\" specified on the command line inside of the output filename \"{1}\" to add appendage.", ModuleName, FilePath);
 			}
 			var Appendage = "-" + Suffix;
-			return FilePath.Insert(MatchPos + ModuleName.Length, Appendage);
+			return new FileReference(FilePath.FullName.Insert(MatchPos + ModuleName.Length, Appendage));
 		}
 
 		private static List<UEBuildBinary> GetFilteredOnlyModules(List<UEBuildBinary> Binaries, List<OnlyModule> OnlyModules)
@@ -2329,8 +2319,7 @@ namespace UnrealBuildTool
 			{
 				foreach (var TargetModuleName in BinaryCPP.ModuleNames)
 				{
-					string UnusedFilename;
-					ModuleRules CheckRules = RulesAssembly.CreateModuleRules(TargetModuleName, TargetInfo, out UnusedFilename);
+					ModuleRules CheckRules = RulesAssembly.CreateModuleRules(TargetModuleName, TargetInfo);
 					if (CheckRules.Type != ModuleRules.ModuleType.External)
 					{
 						PrivateDependencyModuleNames.Add(TargetModuleName);
@@ -2344,11 +2333,11 @@ namespace UnrealBuildTool
 
 				// Include an empty header so UEBuildModule.Compile does not complain about a lack of PCH
 				string HeaderFilename                    = LinkerFixupsName + "Name.h";
-				string LinkerFixupHeaderFilenameWithPath = Path.Combine(GlobalCompileEnvironment.Config.OutputDirectory, HeaderFilename);
+				FileReference LinkerFixupHeaderFilenameWithPath = FileReference.Combine(GlobalCompileEnvironment.Config.OutputDirectory, HeaderFilename);
 
 				// Create the cpp filename
-				string LinkerFixupCPPFilename = Path.Combine(GlobalCompileEnvironment.Config.OutputDirectory, LinkerFixupsName + ".cpp");
-				if (!File.Exists(LinkerFixupCPPFilename))
+				FileReference LinkerFixupCPPFilename = FileReference.Combine(GlobalCompileEnvironment.Config.OutputDirectory, LinkerFixupsName + ".cpp");
+				if (!LinkerFixupCPPFilename.Exists())
 				{
 					// Create a dummy file in case it doesn't exist yet so that the module does not complain it's not there
 					ResponseFile.Create(LinkerFixupHeaderFilenameWithPath, new List<string>());
@@ -2357,11 +2346,11 @@ namespace UnrealBuildTool
 
 				// Create the source file list (just the one cpp file)
 				List<FileItem> SourceFiles = new List<FileItem>();
-				var LinkerFixupCPPFileItem = FileItem.GetItemByPath(LinkerFixupCPPFilename);
+				var LinkerFixupCPPFileItem = FileItem.GetItemByFileReference(LinkerFixupCPPFilename);
 				SourceFiles.Add(LinkerFixupCPPFileItem);
 
 				// Create the CPP module
-				var FakeModuleDirectory = Path.GetDirectoryName( LinkerFixupCPPFilename );
+				var FakeModuleDirectory = LinkerFixupCPPFilename.Directory;
 				var NewModule = CreateArtificialModule(LinkerFixupsName, FakeModuleDirectory, SourceFiles, PrivateDependencyModuleNames);
 
 				// Now bind this new module to the executable binary so it will link the plugin libs correctly
@@ -2376,9 +2365,9 @@ namespace UnrealBuildTool
 
 				// Determine if the file changed. Write it if it either doesn't exist or the contents are different.
 				bool bShouldWriteFile = true;
-				if (File.Exists(LinkerFixupCPPFilename))
+				if (LinkerFixupCPPFilename.Exists())
 				{
-					string[] ExistingFixupText = File.ReadAllLines(LinkerFixupCPPFilename);
+					string[] ExistingFixupText = File.ReadAllLines(LinkerFixupCPPFilename.FullName);
 					string JoinedNewContents = string.Join("", LinkerFixupsFileContents.ToArray());
 					string JoinedOldContents = string.Join("", ExistingFixupText);
 					bShouldWriteFile = (JoinedNewContents != JoinedOldContents);
@@ -2391,7 +2380,7 @@ namespace UnrealBuildTool
 					ResponseFile.Create(LinkerFixupCPPFilename, LinkerFixupsFileContents);
 
 					// Update the cached file states so that the linker fixups definitely get rebuilt
-					FileItem.GetItemByPath(LinkerFixupHeaderFilenameWithPath).ResetFileInfo();
+					FileItem.GetItemByFileReference(LinkerFixupHeaderFilenameWithPath).ResetFileInfo();
 					LinkerFixupCPPFileItem.ResetFileInfo();
 				}
 			}
@@ -2471,7 +2460,7 @@ namespace UnrealBuildTool
 		/// <param name="SourceFiles">Source files.</param>
 		/// <param name="PrivateDependencyModuleNames">Private dependency list.</param>
 		/// <returns>Created module.</returns>
-		private UEBuildModuleCPP CreateArtificialModule(string Name, string Directory, IEnumerable<FileItem> SourceFiles, IEnumerable<string> PrivateDependencyModuleNames)
+		private UEBuildModuleCPP CreateArtificialModule(string Name, DirectoryReference Directory, IEnumerable<FileItem> SourceFiles, IEnumerable<string> PrivateDependencyModuleNames)
 		{
 			ModuleRules Rules = new ModuleRules();
 			Rules.PrivateDependencyModuleNames.AddRange(PrivateDependencyModuleNames);
@@ -2600,8 +2589,8 @@ namespace UnrealBuildTool
 					if (Module.IsCompiledInConfiguration(Platform, TargetType, UEBuildConfiguration.bBuildDeveloperTools, UEBuildConfiguration.bBuildEditor))
 					{
 						// Add the corresponding binary for it
-						string ModuleFileName = RulesAssembly.GetModuleFileName(Module.Name);
-						bool bHasSource = (!String.IsNullOrEmpty(ModuleFileName) && Directory.EnumerateFiles(Path.GetDirectoryName(ModuleFileName), "*.cpp", SearchOption.AllDirectories).Any());
+						FileReference ModuleFileName = RulesAssembly.GetModuleFileName(Module.Name);
+						bool bHasSource = (ModuleFileName != null && Directory.EnumerateFiles(ModuleFileName.Directory.FullName, "*.cpp", SearchOption.AllDirectories).Any());
 						AddBinaryForModule(Module.Name, BinaryType, bAllowCompilation: bHasSource, bIsCrossTarget: false);
 
 						// Add it to the binary if we're compiling monolithic (and it's enabled)
@@ -2746,18 +2735,18 @@ namespace UnrealBuildTool
 			PluginInfo Plugin = FindPluginForModule(ModuleName);
 
 			// Get the root output directory and base name (target name/app name) for this binary
-			string BaseOutputDirectory;
+			DirectoryReference BaseOutputDirectory;
 			if(Plugin != null)
 			{
-				BaseOutputDirectory = Path.GetFullPath(Plugin.Directory);
+				BaseOutputDirectory = Plugin.Directory;
 			}
 			else if(RulesAssembly.IsGameModule(ModuleName) || !bUseSharedBuildEnvironment)
 			{
-				BaseOutputDirectory = Path.GetFullPath(ProjectDirectory);
+				BaseOutputDirectory = ProjectDirectory;
 			}
 			else
 			{
-				BaseOutputDirectory = Path.GetFullPath(BuildConfiguration.RelativeEnginePath);
+				BaseOutputDirectory = UnrealBuildTool.EngineDirectory;
 			}
 
 			// Get the configuration that this module will be built in. Engine modules compiled in DebugGame will use Development.
@@ -2768,20 +2757,20 @@ namespace UnrealBuildTool
 			}
 
 			// Get the output and intermediate directories for this module
-			string OutputDirectory = Path.Combine(BaseOutputDirectory, "Binaries", Platform.ToString());
- 			string IntermediateDirectory = Path.Combine(BaseOutputDirectory, BuildConfiguration.PlatformIntermediateFolder, AppName, ModuleConfiguration.ToString());
+			DirectoryReference OutputDirectory = DirectoryReference.Combine(BaseOutputDirectory, "Binaries", Platform.ToString());
+ 			DirectoryReference IntermediateDirectory = DirectoryReference.Combine(BaseOutputDirectory, BuildConfiguration.PlatformIntermediateFolder, AppName, ModuleConfiguration.ToString());
 
 			// Append a subdirectory if the module rules specifies one
 			ModuleRules ModuleRules;
 			if(RulesAssembly.TryCreateModuleRules(ModuleName, TargetInfo, out ModuleRules) && !String.IsNullOrEmpty(ModuleRules.BinariesSubFolder))
 			{
-				OutputDirectory = Path.Combine(OutputDirectory, ModuleRules.BinariesSubFolder);
-				IntermediateDirectory = Path.Combine(IntermediateDirectory, ModuleRules.BinariesSubFolder);
+				OutputDirectory = DirectoryReference.Combine(OutputDirectory, ModuleRules.BinariesSubFolder);
+				IntermediateDirectory = DirectoryReference.Combine(IntermediateDirectory, ModuleRules.BinariesSubFolder);
 			}
 
 			// Get the output filenames
-			string BaseBinaryPath = Path.Combine(OutputDirectory, MakeBinaryFileName(AppName + "-" + ModuleName, Platform, ModuleConfiguration, Rules.UndecoratedConfiguration, BinaryType));
-			List<string> OutputFilePaths = UEBuildPlatform.GetBuildPlatform(Platform).FinalizeBinaryPaths(BaseBinaryPath);
+			FileReference BaseBinaryPath = FileReference.Combine(OutputDirectory, MakeBinaryFileName(AppName + "-" + ModuleName, Platform, ModuleConfiguration, Rules.UndecoratedConfiguration, BinaryType));
+			List<FileReference> OutputFilePaths = UEBuildPlatform.GetBuildPlatform(Platform).FinalizeBinaryPaths(BaseBinaryPath);
 
 			// Prepare the configuration object
 			UEBuildBinaryConfiguration Config = new UEBuildBinaryConfiguration(
@@ -2822,7 +2811,7 @@ namespace UnrealBuildTool
 		/**
 		 * @return The overall output directory of actions for this target.
 		 */
-		public string GetOutputDirectory()
+		public DirectoryReference GetOutputDirectory()
 		{
 			return GlobalCompileEnvironment.Config.OutputDirectory;
 		}
@@ -2887,7 +2876,7 @@ namespace UnrealBuildTool
 		/// <param name="bIncludesGameModules">Whether this executable contains game modules</param>
 		/// <param name="ExeSubFolder">Subfolder for executables. May be null.</param>
 		/// <returns>List of executable paths for this target</returns>
-		public static List<string> MakeExecutablePaths(string BaseDirectory, string BinaryName, UnrealTargetPlatform Platform, UnrealTargetConfiguration Configuration, UnrealTargetConfiguration UndecoratedConfiguration, bool bIncludesGameModules, string ExeSubFolder)
+		public static List<FileReference> MakeExecutablePaths(DirectoryReference BaseDirectory, string BinaryName, UnrealTargetPlatform Platform, UnrealTargetConfiguration Configuration, UnrealTargetConfiguration UndecoratedConfiguration, bool bIncludesGameModules, string ExeSubFolder)
 		{
 			// Get the configuration for the executable. If we're building DebugGame, and this executable only contains engine modules, use the same name as development.
 			UnrealTargetConfiguration ExeConfiguration = Configuration;
@@ -2897,15 +2886,15 @@ namespace UnrealBuildTool
 			}
 
 			// Build the binary path
-			string BinaryPath = Path.Combine(BaseDirectory, "Binaries", Platform.ToString());
+			DirectoryReference BinaryDirectory = DirectoryReference.Combine(BaseDirectory, "Binaries", Platform.ToString());
 			if (!String.IsNullOrEmpty(ExeSubFolder))
 			{
-				BinaryPath = Path.Combine(BinaryPath, ExeSubFolder);
+				BinaryDirectory = DirectoryReference.Combine(BinaryDirectory, ExeSubFolder);
 			}
-			BinaryPath = Path.Combine(BinaryPath, MakeBinaryFileName(BinaryName, Platform, ExeConfiguration, UndecoratedConfiguration, UEBuildBinaryType.Executable));
+			FileReference BinaryFile = FileReference.Combine(BinaryDirectory, MakeBinaryFileName(BinaryName, Platform, ExeConfiguration, UndecoratedConfiguration, UEBuildBinaryType.Executable));
 
 			// Allow the platform to customize the output path (and output several executables at once if necessary)
-			return UEBuildPlatform.GetBuildPlatform(Platform).FinalizeBinaryPaths(BinaryPath);
+			return UEBuildPlatform.GetBuildPlatform(Platform).FinalizeBinaryPaths(BinaryFile);
 		}
 
 		/** Sets up the modules for the target. */
@@ -2921,7 +2910,7 @@ namespace UnrealBuildTool
 		protected virtual void SetupPlugins()
 		{
 			// Filter the plugins list by the current project
-			ValidPlugins = Plugins.ReadAvailablePlugins(BuildConfiguration.RelativeEnginePath, UnrealBuildTool.GetUProjectFile());
+			ValidPlugins = Plugins.ReadAvailablePlugins(UnrealBuildTool.EngineDirectory, UnrealBuildTool.GetUProjectFile());
 
 			// Remove any plugins for platforms we don't have
 			List<string> ExcludeFolders = new List<string>();
@@ -2941,7 +2930,7 @@ namespace UnrealBuildTool
 			// If we're compiling against the engine, add the plugins enabled for this target
 			if(UEBuildConfiguration.bCompileAgainstEngine)
 			{
-				ProjectDescriptor Project = UnrealBuildTool.HasUProjectFile()? ProjectDescriptor.FromFile(UnrealBuildTool.GetUProjectFile()) : null;
+				ProjectDescriptor Project = UnrealBuildTool.HasUProjectFile()? ProjectDescriptor.FromFile(UnrealBuildTool.GetUProjectFile().FullName) : null;
 				foreach(PluginInfo ValidPlugin in ValidPlugins)
 				{
 					if(UProjectInfo.IsPluginEnabledForProject(ValidPlugin, Project, Platform))
@@ -2978,9 +2967,9 @@ namespace UnrealBuildTool
 			// Add any foreign plugins to the list
 			if(ForeignPlugins != null)
 			{
-				foreach(string ForeignPlugin in ForeignPlugins)
+				foreach(FileReference ForeignPlugin in ForeignPlugins)
 				{
-					PluginInfo ForeignPluginInfo = ValidPlugins.FirstOrDefault(x => String.Compare(x.FileName, ForeignPlugin, true) == 0);
+					PluginInfo ForeignPluginInfo = ValidPlugins.FirstOrDefault(x => x.File == ForeignPlugin);
 					if(ForeignPluginInfo == null)
 					{
 						ForeignPluginInfo = new PluginInfo(ForeignPlugin, PluginLoadedFrom.GameProject);
@@ -3001,11 +2990,11 @@ namespace UnrealBuildTool
 			string RelativePathFromRoot;
 			if(Plugin.LoadedFrom == PluginLoadedFrom.Engine)
 			{
-				RelativePathFromRoot = Utils.CleanDirectorySeparators(Utils.MakePathRelativeTo(Plugin.FileName, BuildConfiguration.RelativeEnginePath), '/');
+				RelativePathFromRoot = Utils.CleanDirectorySeparators(Plugin.File.MakeRelativeTo(UnrealBuildTool.EngineDirectory), '/');
 			}
 			else
 			{
-				RelativePathFromRoot = Utils.CleanDirectorySeparators(Utils.MakePathRelativeTo(Plugin.FileName, UnrealBuildTool.GetUProjectPath()), '/');
+				RelativePathFromRoot = Utils.CleanDirectorySeparators(Plugin.File.MakeRelativeTo(UnrealBuildTool.GetUProjectPath()), '/');
 			}
 			return ExcludeFragments.Any(x => RelativePathFromRoot.Contains(x));
 		}
@@ -3131,23 +3120,23 @@ namespace UnrealBuildTool
 					IntermediateConfiguration = Configuration == UnrealTargetConfiguration.DebugGame ? UnrealTargetConfiguration.Development : Configuration;
 				}
 
-				GlobalCompileEnvironment.Config.OutputDirectory = Path.Combine(BuildConfiguration.PlatformIntermediatePath, OutputAppName, IntermediateConfiguration.ToString());
+				GlobalCompileEnvironment.Config.OutputDirectory = DirectoryReference.Combine(UnrealBuildTool.EngineDirectory, BuildConfiguration.PlatformIntermediateFolder, OutputAppName, IntermediateConfiguration.ToString());
 				if (ShouldCompileMonolithic())
 				{
-					if(!String.IsNullOrEmpty(UnrealBuildTool.GetUProjectPath()))
+					if(UnrealBuildTool.HasUProjectFile())
 					{
-						GlobalCompileEnvironment.Config.OutputDirectory = Path.Combine(UnrealBuildTool.GetUProjectPath(), BuildConfiguration.PlatformIntermediateFolder, OutputAppName, IntermediateConfiguration.ToString());
+						GlobalCompileEnvironment.Config.OutputDirectory = DirectoryReference.Combine(UnrealBuildTool.GetUProjectPath(), BuildConfiguration.PlatformIntermediateFolder, OutputAppName, IntermediateConfiguration.ToString());
 					}
 					else if(ForeignPlugins.Count > 0)
 					{
-						GlobalCompileEnvironment.Config.OutputDirectory = Path.Combine(Path.GetDirectoryName(ForeignPlugins[0]), BuildConfiguration.PlatformIntermediateFolder, OutputAppName, IntermediateConfiguration.ToString());
+						GlobalCompileEnvironment.Config.OutputDirectory = DirectoryReference.Combine(ForeignPlugins[0].Directory, BuildConfiguration.PlatformIntermediateFolder, OutputAppName, IntermediateConfiguration.ToString());
 					}
 				}
 			}
 			else
 			{
 				var GlobalEnvironmentConfiguration = Configuration == UnrealTargetConfiguration.DebugGame ? UnrealTargetConfiguration.Development : Configuration;
-				GlobalCompileEnvironment.Config.OutputDirectory = Path.Combine(BuildConfiguration.PlatformIntermediatePath, OutputAppName, GlobalEnvironmentConfiguration.ToString());
+				GlobalCompileEnvironment.Config.OutputDirectory = DirectoryReference.Combine(UnrealBuildTool.EngineDirectory, BuildConfiguration.PlatformIntermediateFolder, OutputAppName, GlobalEnvironmentConfiguration.ToString());
 			}
 
 			// By default, shadow source files for this target in the root OutputDirectory
@@ -3296,12 +3285,11 @@ namespace UnrealBuildTool
 		{
 			UEBuildPlatform.GetBuildPlatform(Platform).SetUpConfigurationEnvironment(this);
 
-			var UpperExtension = Path.GetExtension (OutputPaths[0]).ToUpperInvariant (); 
 			// Check to see if we're compiling a library or not
-			bool bIsBuildingDLL = UpperExtension == ".DLL";
+			bool bIsBuildingDLL = OutputPaths[0].HasExtension(".dll");
 			GlobalCompileEnvironment.Config.bIsBuildingDLL = bIsBuildingDLL;
 			GlobalLinkEnvironment.Config.bIsBuildingDLL = bIsBuildingDLL;
-			bool bIsBuildingLibrary = UpperExtension == ".LIB";
+			bool bIsBuildingLibrary = OutputPaths[0].HasExtension(".lib");
 			GlobalCompileEnvironment.Config.bIsBuildingLibrary = bIsBuildingLibrary;
 			GlobalLinkEnvironment.Config.bIsBuildingLibrary = bIsBuildingLibrary;
 		}
@@ -3329,18 +3317,15 @@ namespace UnrealBuildTool
 
 				// @todo projectfiles: Cross-platform modules can appear here during project generation, but they may have already
 				//   been filtered out by the project generator.  This causes the projects to not be added to directories properly.
-				string ModuleFileName;
+				FileReference ModuleFileName;
 				var RulesObject = RulesAssembly.CreateModuleRules(ModuleName, TargetInfo, out ModuleFileName);
-				var ModuleDirectory = Path.GetDirectoryName(ModuleFileName);
-
+				DirectoryReference ModuleDirectory = ModuleFileName.Directory;
+				
 				// Get the plugin for this module
 				PluginInfo Plugin = FindPluginForModule(ModuleName);
 
-				// Making an assumption here that any project file path that contains '../../'
-				// is NOT from the engine and therefore must be an application-specific module.
 				// Get the type of module we're creating
 				var ModuleType = UEBuildModuleType.Unknown;
-				var ModuleFileRelativeToEngineDirectory = Utils.MakePathRelativeTo(ModuleFileName, ProjectFileGenerator.EngineRelativePath);
 
 				// see if it's external
 				if (RulesObject.Type == ModuleRules.ModuleType.External)
@@ -3361,13 +3346,13 @@ namespace UnrealBuildTool
 					if (ModuleType == UEBuildModuleType.Unknown)
 					{
 						// not a plugin, see if it is a game module 
-						if (ModuleFileRelativeToEngineDirectory.StartsWith("..") || Path.IsPathRooted(ModuleFileRelativeToEngineDirectory))
+						if(!ModuleFileName.IsUnderDirectory(UnrealBuildTool.EngineSourceDirectory))
 						{
 							ModuleType = UEBuildModuleType.Game;
 						}
 						else
 						{
-							ModuleType = UEBuildModule.GetEngineModuleTypeBasedOnLocation(ModuleName, ModuleType, ModuleFileRelativeToEngineDirectory);
+							ModuleType = UEBuildModule.GetEngineModuleTypeBasedOnLocation(ModuleName, ModuleType, ModuleFileName);
 							if (ModuleType == UEBuildModuleType.Unknown)
 							{
 								throw new BuildException("Unable to determine module type for {0}", ModuleFileName);
@@ -3379,50 +3364,49 @@ namespace UnrealBuildTool
 				var IsGameModule = !ModuleType.IsEngineModule();
 
 				// Get the base directory for paths referenced by the module. If the module's under the UProject source directory use that, otherwise leave it relative to the Engine source directory.
-				string ProjectSourcePath = UnrealBuildTool.GetUProjectSourcePath();
-				if (ProjectSourcePath != null)
+				DirectoryReference ProjectSourceDirectoryName = UnrealBuildTool.GetUProjectSourcePath();
+				if (ProjectSourceDirectoryName != null)
 				{
-					string FullProjectSourcePath = Path.GetFullPath(ProjectSourcePath);
-					if (Utils.IsFileUnderDirectory(ModuleFileName, FullProjectSourcePath))
+					if (ModuleFileName.IsUnderDirectory(ProjectSourceDirectoryName))
 					{
-						RulesObject.PublicIncludePaths = CombinePathList(ProjectSourcePath, RulesObject.PublicIncludePaths);
-						RulesObject.PrivateIncludePaths = CombinePathList(ProjectSourcePath, RulesObject.PrivateIncludePaths);
-						RulesObject.PublicLibraryPaths = CombinePathList(ProjectSourcePath, RulesObject.PublicLibraryPaths);
-						RulesObject.PublicAdditionalShadowFiles = CombinePathList(ProjectSourcePath, RulesObject.PublicAdditionalShadowFiles);
+						RulesObject.PublicIncludePaths = CombinePathList(ProjectSourceDirectoryName, RulesObject.PublicIncludePaths);
+						RulesObject.PrivateIncludePaths = CombinePathList(ProjectSourceDirectoryName, RulesObject.PrivateIncludePaths);
+						RulesObject.PublicLibraryPaths = CombinePathList(ProjectSourceDirectoryName, RulesObject.PublicLibraryPaths);
+						RulesObject.PublicAdditionalShadowFiles = CombinePathList(ProjectSourceDirectoryName, RulesObject.PublicAdditionalShadowFiles);
 					}
 				}
 
 				// Get the generated code directory. Plugins always write to their own intermediate directory so they can be copied between projects, shared engine 
 				// intermediates go in the engine intermediate folder, and anything else goes in the project folder.
-				string GeneratedCodeDirectory = null;
+				DirectoryReference GeneratedCodeDirectory = null;
 				if (RulesObject.Type != ModuleRules.ModuleType.External)
 				{
 					if(Plugin != null)
 					{
 						GeneratedCodeDirectory = Plugin.Directory;
 					}
-					else if(bUseSharedBuildEnvironment && Utils.IsFileUnderDirectory(RulesAssembly.GetModuleFileName(ModuleName), BuildConfiguration.RelativeEnginePath))
+					else if(bUseSharedBuildEnvironment && ModuleFileName.IsUnderDirectory(UnrealBuildTool.EngineDirectory))
 					{
-						GeneratedCodeDirectory = BuildConfiguration.RelativeEnginePath;
+						GeneratedCodeDirectory = UnrealBuildTool.EngineDirectory;
 					}
 					else
 					{
 						GeneratedCodeDirectory = ProjectDirectory;
 					}
-					GeneratedCodeDirectory = Path.GetFullPath(Path.Combine(GeneratedCodeDirectory, BuildConfiguration.PlatformIntermediateFolder, AppName, "Inc", ModuleName));
+					GeneratedCodeDirectory = DirectoryReference.Combine(GeneratedCodeDirectory, BuildConfiguration.PlatformIntermediateFolder, AppName, "Inc", ModuleName);
 				}
 
 				// Don't generate include paths for third party modules; they don't follow our conventions. Core is a special-case... leave it alone
 				if (RulesObject.Type != ModuleRules.ModuleType.External && ModuleName != "Core")
 				{
 					// Add the default include paths to the module rules, if they exist.
-					AddDefaultIncludePathsToModuleRules(ModuleName, ModuleFileName, ModuleFileRelativeToEngineDirectory, IsGameModule, RulesObject);
+					string ModuleFileRelativeToEngineDirectory = ModuleFileName.MakeRelativeTo(UnrealBuildTool.EngineDirectory);
+					AddDefaultIncludePathsToModuleRules(ModuleName, ModuleFileName.FullName, ModuleFileRelativeToEngineDirectory, IsGameModule, RulesObject);
 
 					// Add the path to the generated headers 
 					if(GeneratedCodeDirectory != null)
 					{
-						string EngineSourceDirectory = Path.Combine(ProjectFileGenerator.RootRelativePath, "Engine/Source");
-						string RelativeGeneratedCodeDirectory = Utils.CleanDirectorySeparators(Utils.MakePathRelativeTo(GeneratedCodeDirectory, EngineSourceDirectory), '/');
+						string RelativeGeneratedCodeDirectory = Utils.CleanDirectorySeparators(GeneratedCodeDirectory.MakeRelativeTo(UnrealBuildTool.EngineSourceDirectory), '/');
 						RulesObject.PublicIncludePaths.Add(RelativeGeneratedCodeDirectory);
 					}
 				}
@@ -3443,15 +3427,15 @@ namespace UnrealBuildTool
 					}
 
 					// So all we care about are the game module and/or plugins.
-					if (bDiscoverFiles && (!UnrealBuildTool.IsEngineInstalled() || !Utils.IsFileUnderDirectory(ModuleFileName, BuildConfiguration.RelativeEnginePath)))
+					if (bDiscoverFiles && (!UnrealBuildTool.IsEngineInstalled() || !ModuleFileName.IsUnderDirectory(UnrealBuildTool.EngineDirectory)))
 					{
-						var SourceFilePaths = new List<string>();
+						List<FileReference> SourceFilePaths = new List<FileReference>();
 
 						if (ProjectFile != null)
 						{
 							foreach (var SourceFile in ProjectFile.SourceFiles)
 							{
-								SourceFilePaths.Add(SourceFile.FilePath);
+								SourceFilePaths.Add(SourceFile.Reference);
 							}
 						}
 						else
@@ -3486,12 +3470,12 @@ namespace UnrealBuildTool
 			ModuleRules          RulesObject,
 			string               ModuleName,
 			UEBuildModuleType    ModuleType,
-			string               ModuleDirectory,
-			string               GeneratedCodeDirectory,
+			DirectoryReference    ModuleDirectory,
+			DirectoryReference    GeneratedCodeDirectory,
 			IntelliSenseGatherer IntelliSenseGatherer,
 			List<FileItem>       ModuleSourceFiles,
 			bool                 bBuildSourceFiles,
-			string               InBuildCsFile)
+			FileReference         InBuildCsFile)
 		{
 			switch (RulesObject.Type)
 			{
@@ -3506,7 +3490,7 @@ namespace UnrealBuildTool
 							InSourceFiles: ModuleSourceFiles,
 							InRules: RulesObject,
 							bInBuildSourceFiles: bBuildSourceFiles,
-							InBuildCsFilename: InBuildCsFile
+							InBuildCsFilename: InBuildCsFile.FullName
 						);
 
 				case ModuleRules.ModuleType.CPlusPlusCLR:
@@ -3520,7 +3504,7 @@ namespace UnrealBuildTool
 							InSourceFiles: ModuleSourceFiles,
 							InRules: RulesObject,
 							bInBuildSourceFiles : bBuildSourceFiles,
-							InBuildCsFilename: InBuildCsFile
+							InBuildCsFilename: InBuildCsFile.FullName
 						);
 
 				case ModuleRules.ModuleType.External:
@@ -3530,7 +3514,7 @@ namespace UnrealBuildTool
 							InType: ModuleType,
 							InModuleDirectory: ModuleDirectory,
 							InRules: RulesObject,
-							InBuildCsFilename: InBuildCsFile
+							InBuildCsFilename: InBuildCsFile.FullName
 						);
 
 				default:
@@ -3588,7 +3572,7 @@ namespace UnrealBuildTool
 			string ClassesDirectory = Path.Combine(ModuleDirectoryRelativeToEngineSourceDirectory, "Classes/");	// @todo uht: Deprecate eventually.  Or force it to be manually specified...
 			string PublicDirectory = Path.Combine(ModuleDirectoryRelativeToEngineSourceDirectory, "Public/");
 
-			if(!Utils.IsFileUnderDirectory(InModuleFilename, AbsEngineSourceDirectory))
+			if(!Utils.IsFileUnderDirectory(InModuleFilename, UnrealBuildTool.EngineSourceDirectory.FullName))
 			{
 				// This will be either the format 
 				//		../<Game>/Source/<Module>
@@ -3702,16 +3686,12 @@ namespace UnrealBuildTool
 		/// <param name="BasePath">Base path to combine with. May be null or empty.</param>
 		/// <param name="PathList">List of input paths to combine with. May be null.</param>
 		/// <returns>List of paths relative The build module object for the specified build rules source file</returns>
-		private static List<string> CombinePathList(string BasePath, List<string> PathList)
+		private static List<string> CombinePathList(DirectoryReference BasePath, List<string> PathList)
 		{
-			List<string> NewPathList = PathList;
-			if (PathList != null && !String.IsNullOrEmpty(BasePath))
+			List<string> NewPathList = new List<string>();
+			foreach (string Path in PathList)
 			{
-				NewPathList = new List<string>();
-				foreach (string Path in PathList)
-				{
-					NewPathList.Add(System.IO.Path.Combine(BasePath, Path));
-				}
+				NewPathList.Add(System.IO.Path.Combine(BasePath.FullName, Path));
 			}
 			return NewPathList;
 		}
@@ -3724,7 +3704,7 @@ namespace UnrealBuildTool
 		/// <param name="SourceFilesBaseDirectory">Directory that the source files are in</param>
 		/// <param name="TargetPlatform">The platform we're going to compile for</param>
 		/// <returns>The list of source files to actually compile</returns>
-		static List<FileItem> GetCPlusPlusFilesToBuild(List<string> SourceFiles, string SourceFilesBaseDirectory, UnrealTargetPlatform TargetPlatform)
+		static List<FileItem> GetCPlusPlusFilesToBuild(List<FileReference> SourceFiles, DirectoryReference SourceFilesBaseDirectory, UnrealTargetPlatform TargetPlatform)
 		{
 			// Make a list of all platform name strings that we're *not* currently compiling, to speed
 			// up file path comparisons later on
@@ -3757,7 +3737,7 @@ namespace UnrealBuildTool
 				bool IsCompilableSourceFile = false;
 				foreach (var CurExtension in CompilableSourceFileTypes)
 				{
-					if (SourceFilePath.EndsWith(CurExtension, StringComparison.InvariantCultureIgnoreCase))
+					if (SourceFilePath.HasExtension(CurExtension))
 					{
 						IsCompilableSourceFile = true;
 						break;
@@ -3766,10 +3746,10 @@ namespace UnrealBuildTool
 
 				if (IsCompilableSourceFile)
 				{
-					if (SourceFilePath.StartsWith(SourceFilesBaseDirectory + Path.DirectorySeparatorChar))
+					if (SourceFilePath.IsUnderDirectory(SourceFilesBaseDirectory))
 					{
 						// Store the path as relative to the project file
-						var RelativeFilePath = Utils.MakePathRelativeTo(SourceFilePath, SourceFilesBaseDirectory);
+						var RelativeFilePath = SourceFilePath.MakeRelativeTo(SourceFilesBaseDirectory);
 
 						// All compiled files should always be in a sub-directory under the project file directory.  We enforce this here.
 						if (Path.IsPathRooted(RelativeFilePath) || RelativeFilePath.StartsWith(".."))
@@ -3792,7 +3772,7 @@ namespace UnrealBuildTool
 
 						if (IncludeThisFile)
 						{
-							FilteredFileItems.Add(FileItem.GetItemByFullPath(SourceFilePath));
+							FilteredFileItems.Add(FileItem.GetItemByFileReference(SourceFilePath));
 						}
 					}
 				}
diff --git a/Engine/Source/Programs/UnrealBuildTool/HTML5/HTML5ToolChain.cs b/Engine/Source/Programs/UnrealBuildTool/HTML5/HTML5ToolChain.cs
--- a/Engine/Source/Programs/UnrealBuildTool/HTML5/HTML5ToolChain.cs
+++ b/Engine/Source/Programs/UnrealBuildTool/HTML5/HTML5ToolChain.cs
@@ -339,10 +339,10 @@ namespace UnrealBuildTool
                 string FileArguments = string.Format(" \"{0}\"", SourceFile.AbsolutePath);
 				var ObjectFileExtension = UEBuildPlatform.BuildPlatformDictionary[UnrealTargetPlatform.HTML5].GetBinaryExtension(UEBuildBinaryType.Object);
             	// Add the object file to the produced item list.
-				FileItem ObjectFile = FileItem.GetItemByPath(
-					Path.Combine(
+				FileItem ObjectFile = FileItem.GetItemByFileReference(
+					FileReference.Combine(
 						CompileEnvironment.Config.OutputDirectory,
-					Path.GetFileName(SourceFile.AbsolutePath) + ObjectFileExtension
+						Path.GetFileName(SourceFile.AbsolutePath) + ObjectFileExtension
 						)
 					);
 				CompileAction.ProducedItems.Add(ObjectFile);
@@ -358,7 +358,7 @@ namespace UnrealBuildTool
 					FileArguments += GetCLArguments_CPP(CompileEnvironment);
 				}
 
-				CompileAction.WorkingDirectory = Path.GetFullPath(".");
+				CompileAction.WorkingDirectory = UnrealBuildTool.EngineSourceDirectory.FullName;
 				CompileAction.CommandPath = HTML5SDKInfo.Python();
 
 				CompileAction.CommandArguments = HTML5SDKInfo.EmscriptenCompiler() + " " + Arguments + FileArguments + CompileEnvironment.Config.AdditionalArguments;
@@ -468,7 +468,7 @@ namespace UnrealBuildTool
 			List<string> ReponseLines = new List<string>(); 
 
 			LinkAction.bCanExecuteRemotely = false;
-			LinkAction.WorkingDirectory = Path.GetFullPath(".");
+			LinkAction.WorkingDirectory = UnrealBuildTool.EngineSourceDirectory.FullName;
 			LinkAction.CommandPath = HTML5SDKInfo.Python();
 			LinkAction.CommandArguments = HTML5SDKInfo.EmscriptenCompiler();
 			ReponseLines.Add(GetLinkArguments(LinkEnvironment));
@@ -511,17 +511,17 @@ namespace UnrealBuildTool
 			// make the file we will create
 
 
-			OutputFile = FileItem.GetItemByPath(LinkEnvironment.Config.OutputFilePath);
+			OutputFile = FileItem.GetItemByFileReference(LinkEnvironment.Config.OutputFilePath);
 			LinkAction.ProducedItems.Add(OutputFile);
 			ReponseLines.Add(string.Format(" -o \"{0}\"", OutputFile.AbsolutePath));
 
-		    FileItem OutputBC = FileItem.GetItemByPath(LinkEnvironment.Config.OutputFilePath.Replace(".js", ".bc").Replace(".html", ".bc"));
+		    FileItem OutputBC = FileItem.GetItemByPath(LinkEnvironment.Config.OutputFilePath.FullName.Replace(".js", ".bc").Replace(".html", ".bc"));
 		    LinkAction.ProducedItems.Add(OutputBC);
 			ReponseLines.Add( " --emit-symbol-map " + string.Format(" --save-bc \"{0}\"", OutputBC.AbsolutePath));
 
      		LinkAction.StatusDescription = Path.GetFileName(OutputFile.AbsolutePath);
 
-			string ResponseFileName = GetResponseFileName(LinkEnvironment, OutputFile);
+			FileReference ResponseFileName = GetResponseFileName(LinkEnvironment, OutputFile);
 
 
 			LinkAction.CommandArguments += string.Format(" @\"{0}\""  , ResponseFile.Create(ResponseFileName, ReponseLines));
@@ -531,7 +531,7 @@ namespace UnrealBuildTool
 			return OutputFile;
 		}
 
-		public override void CompileCSharpProject(CSharpEnvironment CompileEnvironment, string ProjectFileName, string DestinationFile)
+		public override void CompileCSharpProject(CSharpEnvironment CompileEnvironment, FileReference ProjectFileName, FileReference DestinationFile)
 		{
 			throw new BuildException("HTML5 cannot compile C# files");
 		}
diff --git a/Engine/Source/Programs/UnrealBuildTool/HTML5/UEBuildHTML5.cs b/Engine/Source/Programs/UnrealBuildTool/HTML5/UEBuildHTML5.cs
--- a/Engine/Source/Programs/UnrealBuildTool/HTML5/UEBuildHTML5.cs
+++ b/Engine/Source/Programs/UnrealBuildTool/HTML5/UEBuildHTML5.cs
@@ -52,10 +52,10 @@ namespace UnrealBuildTool
 		}
 
 		// F5 should always try to run the Win32 version
-		public override string ModifyNMakeOutput(string ExeName)
+		public override FileReference ModifyNMakeOutput(FileReference ExeName)
 		{
 			// nmake Run should always run the win32 version
-			return Path.ChangeExtension(ExeName+"-win32", ".exe");
+			return (ExeName + "-win32").ChangeExtension(".exe");
 		}
 
 		/**
diff --git a/Engine/Source/Programs/UnrealBuildTool/IOS/IOSToolChain.cs b/Engine/Source/Programs/UnrealBuildTool/IOS/IOSToolChain.cs
--- a/Engine/Source/Programs/UnrealBuildTool/IOS/IOSToolChain.cs
+++ b/Engine/Source/Programs/UnrealBuildTool/IOS/IOSToolChain.cs
@@ -64,7 +64,7 @@ namespace UnrealBuildTool
 		/** Which library archiver to use */
 		private static string IOSArchiver = "libtool";
 
-		public List<string> BuiltBinaries = new List<string>();
+		public List<FileReference> BuiltBinaries = new List<FileReference>();
 
 		/** Additional frameworks stored locally so we have access without LinkEnvironment */
 		public List<UEBuildFramework> RememberedAdditionalFrameworks = new List<UEBuildFramework>();
@@ -98,7 +98,7 @@ namespace UnrealBuildTool
 		{
 			if (BuildConfiguration.bCreateStubIPA && Binary.Config.Type != UEBuildBinaryType.StaticLibrary)
 			{
-				string StubFile = Path.Combine (Path.GetDirectoryName (Binary.Config.OutputFilePath), Path.GetFileNameWithoutExtension (Binary.Config.OutputFilePath) + ".stub");
+				FileReference StubFile = FileReference.Combine(Binary.Config.OutputFilePath.Directory, Binary.Config.OutputFilePath.GetFileNameWithoutExtension() + ".stub");
 				Receipt.AddBuildProduct(StubFile, BuildProductType.Executable);
 			}
 		}
@@ -108,7 +108,7 @@ namespace UnrealBuildTool
 		/// </summary>
 		/// <param name="OutputFile">Build product to add</param>
 		/// <param name="DebugExtension">Extension for the matching debug file (may be null).</param>
-		public override bool ShouldAddDebugFileToReceipt(string OutputFile, BuildProductType OutputType)
+		public override bool ShouldAddDebugFileToReceipt(FileReference OutputFile, BuildProductType OutputType)
 		{
 			return OutputType == BuildProductType.Executable;
 		}
@@ -505,10 +505,10 @@ namespace UnrealBuildTool
 				{
 					var PrecompiledFileExtension = UEBuildPlatform.BuildPlatformDictionary[UnrealTargetPlatform.IOS].GetBinaryExtension(UEBuildBinaryType.PrecompiledHeader);
 					// Add the precompiled header file to the produced item list.
-					FileItem PrecompiledHeaderFile = FileItem.GetItemByPath(
-						Path.Combine(
+					FileItem PrecompiledHeaderFile = FileItem.GetItemByFileReference(
+						FileReference.Combine(
 							CompileEnvironment.Config.OutputDirectory,
-						Path.GetFileName(SourceFile.AbsolutePath) + PrecompiledFileExtension
+							Path.GetFileName(SourceFile.AbsolutePath) + PrecompiledFileExtension
 							)
 						);
 
@@ -528,8 +528,8 @@ namespace UnrealBuildTool
 					}
 					var ObjectFileExtension = UEBuildPlatform.BuildPlatformDictionary[UnrealTargetPlatform.IOS].GetBinaryExtension(UEBuildBinaryType.Object);
 					// Add the object file to the produced item list.
-					FileItem ObjectFile = FileItem.GetItemByPath(
-						Path.Combine(
+					FileItem ObjectFile = FileItem.GetItemByFileReference(
+						FileReference.Combine(
 							CompileEnvironment.Config.OutputDirectory,
 							Path.GetFileName(SourceFile.AbsolutePath) + ObjectFileExtension
 							)
@@ -599,7 +599,7 @@ namespace UnrealBuildTool
 					if (!String.IsNullOrEmpty(Path.GetDirectoryName(AdditionalLibrary)))
 					{
 						// add it to the prerequisites to make sure it's built first (this should be the case of non-system libraries)
-						FileItem LibFile = FileItem.GetItemByPath(Path.GetFullPath(AdditionalLibrary));
+						FileItem LibFile = FileItem.GetItemByPath(AdditionalLibrary);
 						FileItem RemoteLibFile = LocalToRemoteFileItem(LibFile, true);
 						LinkAction.PrerequisiteItems.Add(RemoteLibFile);
 
@@ -666,7 +666,7 @@ namespace UnrealBuildTool
 			}
 
 			// Add the output file as a production of the link action.
-			FileItem OutputFile = FileItem.GetItemByPath(Path.GetFullPath(LinkEnvironment.Config.OutputFilePath));
+			FileItem OutputFile = FileItem.GetItemByFileReference(LinkEnvironment.Config.OutputFilePath);
 			FileItem RemoteOutputFile = LocalToRemoteFileItem(OutputFile, false);
 			LinkAction.ProducedItems.Add(RemoteOutputFile);
 
@@ -692,18 +692,18 @@ namespace UnrealBuildTool
 			}
 			else
 			{
-				bool bIsUE4Game = LinkEnvironment.Config.OutputFilePath.Contains("UE4Game");
-				string ResponsePath = Path.GetFullPath(Path.Combine((!bIsUE4Game && !string.IsNullOrEmpty(UnrealBuildTool.GetUProjectPath())) ? UnrealBuildTool.GetUProjectPath() : BuildConfiguration.RelativeEnginePath, BuildConfiguration.PlatformIntermediateFolder, "LinkFileList_" + Path.GetFileNameWithoutExtension(LinkEnvironment.Config.OutputFilePath) + ".tmp")); 
+				bool bIsUE4Game = LinkEnvironment.Config.OutputFilePath.FullName.Contains("UE4Game");
+				FileReference ResponsePath = FileReference.Combine((!bIsUE4Game && UnrealBuildTool.HasUProjectFile()) ? UnrealBuildTool.GetUProjectPath() : UnrealBuildTool.EngineDirectory, BuildConfiguration.PlatformIntermediateFolder, "LinkFileList_" + LinkEnvironment.Config.OutputFilePath.GetFileNameWithoutExtension() + ".tmp"); 
 				if (!Utils.IsRunningOnMono && BuildHostPlatform.Current.Platform != UnrealTargetPlatform.Mac)
 				{
 					ResponseFile.Create (ResponsePath, InputFileNames);
-					RPCUtilHelper.CopyFile (ResponsePath, ConvertPath (ResponsePath), true);
+					RPCUtilHelper.CopyFile (ResponsePath.FullName, ConvertPath (ResponsePath.FullName), true);
 				}
 				else
 				{
-					ResponseFile.Create(ConvertPath(ResponsePath), InputFileNames);
+					ResponseFile.Create(new FileReference(ConvertPath(ResponsePath.FullName)), InputFileNames);
 				}
-				LinkAction.CommandArguments += string.Format(" @\"{0}\"", ConvertPath(ResponsePath));
+				LinkAction.CommandArguments += string.Format(" @\"{0}\"", ConvertPath(ResponsePath.FullName));
 			}
 
 			// Add the output file to the command-line.
@@ -900,7 +900,7 @@ namespace UnrealBuildTool
 		{
 			if (BuildHostPlatform.Current.Platform != UnrealTargetPlatform.Mac)
 			{
-				BuiltBinaries = new List<string>();
+				BuiltBinaries = new List<FileReference>();
 			}
 
 			base.PreBuildSync();
@@ -946,7 +946,7 @@ namespace UnrealBuildTool
 
 			if (BuildHostPlatform.Current.Platform == UnrealTargetPlatform.Mac)
 			{
-				string RemoteShadowDirectoryMac = Path.GetDirectoryName(Target.OutputPath);
+				string RemoteShadowDirectoryMac = Path.GetDirectoryName(Target.OutputPath.FullName);
 				string FinalRemoteExecutablePath = String.Format("{0}/Payload/{1}.app/{1}", RemoteShadowDirectoryMac, AppName);
 
 				// strip the debug info from the executable if needed
@@ -973,7 +973,7 @@ namespace UnrealBuildTool
 				{
 					Directory.CreateDirectory(String.Format("{0}/Payload/{1}.app", RemoteShadowDirectoryMac, AppName));
 				}
-				File.Copy(Target.OutputPath, FinalRemoteExecutablePath, true);
+				File.Copy(Target.OutputPath.FullName, FinalRemoteExecutablePath, true);
 
 				if (BuildConfiguration.bCreateStubIPA)
 				{
@@ -1032,7 +1032,7 @@ namespace UnrealBuildTool
 						}
 
 						// Package the stub
-						PackageStub (RemoteShadowDirectoryMac, AppName, Path.GetFileNameWithoutExtension (Target.OutputPath));
+						PackageStub (RemoteShadowDirectoryMac, AppName, Target.OutputPath.GetFileNameWithoutExtension());
 					}
 				}
 
@@ -1074,24 +1074,24 @@ namespace UnrealBuildTool
 				// store off the binaries
 				foreach (UEBuildBinary Binary in Target.AppBinaries)
 				{
-					BuiltBinaries.Add(Path.GetFullPath(Binary.ToString()));
+					BuiltBinaries.Add(Binary.Config.OutputFilePath);
 				}
 
 				// check to see if the DangerouslyFast mode is valid (in other words, a build has gone through since a Rebuild/Clean operation)
-				string DangerouslyFastValidFile = Path.Combine(Target.GlobalLinkEnvironment.Config.IntermediateDirectory, "DangerouslyFastIsNotDangerous");
+				FileReference DangerouslyFastValidFile = FileReference.Combine(Target.GlobalLinkEnvironment.Config.IntermediateDirectory, "DangerouslyFastIsNotDangerous");
 				bool bUseDangerouslyFastModeWasRequested = bUseDangerouslyFastMode;
 				if (bUseDangerouslyFastMode)
 				{
-					if (!File.Exists(DangerouslyFastValidFile))
+					if (!DangerouslyFastValidFile.Exists())
 					{
 						Log.TraceInformation("Dangeroulsy Fast mode was requested, but a slow mode hasn't completed. Performing slow now...");
 						bUseDangerouslyFastMode = false;
 					}
 				}
 
-				foreach (string FilePath in BuiltBinaries)
+				foreach (FileReference FilePath in BuiltBinaries)
 				{
-					string RemoteExecutablePath = ConvertPath(FilePath);
+					string RemoteExecutablePath = ConvertPath(FilePath.FullName);
 
 					// when going super fast, just copy the executable to the final resting spot
 					if (bUseDangerouslyFastMode)
@@ -1101,18 +1101,18 @@ namespace UnrealBuildTool
 						Log.TraceInformation("==============================================================================");
 
 						// copy the executable
-						string RemoteShadowDirectoryMac = ConvertPath(Path.GetDirectoryName(Target.OutputPath));
+						string RemoteShadowDirectoryMac = ConvertPath(Path.GetDirectoryName(Target.OutputPath.FullName));
 						string FinalRemoteExecutablePath = String.Format("{0}/Payload/{1}.app/{1}", RemoteShadowDirectoryMac, Target.TargetName);
 						RPCUtilHelper.Command("/", String.Format("cp -f {0} {1}", RemoteExecutablePath, FinalRemoteExecutablePath), "", null);
 					}
 					else if (!Utils.IsRunningOnMono && BuildHostPlatform.Current.Platform != UnrealTargetPlatform.Mac)
 					{
-						RPCUtilHelper.CopyFile(RemoteExecutablePath, FilePath, false);
+						RPCUtilHelper.CopyFile(RemoteExecutablePath, FilePath.FullName, false);
 
 						if (BuildConfiguration.bGeneratedSYMFile == true)
 						{
 							string DSYMExt = ".app.dSYM.zip";
-							RPCUtilHelper.CopyFile(RemoteExecutablePath + DSYMExt, FilePath + DSYMExt, false);
+							RPCUtilHelper.CopyFile(RemoteExecutablePath + DSYMExt, FilePath.FullName + DSYMExt, false);
 						}
 					}
 				}
@@ -1130,7 +1130,7 @@ namespace UnrealBuildTool
 					if (!bUseDangerouslyFastMode)
 					{
 						// generate the dummy project so signing works
-						UnrealBuildTool.GenerateProjectFiles(new XcodeProjectFileGenerator(), new string[] { "-NoIntellisense", "-iosdeployonly", (!string.IsNullOrEmpty(UnrealBuildTool.GetUProjectPath()) ? "-game" : "") });
+						UnrealBuildTool.GenerateProjectFiles(new XcodeProjectFileGenerator(), new string[] { "-NoIntellisense", "-iosdeployonly", (UnrealBuildTool.HasUProjectFile() ? "-game" : "") });
 					}
 
 					// now that 
@@ -1139,7 +1139,7 @@ namespace UnrealBuildTool
 					StubGenerateProcess.StartInfo.FileName = Path.Combine(StubGenerateProcess.StartInfo.WorkingDirectory, "iPhonePackager.exe");
 
 					string Arguments = "";
-					string PathToApp = Target.RulesAssembly.GetTargetFileName(AppName);
+					string PathToApp = Target.RulesAssembly.GetTargetFileName(AppName).FullName;
 
 					// right now, no programs have a Source subdirectory, so assume the PathToApp is directly in the root
 					if (Path.GetDirectoryName(PathToApp).Contains(@"\Engine\Source\Programs"))
@@ -1224,12 +1224,12 @@ namespace UnrealBuildTool
 					// now that a slow mode sync has finished, we can now do DangerouslyFast mode again (if requested)
 					if (bUseDangerouslyFastModeWasRequested)
 					{
-						File.Create(DangerouslyFastValidFile);
+						File.Create(DangerouslyFastValidFile.FullName);
 					}
 					else
 					{
 						// if we didn't want dangerously fast, then delete the file so that setting/unsetting the flag will do the right thing without a Rebuild
-						File.Delete(DangerouslyFastValidFile);
+						File.Delete(DangerouslyFastValidFile.FullName);
 					}
 				}
 
diff --git a/Engine/Source/Programs/UnrealBuildTool/IOS/UEBuildIOS.cs b/Engine/Source/Programs/UnrealBuildTool/IOS/UEBuildIOS.cs
--- a/Engine/Source/Programs/UnrealBuildTool/IOS/UEBuildIOS.cs
+++ b/Engine/Source/Programs/UnrealBuildTool/IOS/UEBuildIOS.cs
@@ -269,7 +269,7 @@ namespace UnrealBuildTool
 		 *
 		 * return true if the project uses the default build config
 		 */
-		public override bool HasDefaultBuildConfig(UnrealTargetPlatform Platform, string ProjectPath)
+		public override bool HasDefaultBuildConfig(UnrealTargetPlatform Platform, DirectoryReference ProjectDirectoryName)
 		{
 			string[] BoolKeys = new string[] {
 				"bDevForArmV7", "bDevForArm64", "bDevForArmV7S", "bShipForArmV7", 
@@ -282,14 +282,14 @@ namespace UnrealBuildTool
 			};
 
 			// look up iOS specific settings
-			if (!DoProjectSettingsMatchDefault(Platform, ProjectPath, "/Script/IOSRuntimeSettings.IOSRuntimeSettings",
+			if (!DoProjectSettingsMatchDefault(Platform, ProjectDirectoryName, "/Script/IOSRuntimeSettings.IOSRuntimeSettings",
 				    BoolKeys, null, StringKeys))
 			{
 				return false;
 			}
 
 			// check the base settings
-			return base.HasDefaultBuildConfig(Platform, ProjectPath);
+			return base.HasDefaultBuildConfig(Platform, ProjectDirectoryName);
 		}
 
 		/**
diff --git a/Engine/Source/Programs/UnrealBuildTool/IOS/UEDeployIOS.cs b/Engine/Source/Programs/UnrealBuildTool/IOS/UEDeployIOS.cs
--- a/Engine/Source/Programs/UnrealBuildTool/IOS/UEDeployIOS.cs
+++ b/Engine/Source/Programs/UnrealBuildTool/IOS/UEDeployIOS.cs
@@ -681,8 +681,8 @@ namespace UnrealBuildTool.IOS
 		public override bool PrepTargetForDeployment(UEBuildTarget InTarget)
 		{
 			string GameName = InTarget.TargetName;
-			string BuildPath = (GameName == "UE4Game" ? "../../Engine" : InTarget.ProjectDirectory) + "/Binaries/IOS";
-			string ProjectDirectory = InTarget.ProjectDirectory;
+			string BuildPath = (GameName == "UE4Game" ? "../../Engine" : InTarget.ProjectDirectory.FullName) + "/Binaries/IOS";
+			string ProjectDirectory = InTarget.ProjectDirectory.FullName;
 			bool bIsUE4Game = GameName.Contains("UE4Game");
 
 			string DecoratedGameName;
@@ -710,7 +710,7 @@ namespace UnrealBuildTool.IOS
 
 					try
 					{
-						string BinaryDir = Path.GetDirectoryName(InTarget.OutputPath) + "\\";
+						string BinaryDir = Path.GetDirectoryName(InTarget.OutputPath.FullName) + "\\";
 						if (BinaryDir.EndsWith(InTarget.AppName + "\\Binaries\\IOS\\") && InTarget.TargetType != TargetRules.TargetType.Game)
 						{
 							BinaryDir = BinaryDir.Replace(InTarget.TargetType.ToString(), "Game");
@@ -724,7 +724,7 @@ namespace UnrealBuildTool.IOS
 							AppFullName += "-" + InTarget.Configuration.ToString();
 						}
 
-						foreach (string BinaryPath in Toolchain.BuiltBinaries)
+						foreach (string BinaryPath in Toolchain.BuiltBinaries.Select(x => x.FullName))
 						{
 							if (!BinaryPath.Contains("Dummy"))
 							{
@@ -739,7 +739,7 @@ namespace UnrealBuildTool.IOS
 					}
 				}
 
-				GeneratePList(ProjectDirectory, bIsUE4Game, GameName, Path.GetFileNameWithoutExtension(UnrealBuildTool.GetUProjectFile()), "../../Engine", "");
+				GeneratePList(ProjectDirectory, bIsUE4Game, GameName, Path.GetFileNameWithoutExtension(UnrealBuildTool.GetUProjectFile().FullName), "../../Engine", "");
 			}
 			return true;
 		}
diff --git a/Engine/Source/Programs/UnrealBuildTool/Linux/LinuxToolChain.cs b/Engine/Source/Programs/UnrealBuildTool/Linux/LinuxToolChain.cs
--- a/Engine/Source/Programs/UnrealBuildTool/Linux/LinuxToolChain.cs
+++ b/Engine/Source/Programs/UnrealBuildTool/Linux/LinuxToolChain.cs
@@ -669,8 +669,8 @@ namespace UnrealBuildTool
                 {
 					var PrecompiledFileExtension = UEBuildPlatform.BuildPlatformDictionary[UnrealTargetPlatform.Linux].GetBinaryExtension(UEBuildBinaryType.PrecompiledHeader);
                     // Add the precompiled header file to the produced item list.
-                    FileItem PrecompiledHeaderFile = FileItem.GetItemByPath(
-                        Path.Combine(
+                    FileItem PrecompiledHeaderFile = FileItem.GetItemByFileReference(
+                        FileReference.Combine(
                             CompileEnvironment.Config.OutputDirectory,
                             Path.GetFileName(SourceFile.AbsolutePath) + PrecompiledFileExtension
                             )
@@ -692,8 +692,8 @@ namespace UnrealBuildTool
 
 					var ObjectFileExtension = UEBuildPlatform.BuildPlatformDictionary[UnrealTargetPlatform.Linux].GetBinaryExtension(UEBuildBinaryType.Object);
                     // Add the object file to the produced item list.
-                    FileItem ObjectFile = FileItem.GetItemByPath(
-                        Path.Combine(
+                    FileItem ObjectFile = FileItem.GetItemByFileReference(
+                        FileReference.Combine(
                             CompileEnvironment.Config.OutputDirectory,
                             Path.GetFileName(SourceFile.AbsolutePath) + ObjectFileExtension
                             )
@@ -707,7 +707,7 @@ namespace UnrealBuildTool
                 // Add the source file path to the command-line.
                 FileArguments += string.Format(" \"{0}\"", SourceFile.AbsolutePath);
 
-                CompileAction.WorkingDirectory = Path.GetFullPath(".");
+                CompileAction.WorkingDirectory = UnrealBuildTool.EngineSourceDirectory.FullName;
                 if (!UsingClang())
                 {
                     CompileAction.CommandPath = GCCPath;
@@ -737,7 +737,7 @@ namespace UnrealBuildTool
         {
             // Create an archive action
             Action ArchiveAction = new Action(ActionType.Link);
-            ArchiveAction.WorkingDirectory = Path.GetFullPath(".");
+            ArchiveAction.WorkingDirectory = UnrealBuildTool.EngineSourceDirectory.FullName;
             bool bUsingSh = BuildHostPlatform.Current.Platform != UnrealTargetPlatform.Win64 && BuildHostPlatform.Current.Platform != UnrealTargetPlatform.Win32;
             if (bUsingSh)
             {
@@ -754,7 +754,7 @@ namespace UnrealBuildTool
             ArchiveAction.bProducesImportLibrary = true;
 
             // Add the output file as a production of the link action.
-            FileItem OutputFile = FileItem.GetItemByPath(LinkEnvironment.Config.OutputFilePath);
+            FileItem OutputFile = FileItem.GetItemByFileReference(LinkEnvironment.Config.OutputFilePath);
             ArchiveAction.ProducedItems.Add(OutputFile);
             ArchiveAction.CommandDescription = "Archive";
             ArchiveAction.StatusDescription = Path.GetFileName(OutputFile.AbsolutePath);
@@ -806,10 +806,10 @@ namespace UnrealBuildTool
             string ExecuteSwitch = bUseCmdExe ? " /C" : ""; // avoid -c so scripts don't need +x
             string ScriptName = bUseCmdExe ? "FixDependencies.bat" : "FixDependencies.sh";
 
-            FileItem FixDepsScript = FileItem.GetItemByFullPath(Path.Combine(LinkEnvironment.Config.LocalShadowDirectory, ScriptName));
+            FileItem FixDepsScript = FileItem.GetItemByFileReference(FileReference.Combine(LinkEnvironment.Config.LocalShadowDirectory, ScriptName));
 
 			Action PostLinkAction = new Action(ActionType.Link);
-			PostLinkAction.WorkingDirectory = Path.GetFullPath(".");
+			PostLinkAction.WorkingDirectory = UnrealBuildTool.EngineSourceDirectory.FullName;
             PostLinkAction.CommandPath = ShellBinary;
 			PostLinkAction.StatusDescription = string.Format("{0}", Path.GetFileName(Executable.AbsolutePath));
 			PostLinkAction.CommandDescription = "FixDeps";
@@ -818,7 +818,7 @@ namespace UnrealBuildTool
             
             PostLinkAction.CommandArguments += bUseCmdExe ? " \"" : " -c '";
 
-			FileItem OutputFile = FileItem.GetItemByPath(Path.Combine(LinkEnvironment.Config.LocalShadowDirectory, Path.GetFileNameWithoutExtension(Executable.AbsolutePath) + ".link"));
+			FileItem OutputFile = FileItem.GetItemByFileReference(FileReference.Combine(LinkEnvironment.Config.LocalShadowDirectory, Path.GetFileNameWithoutExtension(Executable.AbsolutePath) + ".link"));
 
 			// Make sure we don't run this script until the all executables and shared libraries
 			// have been built.
@@ -855,7 +855,7 @@ namespace UnrealBuildTool
 
             // Create an action that invokes the linker.
             Action LinkAction = new Action(ActionType.Link);
-            LinkAction.WorkingDirectory = Path.GetFullPath(".");
+            LinkAction.WorkingDirectory = UnrealBuildTool.EngineSourceDirectory.FullName;
             if (String.IsNullOrEmpty(ClangPath))
             {
                 LinkAction.CommandPath = GCCPath;
@@ -873,7 +873,7 @@ namespace UnrealBuildTool
             LinkAction.bProducesImportLibrary = LinkEnvironment.Config.bIsBuildingDLL;
 
             // Add the output file as a production of the link action.
-            FileItem OutputFile = FileItem.GetItemByPath(LinkEnvironment.Config.OutputFilePath);
+            FileItem OutputFile = FileItem.GetItemByFileReference(LinkEnvironment.Config.OutputFilePath);
             LinkAction.ProducedItems.Add(OutputFile);
             LinkAction.CommandDescription = "Link";
             LinkAction.StatusDescription = Path.GetFileName(OutputFile.AbsolutePath);
@@ -889,7 +889,7 @@ namespace UnrealBuildTool
                 LinkAction.PrerequisiteItems.Add(InputFile);
             }
 
-            string ResponseFileName = GetResponseFileName(LinkEnvironment, OutputFile);
+            FileReference ResponseFileName = GetResponseFileName(LinkEnvironment, OutputFile);
             LinkAction.CommandArguments += string.Format(" -Wl,@\"{0}\"", ResponseFile.Create(ResponseFileName, InputFileNames));
 
             if (LinkEnvironment.Config.bIsBuildingDLL)
@@ -997,17 +997,17 @@ namespace UnrealBuildTool
             LinkAction.CommandArguments = LinkAction.CommandArguments.Replace("\\", "/");
 
             // prepare a linker script
-            string LinkerScriptPath = Path.Combine(LinkEnvironment.Config.LocalShadowDirectory, "remove-sym.ldscript");
-            if (!Directory.Exists(LinkEnvironment.Config.LocalShadowDirectory))
+            FileReference LinkerScriptPath = FileReference.Combine(LinkEnvironment.Config.LocalShadowDirectory, "remove-sym.ldscript");
+            if (!LinkEnvironment.Config.LocalShadowDirectory.Exists())
             {
-                Directory.CreateDirectory(LinkEnvironment.Config.LocalShadowDirectory);
+                LinkEnvironment.Config.LocalShadowDirectory.CreateDirectory();
             }
-            if (File.Exists(LinkerScriptPath))
+            if (LinkerScriptPath.Exists())
             {
-                File.Delete(LinkerScriptPath);
+                LinkerScriptPath.Delete();
             }
 
-            using (StreamWriter Writer = File.CreateText(LinkerScriptPath))
+            using (StreamWriter Writer = File.CreateText(LinkerScriptPath.FullName))
             {
                 Writer.WriteLine("UE4 {");
                 Writer.WriteLine("  global: *;");
@@ -1030,7 +1030,7 @@ namespace UnrealBuildTool
                 bool bUseCmdExe = BuildHostPlatform.Current.Platform == UnrealTargetPlatform.Win64 || BuildHostPlatform.Current.Platform == UnrealTargetPlatform.Win32;
                 string ScriptName = bUseCmdExe ? "FixDependencies.bat" : "FixDependencies.sh";
 
-				string FixDepsScriptPath = Path.Combine(LinkEnvironment.Config.LocalShadowDirectory, ScriptName);
+				string FixDepsScriptPath = Path.Combine(LinkEnvironment.Config.LocalShadowDirectory.FullName, ScriptName);
 				if (!bHasWipedFixDepsScript)
 				{
 					bHasWipedFixDepsScript = true;
@@ -1091,7 +1091,7 @@ namespace UnrealBuildTool
             return OutputFile;
         }
 
-        public override void CompileCSharpProject(CSharpEnvironment CompileEnvironment, string ProjectFileName, string DestinationFile)
+        public override void CompileCSharpProject(CSharpEnvironment CompileEnvironment, FileReference ProjectFileName, FileReference DestinationFile)
         {
             throw new BuildException("Linux cannot compile C# files");
         }
@@ -1100,12 +1100,12 @@ namespace UnrealBuildTool
         {
             foreach (UEBuildBinary Binary in Binaries)
             {
-                BundleDependencies.Add(FileItem.GetItemByPath(Binary.ToString()));
+                BundleDependencies.Add(FileItem.GetItemByFileReference(Binary.Config.OutputFilePath));
             }
         }
 
         /** Converts the passed in path from UBT host to compiler native format. */
-        public override String ConvertPath(String OriginalPath)
+        public override string ConvertPath(string OriginalPath)
         {
             if (BuildHostPlatform.Current.Platform == UnrealTargetPlatform.Linux)
             {
