Use FileReference and DirectoryReference classes instead of strings for paths. Massively reduces the amount of paranoid calls to GetFullPath(), CleanDirectorySeparators() et al., and improves performa

diff --git a/Engine/Source/Programs/UnrealBuildTool/Mac/MacToolChain.cs b/Engine/Source/Programs/UnrealBuildTool/Mac/MacToolChain.cs
--- a/Engine/Source/Programs/UnrealBuildTool/Mac/MacToolChain.cs
+++ b/Engine/Source/Programs/UnrealBuildTool/Mac/MacToolChain.cs
@@ -359,8 +359,8 @@ namespace UnrealBuildTool
 				{
 					var PrecompiledHeaderExtension = UEBuildPlatform.BuildPlatformDictionary[UnrealTargetPlatform.Mac].GetBinaryExtension(UEBuildBinaryType.PrecompiledHeader);
 					// Add the precompiled header file to the produced item list.
-					FileItem PrecompiledHeaderFile = FileItem.GetItemByPath(
-						Path.Combine(
+					FileItem PrecompiledHeaderFile = FileItem.GetItemByFileReference(
+						FileReference.Combine(
 							CompileEnvironment.Config.OutputDirectory,
 							Path.GetFileName(SourceFile.AbsolutePath) + PrecompiledHeaderExtension
 							)
@@ -382,8 +382,8 @@ namespace UnrealBuildTool
 					}
 					var ObjectFileExtension = UEBuildPlatform.BuildPlatformDictionary[UnrealTargetPlatform.Mac].GetBinaryExtension(UEBuildBinaryType.Object);
 					// Add the object file to the produced item list.
-					FileItem ObjectFile = FileItem.GetItemByPath(
-						Path.Combine(
+					FileItem ObjectFile = FileItem.GetItemByFileReference(
+						FileReference.Combine(
 							CompileEnvironment.Config.OutputDirectory,
 							Path.GetFileName(SourceFile.AbsolutePath) + ObjectFileExtension
 							)
@@ -592,7 +592,7 @@ namespace UnrealBuildTool
 			LinkAction.bProducesImportLibrary = !Utils.IsRunningOnMono && (bBuildImportLibraryOnly || LinkEnvironment.Config.bIsBuildingDLL);
 
 			// Add the output file as a production of the link action.
-			FileItem OutputFile = FileItem.GetItemByPath(Path.GetFullPath(LinkEnvironment.Config.OutputFilePath));
+			FileItem OutputFile = FileItem.GetItemByFileReference(LinkEnvironment.Config.OutputFilePath);
 			OutputFile.bNeedsHotReloadNumbersDLLCleanUp = LinkEnvironment.Config.bIsBuildingDLL;
 
 			FileItem RemoteOutputFile = LocalToRemoteFileItem(OutputFile, false);
@@ -686,7 +686,7 @@ namespace UnrealBuildTool
 
 						if (!Utils.IsRunningOnMono)
 						{
-							FileItem EngineLibDependency = FileItem.GetItemByPath(Path.GetFullPath(AdditionalLibrary));
+							FileItem EngineLibDependency = FileItem.GetItemByPath(AdditionalLibrary);
 							LinkAction.PrerequisiteItems.Add(EngineLibDependency);
 							FileItem RemoteEngineLibDependency = FileItem.GetRemoteItemByPath(ConvertPath(Path.GetFullPath(AdditionalLibrary)), UnrealTargetPlatform.Mac);
 							LinkAction.PrerequisiteItems.Add(RemoteEngineLibDependency);
@@ -694,7 +694,7 @@ namespace UnrealBuildTool
 						}
 						else if (LinkEnvironment.Config.bIsCrossReferenced == false)
 						{
-							FileItem EngineLibDependency = FileItem.GetItemByPath(Path.GetFullPath(AdditionalLibrary));
+							FileItem EngineLibDependency = FileItem.GetItemByPath(AdditionalLibrary);
 							LinkAction.PrerequisiteItems.Add(EngineLibDependency);
 						}
 					}
@@ -784,7 +784,7 @@ namespace UnrealBuildTool
 			else
 			{
 				// Write the list of input files to a response file, with a tempfilename, on remote machine
-				string ResponsePath = Path.Combine(LinkEnvironment.Config.IntermediateDirectory, Path.GetFileName(OutputFile.AbsolutePath) + ".response");
+				FileReference ResponsePath = FileReference.Combine(LinkEnvironment.Config.IntermediateDirectory, Path.GetFileName(OutputFile.AbsolutePath) + ".response");
 				
 				// Never create response files when we are only generating IntelliSense data
 				if (!ProjectFileGenerator.bGenerateProjectFiles)
@@ -793,11 +793,11 @@ namespace UnrealBuildTool
 
 					if (BuildHostPlatform.Current.Platform != UnrealTargetPlatform.Mac)
 					{
-						RPCUtilHelper.CopyFile(ResponsePath, ConvertPath(ResponsePath), true);
+						RPCUtilHelper.CopyFile(ResponsePath.FullName, ConvertPath(ResponsePath.FullName), true);
 					}
 				}
 
-				LinkCommand += string.Format(" @\"{0}\"", ConvertPath(ResponsePath));
+				LinkCommand += string.Format(" @\"{0}\"", ConvertPath(ResponsePath.FullName));
 			}
 
 			if (LinkEnvironment.Config.bIsBuildingDLL)
@@ -878,7 +878,7 @@ namespace UnrealBuildTool
 			
 			LinkAction.ProducedItems.Add(RemoteOutputFile);
 
-			if (!Directory.Exists(LinkEnvironment.Config.IntermediateDirectory))
+			if (!LinkEnvironment.Config.IntermediateDirectory.Exists())
 			{
 				return OutputFile;
 			}
@@ -886,22 +886,22 @@ namespace UnrealBuildTool
 			if (!bIsBuildingLibrary)
 			{
 				// Prepare a script that will run later, once every dylibs and the executable are created. This script will be called by action created in FixDylibDependencies()
-				string FixDylibDepsScriptPath = Path.Combine(LinkEnvironment.Config.LocalShadowDirectory, "FixDylibDependencies.sh");
+				FileReference FixDylibDepsScriptPath = FileReference.Combine(LinkEnvironment.Config.LocalShadowDirectory, "FixDylibDependencies.sh");
 				if (!bHasWipedFixDylibScript)
 				{
-					if (File.Exists(FixDylibDepsScriptPath))
+					if (FixDylibDepsScriptPath.Exists())
 					{
-						File.Delete(FixDylibDepsScriptPath);
+						FixDylibDepsScriptPath.Delete();
 					}
 					bHasWipedFixDylibScript = true;
 				}
 
-				if (!Directory.Exists(LinkEnvironment.Config.LocalShadowDirectory))
+				if (!LinkEnvironment.Config.LocalShadowDirectory.Exists())
 				{
-					Directory.CreateDirectory(LinkEnvironment.Config.LocalShadowDirectory);
+					LinkEnvironment.Config.LocalShadowDirectory.CreateDirectory();
 				}
 
-				StreamWriter FixDylibDepsScript = File.AppendText(FixDylibDepsScriptPath);
+				StreamWriter FixDylibDepsScript = File.AppendText(FixDylibDepsScriptPath.FullName);
 
 				if (LinkEnvironment.Config.bIsCrossReferenced || !Utils.IsRunningOnMono)
 				{
@@ -923,17 +923,17 @@ namespace UnrealBuildTool
 				// Prepare a script that will be called by FinalizeAppBundle.sh to copy all necessary third party dylibs to the app bundle
 				// This is done this way as FinalizeAppBundle.sh script can be created before all the libraries are processed, so
 				// at the time of it's creation we don't have the full list of third party dylibs all modules need.
-				string DylibCopyScriptPath = Path.Combine(LinkEnvironment.Config.IntermediateDirectory, "DylibCopy.sh");
+				FileReference DylibCopyScriptPath = FileReference.Combine(LinkEnvironment.Config.IntermediateDirectory, "DylibCopy.sh");
 				if (!bHasWipedCopyDylibScript)
 				{
-					if (File.Exists(DylibCopyScriptPath))
+					if (DylibCopyScriptPath.Exists())
 					{
-						File.Delete(DylibCopyScriptPath);
+						DylibCopyScriptPath.Delete();
 					}
 					bHasWipedCopyDylibScript = true;
 				}
-				string ExistingScript = File.Exists(DylibCopyScriptPath) ? File.ReadAllText(DylibCopyScriptPath) : "";
-				StreamWriter DylibCopyScript = File.AppendText(DylibCopyScriptPath);
+				string ExistingScript = DylibCopyScriptPath.Exists() ? File.ReadAllText(DylibCopyScriptPath.FullName) : "";
+				StreamWriter DylibCopyScript = File.AppendText(DylibCopyScriptPath.FullName);
 				foreach (string Library in ThirdPartyLibraries)
 				{
 					string CopyCommandLineEntry = string.Format("cp -f \"{0}\" \"$1.app/Contents/MacOS\"", ConvertPath(Path.GetFullPath(Library)).Replace("$", "\\$"));
@@ -947,8 +947,8 @@ namespace UnrealBuildTool
 				// For non-console application, prepare a script that will create the app bundle. It'll be run by FinalizeAppBundle action
 				if (bIsBuildingAppBundle)
 				{
-					string FinalizeAppBundleScriptPath = Path.Combine(LinkEnvironment.Config.IntermediateDirectory, "FinalizeAppBundle.sh");
-					StreamWriter FinalizeAppBundleScript = File.CreateText(FinalizeAppBundleScriptPath);
+					FileReference FinalizeAppBundleScriptPath = FileReference.Combine(LinkEnvironment.Config.IntermediateDirectory, "FinalizeAppBundle.sh");
+					StreamWriter FinalizeAppBundleScript = File.CreateText(FinalizeAppBundleScriptPath.FullName);
 					AppendMacLine(FinalizeAppBundleScript, "#!/bin/sh");
 					string BinariesPath = Path.GetDirectoryName(OutputFile.AbsolutePath);
 					BinariesPath = Path.GetDirectoryName(BinariesPath.Substring(0, BinariesPath.IndexOf(".app")));
@@ -962,15 +962,15 @@ namespace UnrealBuildTool
 					AppendMacLine(FinalizeAppBundleScript, "mkdir -p \"{0}.app/Contents/Resources\"", ExeName);
 
 					// Copy third party dylibs by calling additional script prepared earlier
-					AppendMacLine(FinalizeAppBundleScript, "sh \"{0}\" \"{1}\"", ConvertPath(DylibCopyScriptPath).Replace("$", "\\$"), ExeName);
+					AppendMacLine(FinalizeAppBundleScript, "sh \"{0}\" \"{1}\"", ConvertPath(DylibCopyScriptPath.FullName).Replace("$", "\\$"), ExeName);
 
 					string IconName = "UE4";
 					string BundleVersion = ExeName.StartsWith("EpicGamesLauncher") ? LoadLauncherDisplayVersion() : LoadEngineDisplayVersion();
                     string EngineSourcePath = ConvertPath(Directory.GetCurrentDirectory()).Replace("$", "\\$");
-					string UProjectFilePath = UProjectInfo.GetProjectFilePath(GameName);
+					FileReference UProjectFilePath;
 					string CustomResourcesPath = "";
 					string CustomBuildPath = "";
-					if (string.IsNullOrEmpty(UProjectFilePath))
+					if(!UProjectInfo.TryGetProjectFileName(GameName, out UProjectFilePath))
 					{
 						string[] TargetFiles = Directory.GetFiles(Directory.GetCurrentDirectory(), GameName + ".Target.cs", SearchOption.AllDirectories);
 						if (TargetFiles.Length == 1)
@@ -986,9 +986,8 @@ namespace UnrealBuildTool
 					else
 					{
                         string ResourceParentFolderName = ExeName.StartsWith("EpicGamesLauncher") ? "Application" : GameName;
-						string FullUProjectFilePath = Path.GetFullPath(UProjectFilePath);
-						CustomResourcesPath = Path.GetDirectoryName(FullUProjectFilePath) + "/Source/" + ResourceParentFolderName + "/Resources/Mac";
-						CustomBuildPath = Path.GetDirectoryName(FullUProjectFilePath) + "/Build/Mac";
+						CustomResourcesPath = Path.GetDirectoryName(UProjectFilePath.FullName) + "/Source/" + ResourceParentFolderName + "/Resources/Mac";
+						CustomBuildPath = Path.GetDirectoryName(UProjectFilePath.FullName) + "/Build/Mac";
 					}
 
 					bool bBuildingEditor = GameName.EndsWith("Editor");
@@ -1014,7 +1013,7 @@ namespace UnrealBuildTool
 
 						if (CustomIcon != DefaultIcon)
 						{
-							QueueFileForBatchUpload(FileItem.GetItemByFullPath(Path.GetFullPath(CustomIcon)));
+							QueueFileForBatchUpload(FileItem.GetItemByFileReference(new FileReference(CustomIcon)));
 							CustomIcon = ConvertPath(CustomIcon);
 						}
 					}
@@ -1032,7 +1031,7 @@ namespace UnrealBuildTool
 					}
 					else
 					{
-						QueueFileForBatchUpload(FileItem.GetItemByFullPath(Path.GetFullPath(InfoPlistFile)));
+						QueueFileForBatchUpload(FileItem.GetItemByFileReference(new FileReference(InfoPlistFile)));
 						InfoPlistFile = ConvertPath(InfoPlistFile);
 					}
 					AppendMacLine(FinalizeAppBundleScript, "cp -f \"{0}\" \"{1}.app/Contents/Info.plist\"", InfoPlistFile, ExeName);
@@ -1054,10 +1053,10 @@ namespace UnrealBuildTool
 
 					// copy over some needed files
 					// @todo mac: Make a QueueDirectoryForBatchUpload
-					QueueFileForBatchUpload(FileItem.GetItemByFullPath(Path.GetFullPath("../../Engine/Source/Runtime/Launch/Resources/Mac/" + GameName + ".icns")));
-					QueueFileForBatchUpload(FileItem.GetItemByFullPath(Path.GetFullPath("../../Engine/Source/Runtime/Launch/Resources/Mac/UProject.icns")));
-					QueueFileForBatchUpload(FileItem.GetItemByFullPath(Path.GetFullPath("../../Engine/Source/Runtime/Launch/Resources/Mac/Info.plist")));
-					QueueFileForBatchUpload(FileItem.GetItemByFullPath(Path.Combine(LinkEnvironment.Config.IntermediateDirectory, "DylibCopy.sh")));
+					QueueFileForBatchUpload(FileItem.GetItemByFileReference(new FileReference("../../Engine/Source/Runtime/Launch/Resources/Mac/" + GameName + ".icns")));
+					QueueFileForBatchUpload(FileItem.GetItemByFileReference(new FileReference("../../Engine/Source/Runtime/Launch/Resources/Mac/UProject.icns")));
+					QueueFileForBatchUpload(FileItem.GetItemByFileReference(new FileReference("../../Engine/Source/Runtime/Launch/Resources/Mac/Info.plist")));
+					QueueFileForBatchUpload(FileItem.GetItemByFileReference(FileReference.Combine(LinkEnvironment.Config.IntermediateDirectory, "DylibCopy.sh")));
 				}
 			}
 
@@ -1074,14 +1073,14 @@ namespace UnrealBuildTool
 		FileItem FixDylibDependencies(LinkEnvironment LinkEnvironment, FileItem Executable)
 		{
 			Action LinkAction = new Action(ActionType.Link);
-			LinkAction.WorkingDirectory = Path.GetFullPath(".");
+			LinkAction.WorkingDirectory = UnrealBuildTool.EngineSourceDirectory.FullName;
 			LinkAction.CommandPath = "/bin/sh";
 			LinkAction.CommandDescription = "";
 
 			// Call the FixDylibDependencies.sh script which will link the dylibs and the main executable, this time proper ones, as it's called
 			// once all are already created, so the cross dependency problem no longer prevents linking.
 			// The script is deleted after it's executed so it's empty when we start appending link commands for the next executable.
-			FileItem FixDylibDepsScript = FileItem.GetItemByFullPath(Path.Combine(LinkEnvironment.Config.LocalShadowDirectory, "FixDylibDependencies.sh"));
+			FileItem FixDylibDepsScript = FileItem.GetItemByFileReference(FileReference.Combine(LinkEnvironment.Config.LocalShadowDirectory, "FixDylibDependencies.sh"));
 			FileItem RemoteFixDylibDepsScript = LocalToRemoteFileItem(FixDylibDepsScript, true);
 
 			LinkAction.CommandArguments = "-c 'chmod +x \"" + RemoteFixDylibDepsScript.AbsolutePath + "\"; \"" + RemoteFixDylibDepsScript.AbsolutePath + "\"; if [[ $? -ne 0 ]]; then exit 1; fi; ";
@@ -1103,7 +1102,7 @@ namespace UnrealBuildTool
 			LinkAction.StatusDescription = string.Format("Fixing dylib dependencies for {0}", Path.GetFileName(Executable.AbsolutePath));
 			LinkAction.bCanExecuteRemotely = false;
 
-			FileItem OutputFile = FileItem.GetItemByPath(Path.Combine(LinkEnvironment.Config.LocalShadowDirectory, Path.GetFileNameWithoutExtension(Executable.AbsolutePath) + ".link"));
+			FileItem OutputFile = FileItem.GetItemByFileReference(FileReference.Combine(LinkEnvironment.Config.LocalShadowDirectory, Path.GetFileNameWithoutExtension(Executable.AbsolutePath) + ".link"));
 			FileItem RemoteOutputFile = LocalToRemoteFileItem(OutputFile, false);
 
 			LinkAction.CommandArguments += "echo \"Dummy\" >> \"" + RemoteOutputFile.AbsolutePath + "\"";
@@ -1155,7 +1154,7 @@ namespace UnrealBuildTool
 				DebugOutputMap.Add (GenDebugAction, OutputFile.AbsolutePath);
 				GenDebugAction.ActionHandler = new Action.BlockingActionHandler(MacToolChain.RPCDebugInfoActionHandler);
 			}
-			GenDebugAction.WorkingDirectory = Path.GetFullPath(".");
+			GenDebugAction.WorkingDirectory = UnrealBuildTool.EngineSourceDirectory.FullName;
 			GenDebugAction.CommandPath = "sh";
 
 			// Deletes ay existing file on the building machine,
@@ -1186,12 +1185,12 @@ namespace UnrealBuildTool
 
 			// Make the compile action
 			Action FinalizeAppBundleAction = new Action(ActionType.CreateAppBundle);
-			FinalizeAppBundleAction.WorkingDirectory = Path.GetFullPath(".");
+			FinalizeAppBundleAction.WorkingDirectory = UnrealBuildTool.EngineSourceDirectory.FullName;
 			FinalizeAppBundleAction.CommandPath = "/bin/sh";
 			FinalizeAppBundleAction.CommandDescription = "";
 
 			// make path to the script
-			FileItem BundleScript = FileItem.GetItemByFullPath(Path.Combine(LinkEnvironment.Config.IntermediateDirectory, "FinalizeAppBundle.sh"));
+			FileItem BundleScript = FileItem.GetItemByFileReference(FileReference.Combine(LinkEnvironment.Config.IntermediateDirectory, "FinalizeAppBundle.sh"));
 			FileItem RemoteBundleScript = LocalToRemoteFileItem(BundleScript, true);
 
 			if (BuildHostPlatform.Current.Platform != UnrealTargetPlatform.Mac)
@@ -1211,7 +1210,7 @@ namespace UnrealBuildTool
 		FileItem CopyBundleResource(UEBuildBundleResource Resource, FileItem Executable)
 		{
 			Action CopyAction = new Action(ActionType.CreateAppBundle);
-			CopyAction.WorkingDirectory = Path.GetFullPath(".");
+			CopyAction.WorkingDirectory = UnrealBuildTool.EngineSourceDirectory.FullName;
 			CopyAction.CommandPath = "/bin/sh";
 			CopyAction.CommandDescription = "";
 
@@ -1245,12 +1244,12 @@ namespace UnrealBuildTool
 			{
 				foreach (string ResourceFile in Directory.GetFiles(Resource.ResourcePath, "*", SearchOption.AllDirectories))
 				{
-					QueueFileForBatchUpload(FileItem.GetItemByFullPath(Path.GetFullPath(ResourceFile)));
+					QueueFileForBatchUpload(FileItem.GetItemByFileReference(new FileReference(ResourceFile)));
 				}
 			}
 			else
 			{
-				QueueFileForBatchUpload(FileItem.GetItemByFullPath(SourcePath));
+				QueueFileForBatchUpload(FileItem.GetItemByFileReference(new FileReference(SourcePath)));
 			}
 
 			return TargetItem;
@@ -1262,7 +1261,7 @@ namespace UnrealBuildTool
 
 			foreach (UEBuildBinary Binary in Binaries)
 			{
-				BundleDependencies.Add(FileItem.GetItemByPath(Binary.ToString()));
+				BundleDependencies.Add(FileItem.GetItemByFileReference(Binary.Config.OutputFilePath));
 			}
 		}
 			
@@ -1270,31 +1269,31 @@ namespace UnrealBuildTool
 		{
 			base.FixBundleBinariesPaths(Target, Binaries);
 
-			string BundleContentsPath = Target.OutputPath + ".app/Contents/";
+			string BundleContentsPath = Target.OutputPath.FullName + ".app/Contents/";
 			foreach (UEBuildBinary Binary in Binaries)
 			{
-				string BinaryFileName = Path.GetFileName(Binary.Config.OutputFilePath);
+				string BinaryFileName = Path.GetFileName(Binary.Config.OutputFilePath.FullName);
 				if (BinaryFileName.EndsWith(".dylib"))
 				{
 					// Only dylibs from the same folder as the executable should be moved to the bundle. UE4Editor-*Game* dylibs and plugins will be loaded
 					// from their Binaries/Mac folders.
-					string DylibDir = Path.GetDirectoryName(Path.GetFullPath(Binary.Config.OutputFilePath));
-					string ExeDir = Path.GetDirectoryName(Path.GetFullPath(Target.OutputPath));
+					string DylibDir = Path.GetDirectoryName(Path.GetFullPath(Binary.Config.OutputFilePath.FullName));
+					string ExeDir = Path.GetDirectoryName(Path.GetFullPath(Target.OutputPath.FullName));
 					if (DylibDir.StartsWith(ExeDir))
 					{
 						// get the subdir, which is the DylibDir - ExeDir
 						string SubDir = DylibDir.Replace(ExeDir, "");
-						Binary.Config.OutputFilePaths[0] = BundleContentsPath + "MacOS" + SubDir + "/" + BinaryFileName;
+						Binary.Config.OutputFilePaths[0] = new FileReference(BundleContentsPath + "MacOS" + SubDir + "/" + BinaryFileName);
 					}
 				}
-				else if (!BinaryFileName.EndsWith(".a") && !Binary.Config.OutputFilePath.Contains(".app/Contents/MacOS/")) // Binaries can contain duplicates
+				else if (!BinaryFileName.EndsWith(".a") && !Binary.Config.OutputFilePath.FullName.Contains(".app/Contents/MacOS/")) // Binaries can contain duplicates
 				{
 					Binary.Config.OutputFilePaths[0] += ".app/Contents/MacOS/" + BinaryFileName;
 				}
 			}
 		}
 
-		static private string BundleContentsDirectory = "";
+		static private DirectoryReference BundleContentsDirectory;
 
         public override void AddFilesToReceipt(TargetReceipt Receipt, UEBuildBinary Binary)
 		{
@@ -1315,7 +1314,7 @@ namespace UnrealBuildTool
 				{
 					if(Receipt.BuildProducts[i].Type == BuildProductType.Executable || Receipt.BuildProducts[i].Type == BuildProductType.DynamicLibrary)
 					{
-						Receipt.AddBuildProduct(Path.ChangeExtension(Receipt.BuildProducts[i].Path, DebugExtension), BuildProductType.SymbolFile);
+						Receipt.AddBuildProduct(new FileReference(Path.ChangeExtension(Receipt.BuildProducts[i].Path, DebugExtension)), BuildProductType.SymbolFile);
 					}
 				}
 			}
@@ -1325,9 +1324,9 @@ namespace UnrealBuildTool
 				return;
 			}
 
-			if (BundleContentsDirectory == "" && Binary.Config.Type == UEBuildBinaryType.Executable)
+			if (BundleContentsDirectory == null && Binary.Config.Type == UEBuildBinaryType.Executable)
 			{
-				BundleContentsDirectory = Path.GetDirectoryName(Path.GetDirectoryName(Binary.Config.OutputFilePath)) + "/";
+				BundleContentsDirectory = Binary.Config.OutputFilePath.Directory.ParentDirectory;
 			}
 
 			// We need to know what third party dylibs would be copied to the bundle
@@ -1347,9 +1346,9 @@ namespace UnrealBuildTool
 					string LibName = Path.GetFileName(AdditionalLibrary);
 					if (LibName.StartsWith("lib"))
 					{
-						if (Path.GetExtension(AdditionalLibrary) == ".dylib" && !String.IsNullOrEmpty(BundleContentsDirectory))
+						if (Path.GetExtension(AdditionalLibrary) == ".dylib" && BundleContentsDirectory != null)
 						{
-							string Entry = BundleContentsDirectory + "MacOS/" + LibName;
+							FileReference Entry = FileReference.Combine(BundleContentsDirectory, "MacOS", LibName);
 							Receipt.AddBuildProduct(Entry, BuildProductType.DynamicLibrary);
 						}
 					}
@@ -1361,12 +1360,12 @@ namespace UnrealBuildTool
 					{
 						foreach (string ResourceFile in Directory.GetFiles(Resource.ResourcePath, "*", SearchOption.AllDirectories))
 						{
-							Receipt.AddBuildProduct(Path.Combine(BundleContentsDirectory, Resource.BundleContentsSubdir, ResourceFile.Substring(Path.GetDirectoryName(Resource.ResourcePath).Length + 1)), BuildProductType.RequiredResource);
+							Receipt.AddBuildProduct(FileReference.Combine(BundleContentsDirectory, Resource.BundleContentsSubdir, ResourceFile.Substring(Path.GetDirectoryName(Resource.ResourcePath).Length + 1)), BuildProductType.RequiredResource);
 						}
 					}
 					else
 					{
-						Receipt.AddBuildProduct(Path.Combine(BundleContentsDirectory, Resource.BundleContentsSubdir, Path.GetFileName(Resource.ResourcePath)), BuildProductType.RequiredResource);
+						Receipt.AddBuildProduct(FileReference.Combine(BundleContentsDirectory, Resource.BundleContentsSubdir, Path.GetFileName(Resource.ResourcePath)), BuildProductType.RequiredResource);
 					}
 				}
 			}
@@ -1374,29 +1373,29 @@ namespace UnrealBuildTool
 			if (Binary.Config.Type == UEBuildBinaryType.Executable)
 			{
 				// And we also need all the resources
-				Receipt.AddBuildProduct(BundleContentsDirectory + "Info.plist", BuildProductType.RequiredResource);
-				Receipt.AddBuildProduct(BundleContentsDirectory + "PkgInfo", BuildProductType.RequiredResource);
+				Receipt.AddBuildProduct(FileReference.Combine(BundleContentsDirectory, "Info.plist"), BuildProductType.RequiredResource);
+				Receipt.AddBuildProduct(FileReference.Combine(BundleContentsDirectory, "PkgInfo"), BuildProductType.RequiredResource);
 
 				if (Binary.Target.TargetType == TargetRules.TargetType.Editor)
 				{
-					Receipt.AddBuildProduct(BundleContentsDirectory + "Resources/UE4Editor.icns", BuildProductType.RequiredResource);
-					Receipt.AddBuildProduct(BundleContentsDirectory + "Resources/UProject.icns", BuildProductType.RequiredResource);
+					Receipt.AddBuildProduct(FileReference.Combine(BundleContentsDirectory, "Resources/UE4Editor.icns"), BuildProductType.RequiredResource);
+					Receipt.AddBuildProduct(FileReference.Combine(BundleContentsDirectory, "Resources/UProject.icns"), BuildProductType.RequiredResource);
 				}
 				else
 				{
-					Receipt.AddBuildProduct(BundleContentsDirectory + "Resources/" + Binary.Target.TargetName + ".icns", BuildProductType.RequiredResource);
+					Receipt.AddBuildProduct(FileReference.Combine(BundleContentsDirectory, "Resources/" + Binary.Target.TargetName + ".icns"), BuildProductType.RequiredResource);
 				}
 			}
 		}
 
 		// @todo Mac: Full implementation.
-		public override void CompileCSharpProject(CSharpEnvironment CompileEnvironment, string ProjectFileName, string DestinationFile)
+		public override void CompileCSharpProject(CSharpEnvironment CompileEnvironment, FileReference ProjectFileName, FileReference DestinationFile)
 		{
-			string ProjectDirectory = Path.GetDirectoryName(Path.GetFullPath(ProjectFileName));
+			string ProjectDirectory = Path.GetDirectoryName(ProjectFileName.FullName);
 
 			if (BuildHostPlatform.Current.Platform != UnrealTargetPlatform.Mac)
 			{
-				RPCUtilHelper.CopyFile(ProjectFileName, ConvertPath(ProjectFileName), true);
+				RPCUtilHelper.CopyFile(ProjectFileName.FullName, ConvertPath(ProjectFileName.FullName), true);
 				RPCUtilHelper.CopyFile("Engine/Source/Programs/DotNETCommon/MetaData.cs", ConvertPath("Engine/Source/Programs/DotNETCommon/MetaData.cs"), true);
 
 				string[] FileList = Directory.GetFiles(ProjectDirectory, "*.cs", SearchOption.AllDirectories);
@@ -1406,7 +1405,7 @@ namespace UnrealBuildTool
 				}
 			}
 
-			string XBuildArgs = "/verbosity:quiet /nologo /target:Rebuild /property:Configuration=Development /property:Platform=AnyCPU " + Path.GetFileName(ProjectFileName);
+			string XBuildArgs = "/verbosity:quiet /nologo /target:Rebuild /property:Configuration=Development /property:Platform=AnyCPU " + ProjectFileName.GetFileName();
 
 			if (BuildHostPlatform.Current.Platform != UnrealTargetPlatform.Mac)
 			{
diff --git a/Engine/Source/Programs/UnrealBuildTool/Mac/UEDeployMac.cs b/Engine/Source/Programs/UnrealBuildTool/Mac/UEDeployMac.cs
--- a/Engine/Source/Programs/UnrealBuildTool/Mac/UEDeployMac.cs
+++ b/Engine/Source/Programs/UnrealBuildTool/Mac/UEDeployMac.cs
@@ -39,7 +39,7 @@ namespace UnrealBuildTool.IOS
 
 			string RemoteWorkingDir = "";
 
-			bool bIsStaticLibrary = InTarget.OutputPath.EndsWith(".a");
+			bool bIsStaticLibrary = InTarget.OutputPath.HasExtension(".a");
 
 			if (BuildHostPlatform.Current.Platform != UnrealTargetPlatform.Mac)
 			{
@@ -95,7 +95,7 @@ namespace UnrealBuildTool.IOS
 
 					try
 					{
-						string BinaryDir = Path.GetDirectoryName(InTarget.OutputPath) + "\\";
+						string BinaryDir = InTarget.OutputPath.Directory + "\\";
 						if (BinaryDir.EndsWith(InTarget.AppName + "\\Binaries\\Mac\\") && InTarget.TargetType != TargetRules.TargetType.Game)
 						{
 							BinaryDir = BinaryDir.Replace(InTarget.TargetType.ToString(), "Game");
diff --git a/Engine/Source/Programs/UnrealBuildTool/System/ActionGraph.cs b/Engine/Source/Programs/UnrealBuildTool/System/ActionGraph.cs
--- a/Engine/Source/Programs/UnrealBuildTool/System/ActionGraph.cs
+++ b/Engine/Source/Programs/UnrealBuildTool/System/ActionGraph.cs
@@ -258,14 +258,14 @@ namespace UnrealBuildTool
 			// For all targets, build a set of all actions that are outdated.
 			var OutdatedActionDictionary = new Dictionary<Action, bool>();
 			var HistoryList = new List<ActionHistory>();
-			var OpenHistoryFiles = new HashSet<string>(StringComparer.InvariantCultureIgnoreCase);
+			var OpenHistoryFiles = new HashSet<FileReference>();
 			TargetToOutdatedPrerequisitesMap = new Dictionary<UEBuildTarget,List<FileItem>>();
 			foreach (var BuildTarget in Targets)	// @todo ubtmake: Optimization: Ideally we don't even need to know about targets for ubtmake -- everything would come from the files
 			{
 				var HistoryFilename = ActionHistory.GeneratePathForTarget(BuildTarget);
 				if (!OpenHistoryFiles.Contains(HistoryFilename))		// @todo ubtmake: Optimization: We should be able to move the command-line outdatedness and build product deletion over to the 'gather' phase, as the command-lines won't change between assembler runs
 				{
-					var History = new ActionHistory(HistoryFilename);
+					var History = new ActionHistory(HistoryFilename.FullName);
 					HistoryList.Add(History);
 					OpenHistoryFiles.Add(HistoryFilename);
 					GatherAllOutdatedActions(BuildTarget, History, ref OutdatedActionDictionary, TargetToOutdatedPrerequisitesMap);
@@ -1336,11 +1336,11 @@ namespace UnrealBuildTool
 					// Resolve the included file name to an actual file.
 					var DirectlyIncludedFiles =
 						DirectlyIncludedFilenames
-						.Where(DirectlyIncludedFilename => !string.IsNullOrEmpty(DirectlyIncludedFilename.IncludeResolvedName))
-						.Select(DirectlyIncludedFilename => DirectlyIncludedFilename.IncludeResolvedName)
+						.Where(DirectlyIncludedFilename => (DirectlyIncludedFilename.IncludeResolvedNameIfSuccessful != null))
+						.Select(DirectlyIncludedFilename => DirectlyIncludedFilename.IncludeResolvedNameIfSuccessful)
 						// Skip same include over and over (.inl files)
 						.Distinct()
-						.Select(FileItem.GetItemByFullPath)
+						.Select(FileItem.GetItemByFileReference)
 						.ToList();
 
 					OverriddenPrerequisites[ FileItem ] = DirectlyIncludedFiles;
diff --git a/Engine/Source/Programs/UnrealBuildTool/System/ActionHistory.cs b/Engine/Source/Programs/UnrealBuildTool/System/ActionHistory.cs
--- a/Engine/Source/Programs/UnrealBuildTool/System/ActionHistory.cs
+++ b/Engine/Source/Programs/UnrealBuildTool/System/ActionHistory.cs
@@ -92,28 +92,28 @@ namespace UnrealBuildTool
 		/// <summary>
 		/// Generates a full path to action history file for the specified target.
 		/// </summary>
-		public static string GeneratePathForTarget(UEBuildTarget Target)
+		public static FileReference GeneratePathForTarget(UEBuildTarget Target)
 		{
-			string Folder = null;
+			DirectoryReference Folder = null;
 			if (Target.ShouldCompileMonolithic() || Target.TargetType == TargetRules.TargetType.Program)
 			{
 				// Monolithic configs and programs have their Action History stored in their respective project folders
 				// or under engine intermediate folder + program name folder
-				string RootDirectory = UnrealBuildTool.GetUProjectPath();
-				if (String.IsNullOrEmpty(RootDirectory))
+				DirectoryReference RootDirectory = UnrealBuildTool.GetUProjectPath();
+				if (RootDirectory == null)
 				{
-					RootDirectory = Path.GetFullPath(BuildConfiguration.RelativeEnginePath);
+					RootDirectory = UnrealBuildTool.EngineDirectory;
 				}
-				Folder = Path.Combine(RootDirectory, BuildConfiguration.PlatformIntermediateFolder, Target.GetTargetName());
+				Folder = DirectoryReference.Combine(RootDirectory, BuildConfiguration.PlatformIntermediateFolder, Target.GetTargetName());
 			}
 			else
 			{
 				// Shared action history (unless this is a rocket target)
 				Folder = (UnrealBuildTool.RunningRocket() && UnrealBuildTool.HasUProjectFile()) ?
-					Path.Combine(UnrealBuildTool.GetUProjectPath(), BuildConfiguration.BaseIntermediateFolder) :
-					BuildConfiguration.BaseIntermediatePath;
+					DirectoryReference.Combine(UnrealBuildTool.GetUProjectPath(), BuildConfiguration.BaseIntermediateFolder) :
+					DirectoryReference.Combine(UnrealBuildTool.EngineDirectory, BuildConfiguration.BaseIntermediateFolder);
 			}
-			return Path.Combine(Folder, "ActionHistory.bin").Replace("\\", "/");
+			return FileReference.Combine(Folder, "ActionHistory.bin");
 		}
 	}
 }
diff --git a/Engine/Source/Programs/UnrealBuildTool/System/CMakefileGenerator.cs b/Engine/Source/Programs/UnrealBuildTool/System/CMakefileGenerator.cs
--- a/Engine/Source/Programs/UnrealBuildTool/System/CMakefileGenerator.cs
+++ b/Engine/Source/Programs/UnrealBuildTool/System/CMakefileGenerator.cs
@@ -23,7 +23,7 @@ namespace UnrealBuildTool
 
 	public class CMakefileProjectFile : ProjectFile
 	{
-		public CMakefileProjectFile( string InitFilePath )
+		public CMakefileProjectFile( FileReference InitFilePath )
 			: base(InitFilePath)
 		{
 		}
@@ -105,8 +105,8 @@ namespace UnrealBuildTool
 			{
 				CMakeGameRootPath = "set(GAME_ROOT_PATH \"" + UnrealBuildTool.GetUProjectPath() + "\")\n";
 
-				GameProjectPath = UnrealBuildTool.GetUProjectPath ();
-				GameProjectFile = UnrealBuildTool.GetUProjectFile ();
+				GameProjectPath = UnrealBuildTool.GetUProjectPath ().FullName;
+				GameProjectFile = UnrealBuildTool.GetUProjectFile ().FullName;
 
 				CMakeGameProjectFile = "set(GAME_PROJECT_FILE \"" + GameProjectFile + "\")\n";
 
@@ -139,7 +139,7 @@ namespace UnrealBuildTool
 			{
 				foreach (var CurPath in CurProject.IntelliSenseIncludeSearchPaths)
 				{
-					string IncludeDirectory = GetIncludeDirectory(CurPath, Path.GetDirectoryName(CurProject.ProjectFilePath));
+					string IncludeDirectory = GetIncludeDirectory(CurPath, Path.GetDirectoryName(CurProject.ProjectFilePath.FullName));
 					if (IncludeDirectory != null && !IncludeDirectories.Contains(IncludeDirectory)) 
 					{
 						IncludeDirectories.Add(IncludeDirectory);
@@ -163,13 +163,13 @@ namespace UnrealBuildTool
 
 			// Create SourceFiles, HeaderFiles, and ConfigFiles sections.
 			var AllModuleFiles = DiscoverModules(FindGameProjects());
-			foreach (string CurModuleFile in AllModuleFiles)
+			foreach (FileReference CurModuleFile in AllModuleFiles)
 			{
 				var FoundFiles = SourceFileSearch.FindModuleSourceFiles(CurModuleFile, ExcludeNoRedistFiles: bExcludeNoRedistFiles);
-				foreach (string CurSourceFile in FoundFiles)
+				foreach (FileReference CurSourceFile in FoundFiles)
 				{
 
-					string SourceFileRelativeToRoot = Utils.MakePathRelativeTo(CurSourceFile, Path.Combine(EngineRelativePath));
+					string SourceFileRelativeToRoot = CurSourceFile.MakeRelativeTo(UnrealBuildTool.EngineDirectory);
 					// Exclude files/folders on a per-platform basis.
 					if ((bIsLinux && IsLinuxFiltered(SourceFileRelativeToRoot)) || (bIsMac && IsMacFiltered(SourceFileRelativeToRoot))
 						||  (bIsWin64 && IsWinFiltered(SourceFileRelativeToRoot))
@@ -191,7 +191,7 @@ namespace UnrealBuildTool
 								}
 								else
 								{
-									CMakeSourceFilesList += ("\t\"${GAME_ROOT_PATH}/" + Utils.MakePathRelativeTo (CurSourceFile, GameProjectPath) + "\"\n");
+									CMakeSourceFilesList += ("\t\"${GAME_ROOT_PATH}/" + CurSourceFile.MakeRelativeTo(new DirectoryReference(Path.GetDirectoryName(GameProjectPath))) + "\"\n");
 								}
 							}
 						}
@@ -211,7 +211,7 @@ namespace UnrealBuildTool
 								}
 								else
 								{
-									CMakeHeaderFilesList += ("\t\"${GAME_ROOT_PATH}/" + Utils.MakePathRelativeTo (CurSourceFile, GameProjectPath) + "\"\n");
+									CMakeHeaderFilesList += ("\t\"${GAME_ROOT_PATH}/" + Utils.MakePathRelativeTo (CurSourceFile.FullName, GameProjectPath) + "\"\n");
 								}
 							}
 						}
@@ -232,7 +232,7 @@ namespace UnrealBuildTool
 								}
 								else
 								{
-									CMakeConfigFilesList += ("\t\"${GAME_ROOT_PATH}/" + Utils.MakePathRelativeTo (CurSourceFile, GameProjectPath) + "\"\n");
+									CMakeConfigFilesList += ("\t\"${GAME_ROOT_PATH}/" + Utils.MakePathRelativeTo (CurSourceFile.FullName, GameProjectPath) + "\"\n");
 								};
 							}
 						}
@@ -276,7 +276,7 @@ namespace UnrealBuildTool
 						continue;
 					}
 
-					var TargetName = Utils.GetFilenameWithoutAnyExtensions(TargetFile.TargetFilePath);		// Remove both ".cs" and ".
+					var TargetName = TargetFile.TargetFilePath.GetFileNameWithoutAnyExtensions();		// Remove both ".cs" and ".
 
 					foreach (UnrealTargetConfiguration CurConfiguration in Enum.GetValues(typeof(UnrealTargetConfiguration)))
 					{
@@ -305,7 +305,7 @@ namespace UnrealBuildTool
 				}
 			}
 
-			var FullFileName = Path.Combine(MasterProjectRelativePath, FileName);
+			var FullFileName = Path.Combine(MasterProjectPath.FullName, FileName);
 			return WriteFileIfChanged(FullFileName, CMakefileContent.ToString());
 		}
 
@@ -335,7 +335,7 @@ namespace UnrealBuildTool
 		/// Adds the include directory to the list, after converting it to relative to UE4 root
 		private string GetIncludeDirectory(string IncludeDir, string ProjectDir)
 		{
-			string FullProjectPath = Path.GetFullPath(ProjectFileGenerator.MasterProjectRelativePath);
+			string FullProjectPath = Path.GetFullPath(ProjectFileGenerator.MasterProjectPath.FullName);
 			string FullPath = "";
 			if (IncludeDir.StartsWith("/") && !IncludeDir.StartsWith(FullProjectPath))
 			{
@@ -370,13 +370,13 @@ namespace UnrealBuildTool
 		/// </summary>
 		/// <param name="InitFilePath">Path to the project file</param>
 		/// <returns>The newly allocated project file object</returns>
-		protected override ProjectFile AllocateProjectFile( string InitFilePath )
+		protected override ProjectFile AllocateProjectFile( FileReference InitFilePath )
 		{
 			return new CMakefileProjectFile( InitFilePath );
 		}
 
 		/// ProjectFileGenerator interface
-		public override void CleanProjectFiles(string InMasterProjectRelativePath, string InMasterProjectName, string InIntermediateProjectFilesPath)
+		public override void CleanProjectFiles(DirectoryReference InMasterProjectDirectory, string InMasterProjectName, DirectoryReference InIntermediateProjectFilesDirectory)
 		{
 		}
 	}
diff --git a/Engine/Source/Programs/UnrealBuildTool/System/CPPEnvironment.cs b/Engine/Source/Programs/UnrealBuildTool/System/CPPEnvironment.cs
--- a/Engine/Source/Programs/UnrealBuildTool/System/CPPEnvironment.cs
+++ b/Engine/Source/Programs/UnrealBuildTool/System/CPPEnvironment.cs
@@ -177,17 +177,17 @@ namespace UnrealBuildTool
 	public class CPPEnvironmentConfiguration : NativeBuildEnvironmentConfiguration
 	{
 		/** The directory to put the output object/debug files in. */
-		public string OutputDirectory = null;
+		public DirectoryReference OutputDirectory = null;
 
 		/** The directory to shadow source files in for syncing to remote compile servers */
-		public string LocalShadowDirectory = null;
+		public DirectoryReference LocalShadowDirectory = null;
 
 		/** PCH header file name as it appears in an #include statement in source code (might include partial, or no relative path.)
 		    This is needed by some compilers to use PCH features. */
 		public string PCHHeaderNameInCode;
 
 		/** The name of the header file which is precompiled. */
-		public string PrecompiledHeaderIncludeFilename = null;
+		public FileReference PrecompiledHeaderIncludeFilename = null;
 
 		/** Whether the compilation should create, use, or do nothing with the precompiled header. */
 		public PrecompiledHeaderAction PrecompiledHeaderAction = PrecompiledHeaderAction.None;
diff --git a/Engine/Source/Programs/UnrealBuildTool/System/CPPHeaders.cs b/Engine/Source/Programs/UnrealBuildTool/System/CPPHeaders.cs
--- a/Engine/Source/Programs/UnrealBuildTool/System/CPPHeaders.cs
+++ b/Engine/Source/Programs/UnrealBuildTool/System/CPPHeaders.cs
@@ -123,9 +123,10 @@ namespace UnrealBuildTool
 				int SearchAttempts = 0;
 				if( Path.IsPathRooted( RelativeIncludePath ) )
 				{
-					if( DirectoryLookupCache.FileExists( RelativeIncludePath ) )
+					FileReference Reference = new FileReference(RelativeIncludePath);
+					if( DirectoryLookupCache.FileExists( Reference ) )
 					{
-						Result = FileItem.GetItemByFullPath( RelativeIncludePath );
+						Result = FileItem.GetItemByFileReference( Reference );
 					}
 					++SearchAttempts;
 				}
@@ -144,17 +145,17 @@ namespace UnrealBuildTool
 						{
 							throw new BuildException(Exception, "Failed to combine null or invalid include paths.");
 						}
-						string FullFilePath = null;
+						FileReference FullFilePath = null;
 						try
 						{
-							FullFilePath = Path.GetFullPath(RelativeFilePath);
+							FullFilePath = new FileReference(RelativeFilePath);
 						}
 						catch (Exception)
 						{
 						}
 						if( FullFilePath != null && DirectoryLookupCache.FileExists( FullFilePath ) )
 						{
-							Result = FileItem.GetItemByFullPath( FullFilePath );
+							Result = FileItem.GetItemByFileReference( FullFilePath );
 							break;
 						}
 					}
@@ -210,7 +211,7 @@ namespace UnrealBuildTool
 			bool bUseFlatCPPIncludeDependencyCache = BuildConfiguration.bUseUBTMakefiles && UnrealBuildTool.IsAssemblingBuild;
 			if( bOnlyCachedDependencies && bUseFlatCPPIncludeDependencyCache )
 			{ 
-				Result = FlatCPPIncludeDependencyCache[Target].GetDependenciesForFile( SourceFile.AbsolutePath );
+				Result = FlatCPPIncludeDependencyCache[Target].GetDependenciesForFile( SourceFile.Reference );
 				if( Result == null )
 				{
 					// Nothing cached for this file!  It is new to us.  This is the expected flow when our CPPIncludeDepencencyCache is missing.
@@ -241,13 +242,13 @@ namespace UnrealBuildTool
 				// Update cache
 				if( bUseFlatCPPIncludeDependencyCache && !bOnlyCachedDependencies )
 				{ 
-					var Dependencies = new List<string>();
+					var Dependencies = new List<FileReference>();
 					foreach( var IncludedFile in Result )
 					{
-						Dependencies.Add( IncludedFile.AbsolutePath );
+						Dependencies.Add( IncludedFile.Reference );
 					}
-					string PCHName = SourceFile.PrecompiledHeaderIncludeFilename;
-					FlatCPPIncludeDependencyCache[Target].SetDependenciesForFile( SourceFile.AbsolutePath, PCHName, Dependencies );
+					FileReference PCHName = SourceFile.PrecompiledHeaderIncludeFilename;
+					FlatCPPIncludeDependencyCache[Target].SetDependenciesForFile( SourceFile.Reference, PCHName, Dependencies );
 				}
 			}
 
@@ -284,7 +285,7 @@ namespace UnrealBuildTool
 				{
 					// Resolve the included file name to an actual file.
 					DependencyInclude DirectInclude = DirectIncludes[DirectlyIncludedFileNameIndex];
-					if (DirectInclude.IncludeResolvedName == null || 
+					if (!DirectInclude.HasAttemptedResolve || 
 						// ignore any preexisting resolve cache if we are not configured to use it.
 						!BuildConfiguration.bUseIncludeDependencyResolveCache ||
 						// if we are testing the resolve cache, we force UBT to resolve every time to look for conflicts
@@ -299,14 +300,14 @@ namespace UnrealBuildTool
 						{
 							DirectlyIncludedFiles.Add(DirectIncludeResolvedFile);
 						}
-						IncludeDependencyCache[Target].CacheResolvedIncludeFullPath(CPPFile, DirectlyIncludedFileNameIndex, DirectIncludeResolvedFile != null ? DirectIncludeResolvedFile.AbsolutePath : "");
+						IncludeDependencyCache[Target].CacheResolvedIncludeFullPath(CPPFile, DirectlyIncludedFileNameIndex, DirectIncludeResolvedFile != null ? DirectIncludeResolvedFile.Reference : null);
 					}
 					else
 					{
 						// we might have cached an attempt to resolve the file, but couldn't actually find the file (system headers, etc).
-						if (DirectInclude.IncludeResolvedName != string.Empty)
+						if (DirectInclude.IncludeResolvedNameIfSuccessful != null)
 						{
-							DirectlyIncludedFiles.Add(FileItem.GetItemByFullPath(DirectInclude.IncludeResolvedName));
+							DirectlyIncludedFiles.Add(FileItem.GetItemByFileReference(DirectInclude.IncludeResolvedNameIfSuccessful));
 						}
 					}
 				}
@@ -499,7 +500,7 @@ namespace UnrealBuildTool
 				}
 
 				//@TODO: The intermediate exclusion is to work around autogenerated absolute paths in Module.SomeGame.cpp style files
-				bool bCheckForBackwardSlashes = FileToRead.StartsWith(InstalledFolder) || (UnrealBuildTool.HasUProjectFile() && Utils.IsFileUnderDirectory(FileToRead, UnrealBuildTool.GetUProjectPath()));
+				bool bCheckForBackwardSlashes = FileToRead.StartsWith(InstalledFolder) || (UnrealBuildTool.HasUProjectFile() && new FileReference(FileToRead).IsUnderDirectory(UnrealBuildTool.GetUProjectPath()));
 				if (bCheckForBackwardSlashes && !FileToRead.Contains("Intermediate") && !FileToRead.Contains("ThirdParty") && HeaderValue.IndexOf('\\', 0) >= 0)
 				{
 					throw new BuildException("In {0}: #include \"{1}\" contains backslashes ('\\'), please use forward slashes ('/') instead.", FileToRead, C.Value);
diff --git a/Engine/Source/Programs/UnrealBuildTool/System/CodeLiteGenerator.cs b/Engine/Source/Programs/UnrealBuildTool/System/CodeLiteGenerator.cs
--- a/Engine/Source/Programs/UnrealBuildTool/System/CodeLiteGenerator.cs
+++ b/Engine/Source/Programs/UnrealBuildTool/System/CodeLiteGenerator.cs
@@ -41,14 +41,14 @@ namespace UnrealBuildTool
 			var CodeCompletionFile = MasterProjectName + CodeCompletionFileName;
 			var CodeCompletionPreProcessorFile = MasterProjectName + CodeCompletionPreProcessorFileName;
 
-			var FullCodeLiteMasterFile = Path.Combine(MasterProjectRelativePath, SolutionFileName);
-			var FullCodeLiteCodeCompletionFile = Path.Combine(MasterProjectRelativePath, CodeCompletionFile);
-			var FullCodeLiteCodeCompletionPreProcessorFile = Path.Combine(MasterProjectRelativePath, CodeCompletionPreProcessorFile);
+			var FullCodeLiteMasterFile = Path.Combine(MasterProjectPath.FullName, SolutionFileName);
+			var FullCodeLiteCodeCompletionFile = Path.Combine(MasterProjectPath.FullName, CodeCompletionFile);
+			var FullCodeLiteCodeCompletionPreProcessorFile = Path.Combine(MasterProjectPath.FullName, CodeCompletionPreProcessorFile);
 
 			//
 			// HACK 
 			// TODO This is for now a hack. According to the original submitter, Eranif (a CodeLite developer) will support code completion folders in *.workspace files.
-			// We create a seperate file with all the folder name in it to copy manually into the code completion
+			// We create a separate file with all the folder name in it to copy manually into the code completion
 			// filed of CodeLite workspace. (Workspace Settings/Code Completion -> copy the content of the file threre.)
 			List<string> IncludeDirectories = new List<string>();
 			List<string> PreProcessor = new List<string>();
@@ -63,15 +63,15 @@ namespace UnrealBuildTool
 
 				foreach (var CurrentPath in Project.IntelliSenseIncludeSearchPaths) 
 				{
-					// Convert relative path into abosulte.
-					string IntelliSenseIncludeSearchPath = Path.GetFullPath(Path.Combine(Path.GetDirectoryName(Path.GetFullPath(Project.ProjectFilePath)), CurrentPath));
-					IncludeDirectories.Add(IntelliSenseIncludeSearchPath);
+					// Convert relative path into absolute.
+					DirectoryReference IntelliSenseIncludeSearchPath = DirectoryReference.Combine(Project.ProjectFilePath.Directory, CurrentPath);
+					IncludeDirectories.Add(IntelliSenseIncludeSearchPath.FullName);
 				}
 				foreach (var CurrentPath in Project.IntelliSenseSystemIncludeSearchPaths) 
 				{
-					// Convert relative path into abosulte.
-					string IntelliSenseSystemIncludeSearchPath = Path.GetFullPath(Path.Combine(Path.GetDirectoryName(Path.GetFullPath(Project.ProjectFilePath)), CurrentPath));
-					IncludeDirectories.Add (IntelliSenseSystemIncludeSearchPath);
+					// Convert relative path into absolute.
+					DirectoryReference IntelliSenseSystemIncludeSearchPath = DirectoryReference.Combine(Project.ProjectFilePath.Directory, CurrentPath);
+					IncludeDirectories.Add (IntelliSenseSystemIncludeSearchPath.FullName);
 				}
 
 				foreach( var CurDef in Project.IntelliSensePreprocessorDefinitions )
@@ -125,7 +125,7 @@ namespace UnrealBuildTool
 			//
 			foreach( var CurProject in AllProjectFiles )
 			{
-				var ProjectExtension = Path.GetExtension (CurProject.ProjectFilePath);
+				var ProjectExtension = CurProject.ProjectFilePath.GetExtension();
 
 				//
 				// TODO For now ignore C# project files.
@@ -141,7 +141,7 @@ namespace UnrealBuildTool
 				foreach (ProjectTarget CurrentTarget in CurProject.ProjectTargets) 
 				{
 					string[] tmp = CurrentTarget.ToString ().Split ('.');
-					string ProjectTargetFileName = Path.GetDirectoryName (CurProject.RelativeProjectFilePath) + "/" + tmp [0] + ProjectExtension;
+					string ProjectTargetFileName = CurProject.ProjectFilePath.Directory.MakeRelativeTo(MasterProjectPath) + "/" + tmp [0] + ProjectExtension;
 					String ProjectName = tmp [0];
 
 
@@ -173,7 +173,7 @@ namespace UnrealBuildTool
 
 					foreach( var CurProject in AllProjectFiles )
 					{
-						var ProjectExtension = Path.GetExtension (CurProject.ProjectFilePath);
+						var ProjectExtension = CurProject.ProjectFilePath.GetExtension();
 
 						//
 						// TODO For now ignore C# project files.
@@ -206,7 +206,7 @@ namespace UnrealBuildTool
 			return true;
 		}
 
-		protected override ProjectFile AllocateProjectFile( string InitFilePath )
+		protected override ProjectFile AllocateProjectFile( FileReference InitFilePath )
 		{
 			return new CodeLiteProject( InitFilePath );
 		}
@@ -216,40 +216,40 @@ namespace UnrealBuildTool
 			return new CodeLiteFolder( InitOwnerProjectFileGenerator, InitFolderName );
 		}
 
-		public override void CleanProjectFiles(string InMasterProjectRelativePath, string InMasterProjectName, string InIntermediateProjectFilesPath)
+		public override void CleanProjectFiles(DirectoryReference InMasterProjectDirectory, string InMasterProjectName, DirectoryReference InIntermediateProjectFilesDirectory)
 		{
 			// TODO Delete all files here. Not finished yet.
 			var SolutionFileName = InMasterProjectName + SolutionExtension;
 			var CodeCompletionFile = InMasterProjectName + CodeCompletionFileName;
 			var CodeCompletionPreProcessorFile = InMasterProjectName + CodeCompletionPreProcessorFileName;
 
-			var FullCodeLiteMasterFile = Path.Combine(InMasterProjectRelativePath, SolutionFileName);
-			var FullCodeLiteCodeCompletionFile = Path.Combine(InMasterProjectRelativePath, CodeCompletionFile);
-			var FullCodeLiteCodeCompletionPreProcessorFile = Path.Combine(InMasterProjectRelativePath, CodeCompletionPreProcessorFile);
+			FileReference FullCodeLiteMasterFile = FileReference.Combine(InMasterProjectDirectory, SolutionFileName);
+			FileReference FullCodeLiteCodeCompletionFile = FileReference.Combine(InMasterProjectDirectory, CodeCompletionFile);
+			FileReference FullCodeLiteCodeCompletionPreProcessorFile = FileReference.Combine(InMasterProjectDirectory, CodeCompletionPreProcessorFile);
 
-			if (File.Exists(FullCodeLiteMasterFile))
+			if (FullCodeLiteMasterFile.Exists())
 			{
-				File.Delete(FullCodeLiteMasterFile);
+				FullCodeLiteMasterFile.Delete();
 			}
-			if (File.Exists(FullCodeLiteCodeCompletionFile))
+			if (FullCodeLiteCodeCompletionFile.Exists())
 			{
-				File.Delete(FullCodeLiteCodeCompletionFile);
+				FullCodeLiteCodeCompletionFile.Delete();
 			}
-			if (File.Exists(FullCodeLiteCodeCompletionPreProcessorFile))
+			if (FullCodeLiteCodeCompletionPreProcessorFile.Exists())
 			{
-				File.Delete(FullCodeLiteCodeCompletionPreProcessorFile);
+				FullCodeLiteCodeCompletionPreProcessorFile.Delete();
 			}
 
 			// Delete the project files folder
-			if (Directory.Exists(InIntermediateProjectFilesPath))
+			if (InIntermediateProjectFilesDirectory.Exists())
 			{
 				try
 				{
-					Directory.Delete(InIntermediateProjectFilesPath, true);
+					Directory.Delete(InIntermediateProjectFilesDirectory.FullName, true);
 				}
 				catch (Exception Ex)
 				{
-					Log.TraceInformation("Error while trying to clean project files path {0}. Ignored.", InIntermediateProjectFilesPath);
+					Log.TraceInformation("Error while trying to clean project files path {0}. Ignored.", InIntermediateProjectFilesDirectory);
 					Log.TraceInformation("\t" + Ex.Message);
 				}
 			}
diff --git a/Engine/Source/Programs/UnrealBuildTool/System/CodeLiteProject.cs b/Engine/Source/Programs/UnrealBuildTool/System/CodeLiteProject.cs
--- a/Engine/Source/Programs/UnrealBuildTool/System/CodeLiteProject.cs
+++ b/Engine/Source/Programs/UnrealBuildTool/System/CodeLiteProject.cs
@@ -11,7 +11,7 @@ namespace UnrealBuildTool
 {
 	public class CodeLiteProject : ProjectFile
 	{
-		public CodeLiteProject( string InitFilePath ) : base(InitFilePath)
+		public CodeLiteProject( FileReference InitFilePath ) : base(InitFilePath)
 		{
 		}
 		
@@ -34,10 +34,9 @@ namespace UnrealBuildTool
 		{
 			bool bSuccess = false;
 
-			string ProjectPath = ProjectFilePath;
-			string ProjectExtension = Path.GetExtension (ProjectFilePath);
+			string ProjectPath = ProjectFilePath.FullName;
+			string ProjectExtension = ProjectFilePath.GetExtension();
 			string ProjectPlatformName = BuildHostPlatform.Current.Platform.ToString();
-			string ProjectRelativeFilePath = this.RelativeProjectFilePath;
 
 			// Get the output directory
 			string EngineRootDirectory = Path.GetFullPath(ProjectFileGenerator.EngineRelativePath);
@@ -49,7 +48,7 @@ namespace UnrealBuildTool
 			string GameWorkingDirectory = "";
 			if (UnrealBuildTool.HasUProjectFile ()) 
 			{
-				GameWorkingDirectory = Path.Combine (Path.GetDirectoryName (UnrealBuildTool.GetUProjectFile ()), "Binaries", ProjectPlatformName);
+				GameWorkingDirectory = Path.Combine (Path.GetDirectoryName (UnrealBuildTool.GetUProjectFile ().FullName), "Binaries", ProjectPlatformName);
 			}
 			//
 			// Build the working directory of the UE4Editor executable.
@@ -59,13 +58,13 @@ namespace UnrealBuildTool
 			//
 			// Create the folder where the project files goes if it does not exist
 			//
-			String FilePath = Path.GetDirectoryName(ProjectFilePath);
+			String FilePath = Path.GetDirectoryName(ProjectFilePath.FullName);
 			if( (FilePath.Length > 0) && !Directory.Exists(FilePath))
 			{
 				Directory.CreateDirectory(FilePath);
 			}
 
-			string GameProjectFile = UnrealBuildTool.GetUProjectFile();
+			string GameProjectFile = UnrealBuildTool.GetUProjectFile().FullName;
 
 			//
 			// Write all targets which will be separate projects.
@@ -73,7 +72,7 @@ namespace UnrealBuildTool
 			foreach (ProjectTarget target in ProjectTargets) 
 			{
 				string[] tmp = target.ToString ().Split ('.');
-				string ProjectTargetFileName = Path.GetDirectoryName (ProjectFilePath) + "/" + tmp [0] +  ProjectExtension;
+				string ProjectTargetFileName = Path.GetDirectoryName (ProjectFilePath.FullName) + "/" + tmp [0] +  ProjectExtension;
 				String ProjectName = tmp [0];
 				var ProjectTargetType = target.TargetRules.Type;
 
@@ -89,12 +88,12 @@ namespace UnrealBuildTool
 				// TODO Maybe skipping those files directly in the following foreach loop is faster?
 				//
 				List<SourceFile> FilterSourceFile = SourceFiles.FindAll(s => (	
-					Path.GetExtension(s.FilePath).Equals(".h") 
-					|| Path.GetExtension(s.FilePath).Equals(".cpp")
-					|| Path.GetExtension(s.FilePath).Equals(".cs")
-					|| Path.GetExtension(s.FilePath).Equals(".uproject")
-					|| Path.GetExtension(s.FilePath).Equals(".ini")
-					|| Path.GetExtension(s.FilePath).Equals(".usf")
+					s.Reference.HasExtension(".h") 
+					|| s.Reference.HasExtension(".cpp")
+					|| s.Reference.HasExtension(".cs")
+					|| s.Reference.HasExtension(".uproject")
+					|| s.Reference.HasExtension(".ini")
+					|| s.Reference.HasExtension(".usf")
 				));
 
 				//
@@ -114,7 +113,7 @@ namespace UnrealBuildTool
 						if(ProjectName.Contains("UE4"))
 						{
 							int Idx = Path.GetFullPath(ProjectFileGenerator.EngineRelativePath).Length;
-							CurrentFilePath = Path.GetDirectoryName(Path.GetFullPath(CurrentFile.FilePath)).Substring(Idx);
+							CurrentFilePath = Path.GetDirectoryName(Path.GetFullPath(CurrentFile.Reference.FullName)).Substring(Idx);
 						}
 						else
 						{
@@ -124,8 +123,8 @@ namespace UnrealBuildTool
 							{
 								ProjectNameRaw = ProjectName.Substring(0, IdxProjectName);
 							}
-							int Idx = Path.GetDirectoryName(Path.GetFullPath(CurrentFile.FilePath)).IndexOf(ProjectNameRaw) + ProjectNameRaw.Length;
-							CurrentFilePath = Path.GetDirectoryName(Path.GetFullPath(CurrentFile.FilePath)).Substring(Idx);
+							int Idx = Path.GetDirectoryName(CurrentFile.Reference.FullName).IndexOf(ProjectNameRaw) + ProjectNameRaw.Length;
+							CurrentFilePath = Path.GetDirectoryName(CurrentFile.Reference.FullName).Substring(Idx);
 						}
 					}
 					else if (ProjectTargetType == TargetRules.TargetType.Program)
@@ -133,16 +132,16 @@ namespace UnrealBuildTool
 						//
 						// We do not need all the editors subfolders to show the content. Find the correct programs subfolder.
 						//
-						int Idx = Path.GetDirectoryName(Path.GetFullPath(CurrentFile.FilePath)).IndexOf(ProjectName) + ProjectName.Length;
-						CurrentFilePath = Path.GetFullPath(Path.GetDirectoryName(CurrentFile.FilePath)).Substring(Idx);
+						int Idx = Path.GetDirectoryName(CurrentFile.Reference.FullName).IndexOf(ProjectName) + ProjectName.Length;
+						CurrentFilePath = Path.GetDirectoryName(CurrentFile.Reference.FullName).Substring(Idx);
 					}
 					else if (ProjectTargetType == TargetRules.TargetType.Game)
 					{
 //						int lengthOfProjectRootPath = Path.GetFullPath(ProjectFileGenerator.MasterProjectRelativePath).Length;
 //						CurrentFilePath = Path.GetDirectoryName(Path.GetFullPath(CurrentFile.FilePath)).Substring(lengthOfProjectRootPath);
 					//	int lengthOfProjectRootPath = EngineRootDirectory.Length;
-						int Idx = Path.GetDirectoryName(Path.GetFullPath(CurrentFile.FilePath)).IndexOf(ProjectName) + ProjectName.Length;
-						CurrentFilePath = Path.GetDirectoryName(Path.GetFullPath(CurrentFile.FilePath)).Substring(Idx);
+						int Idx = Path.GetDirectoryName(CurrentFile.Reference.FullName).IndexOf(ProjectName) + ProjectName.Length;
+						CurrentFilePath = Path.GetDirectoryName(CurrentFile.Reference.FullName).Substring(Idx);
 					}
 
 					string [] SplitFolders = CurrentFilePath.Split('/');
@@ -218,7 +217,7 @@ namespace UnrealBuildTool
 					// If we are at this point we found the correct XElement folder
 					//
 					XElement file = new XElement("File");
-					XAttribute fileAttribute = new XAttribute("Name",  Path.GetFullPath(CurrentFile.FilePath));
+					XAttribute fileAttribute = new XAttribute("Name",  CurrentFile.Reference.FullName);
 					file.Add(fileAttribute);
 					root.Add(file); 
 				}
diff --git a/Engine/Source/Programs/UnrealBuildTool/System/DependencyCache.cs b/Engine/Source/Programs/UnrealBuildTool/System/DependencyCache.cs
--- a/Engine/Source/Programs/UnrealBuildTool/System/DependencyCache.cs
+++ b/Engine/Source/Programs/UnrealBuildTool/System/DependencyCache.cs
@@ -25,10 +25,16 @@ namespace UnrealBuildTool
 		/// These are direct include paths and cannot be resolved to an actual file on disk without using the proper list of include directories for this file's module 
 		/// </summary>
 		public readonly string IncludeName;
+
+		/// <summary>
+		/// Whether we've attempted to resolve this include (may be set even if IncludeResolvedNameIfSuccessful = null, in cases where it couldn't be found)
+		/// </summary>
+		public bool HasAttemptedResolve;
+
 		/// <summary>
-		/// This is the fully resolved name. We can't really store this globally, but we're going to see how well it works.
+		/// This is the fully resolved name, and a bool for whether we've attempted the resolve but failed. We can't really store this globally, but we're going to see how well it works.
 		/// </summary>
-		public string IncludeResolvedName;
+		public FileReference IncludeResolvedNameIfSuccessful;
 
 		public override string ToString()
 		{
@@ -63,13 +69,13 @@ namespace UnrealBuildTool
 
 		/** Path to store the cache data to. */
 		[NonSerialized]
-		private string CachePath;
+		private FileReference CachePath;
 
 		/** Dependency lists, keyed (case-insensitively) on file's absolute path. */
-		private Dictionary<string, DependencyInfo> DependencyMap;
+		private Dictionary<FileReference, DependencyInfo> DependencyMap;
 
 		[NonSerialized]
-		private Dictionary<string, bool> FileExistsInfo;
+		private Dictionary<FileReference, bool> FileExistsInfo;
 
 		/** Whether the dependency cache is dirty and needs to be saved. */
 		[NonSerialized]
@@ -80,10 +86,10 @@ namespace UnrealBuildTool
 		 * 
 		 * @param	CachePath	Name of the cache file to deserialize
 		 */
-		public static DependencyCache Create(string CachePath)
+		public static DependencyCache Create(FileReference CachePath)
 		{
 			// See whether the cache file exists.
-			FileItem Cache = FileItem.GetItemByPath(CachePath);
+			FileItem Cache = FileItem.GetItemByFileReference(CachePath);
 			if (Cache.bExists)
 			{
 				if (BuildConfiguration.bPrintPerformanceInfo)
@@ -149,8 +155,9 @@ namespace UnrealBuildTool
 			{
 				// Don't bother logging this expected error.
 				// It's due to a change in the CacheCreateDate type.
-				if (Ex.Message != "Object of type 'System.DateTime' cannot be converted to type 'System.DateTimeOffset'" &&
-	Ex.Message != "Object of type 'System.Collections.Generic.Dictionary`2[System.String,System.Collections.Generic.List`1[System.String]]' cannot be converted to type 'System.Collections.Generic.Dictionary`2[System.String,UnrealBuildTool.DependencyInfo]'.")	// To catch serialization differences added when we added the DependencyInfo struct
+				if (Ex.Message != "Object of type 'System.Collections.Generic.Dictionary`2[System.String,UnrealBuildTool.DependencyInfo]' cannot be converted to type 'System.Collections.Generic.Dictionary`2[UnrealBuildTool.FileReference,UnrealBuildTool.DependencyInfo]'."
+					&& Ex.Message != "Object of type 'System.DateTime' cannot be converted to type 'System.DateTimeOffset'" 
+					&& Ex.Message != "Object of type 'System.Collections.Generic.Dictionary`2[System.String,System.Collections.Generic.List`1[System.String]]' cannot be converted to type 'System.Collections.Generic.Dictionary`2[System.String,UnrealBuildTool.DependencyInfo]'.")	// To catch serialization differences added when we added the DependencyInfo struct
 				{
 					Console.Error.WriteLine("Failed to read dependency cache: {0}", Ex.Message);
 				}
@@ -167,8 +174,8 @@ namespace UnrealBuildTool
 		{
 			CacheCreateDate = DateTimeOffset.Now;
 			CacheUpdateDate = DateTimeOffset.Now;
-			CachePath = Cache.AbsolutePath;
-			DependencyMap = new Dictionary<string, DependencyInfo>();
+			CachePath = Cache.Reference;
+			DependencyMap = new Dictionary<FileReference, DependencyInfo>();
 			bIsDirty = false;
 			CreateFileExistsInfo();
 		}
@@ -189,8 +196,8 @@ namespace UnrealBuildTool
 				// Serialize the cache to disk.
 				try
 				{
-					Directory.CreateDirectory(Path.GetDirectoryName(CachePath));
-					using (FileStream Stream = new FileStream(CachePath, FileMode.Create, FileAccess.Write))
+					CachePath.Directory.CreateDirectory();
+					using (FileStream Stream = new FileStream(CachePath.FullName, FileMode.Create, FileAccess.Write))
 					{
 						BinaryFormatter Formatter = new BinaryFormatter();
 						Formatter.Serialize(Stream, this);
@@ -215,11 +222,12 @@ namespace UnrealBuildTool
 				}
 			}
 
-			if(File.Exists(CachePath + ".buildmutex"))
+			FileReference MutexPath = CachePath + ".buildmutex";
+			if(MutexPath.Exists())
 			{
 				try
 				{
-					File.Delete(CachePath + ".buildmutex");
+					MutexPath.Delete();
 				}
 				catch
 				{
@@ -236,11 +244,9 @@ namespace UnrealBuildTool
 		 */
 		public List<DependencyInclude> GetCachedDependencyInfo(FileItem File)
 		{
-			string LowercaseFilePath = File.AbsolutePath.ToLowerInvariant();
-
 			// Check whether File is in cache.
 			DependencyInfo Result;
-			if (!DependencyMap.TryGetValue(LowercaseFilePath, out Result))
+			if (!DependencyMap.TryGetValue(File.Reference, out Result))
 			{
 				return null;
 			}
@@ -249,7 +255,7 @@ namespace UnrealBuildTool
 			if (File.LastWriteTime >= CacheCreateDate)
 			{
 				// Remove entry from cache as it's stale.
-				DependencyMap.Remove(LowercaseFilePath);
+				DependencyMap.Remove(File.Reference);
 				bIsDirty = true;
 				return null;
 			}
@@ -257,21 +263,20 @@ namespace UnrealBuildTool
 			// Check if any of the resolved includes is missing
 			foreach (var Include in Result.Includes)
 			{
-				if (!String.IsNullOrEmpty(Include.IncludeResolvedName))
+				if (Include.IncludeResolvedNameIfSuccessful != null)
 				{
 					bool bIncludeExists = false;
-					string LowercaseIncludePath = Include.IncludeResolvedName.ToLowerInvariant();
-					if (!FileExistsInfo.TryGetValue(LowercaseIncludePath, out bIncludeExists))
+					if (!FileExistsInfo.TryGetValue(Include.IncludeResolvedNameIfSuccessful, out bIncludeExists))
 					{
-						bIncludeExists = System.IO.File.Exists(Include.IncludeResolvedName);
-						FileExistsInfo.Add(LowercaseIncludePath, bIncludeExists);
+						bIncludeExists = Include.IncludeResolvedNameIfSuccessful.Exists();
+						FileExistsInfo.Add(Include.IncludeResolvedNameIfSuccessful, bIncludeExists);
 					}
 
 					if (!bIncludeExists)
 					{
 						// Remove entry from cache as it's stale, as well as the include which no longer exists
-						DependencyMap.Remove(LowercaseIncludePath);
-						DependencyMap.Remove(LowercaseFilePath);
+						DependencyMap.Remove(Include.IncludeResolvedNameIfSuccessful);
+						DependencyMap.Remove(File.Reference);
 						bIsDirty = true;
 						return null;
 					}
@@ -291,7 +296,7 @@ namespace UnrealBuildTool
 		 */
 		public void SetDependencyInfo(FileItem File, List<DependencyInclude> Info)
 		{
-			DependencyMap[File.AbsolutePath.ToLowerInvariant()] = new DependencyInfo{ Includes = Info, HasUObjects = false };
+			DependencyMap[File.Reference] = new DependencyInfo{ Includes = Info, HasUObjects = false };
 			bIsDirty = true;
 		}
 
@@ -302,7 +307,7 @@ namespace UnrealBuildTool
 		{
 			// Pre-allocate about 125% of the dependency map count (which amounts to 1.25 unique includes per dependency which is a little more than empirical
 			// results show but gives us some slack in case something gets added).
-			FileExistsInfo = new Dictionary<string, bool>((DependencyMap.Count * 5) / 4);
+			FileExistsInfo = new Dictionary<FileReference, bool>((DependencyMap.Count * 5) / 4);
 		}
 
 		/// <summary>
@@ -314,9 +319,9 @@ namespace UnrealBuildTool
 			{
 				foreach (var Include in Dependency.Value.Includes)
 				{
-					if (Include.IncludeResolvedName == String.Empty)
+					if (Include.HasAttemptedResolve && Include.IncludeResolvedNameIfSuccessful == null)
 					{
-						Include.IncludeResolvedName = null;
+						Include.HasAttemptedResolve = false;
 					}
 				}
 			}
@@ -330,23 +335,24 @@ namespace UnrealBuildTool
 		/// <param name="File">The file whose include is being resolved</param>
 		/// <param name="DirectlyIncludedFileNameIndex">Index in the resolve list to quickly find the include in question in the existing cache.</param>
 		/// <param name="DirectlyIncludedFileNameFullPath">Full path name of the resolve include.</param>
-		public void CacheResolvedIncludeFullPath(FileItem File, int DirectlyIncludedFileNameIndex, string DirectlyIncludedFileNameFullPath)
+		public void CacheResolvedIncludeFullPath(FileItem File, int DirectlyIncludedFileNameIndex, FileReference DirectlyIncludedFileNameFullPath)
 		{
 			if (BuildConfiguration.bUseIncludeDependencyResolveCache)
 			{
-				var Includes = DependencyMap[File.AbsolutePath.ToLowerInvariant()].Includes;
+				var Includes = DependencyMap[File.Reference].Includes;
 				var IncludeToResolve = Includes[DirectlyIncludedFileNameIndex];
 				if (BuildConfiguration.bTestIncludeDependencyResolveCache)
 				{
 					// test whether there are resolve conflicts between modules with different include paths.
-					if (IncludeToResolve.IncludeResolvedName != null && IncludeToResolve.IncludeResolvedName != DirectlyIncludedFileNameFullPath)
+					if (IncludeToResolve.HasAttemptedResolve && IncludeToResolve.IncludeResolvedNameIfSuccessful != DirectlyIncludedFileNameFullPath)
 					{
 						throw new BuildException("Found directly included file that resolved differently in different modules. File ({0}) had previously resolved to ({1}) and now resolves to ({2}).",
-							File.AbsolutePath, IncludeToResolve.IncludeResolvedName, DirectlyIncludedFileNameFullPath);
+							File.AbsolutePath, IncludeToResolve.IncludeResolvedNameIfSuccessful, DirectlyIncludedFileNameFullPath);
 					}
 				}
-				Includes[DirectlyIncludedFileNameIndex].IncludeResolvedName = DirectlyIncludedFileNameFullPath;
-				if (!String.IsNullOrEmpty(DirectlyIncludedFileNameFullPath))
+				Includes[DirectlyIncludedFileNameIndex].HasAttemptedResolve = true;
+				Includes[DirectlyIncludedFileNameIndex].IncludeResolvedNameIfSuccessful = DirectlyIncludedFileNameFullPath;
+				if (DirectlyIncludedFileNameFullPath != null)
 				{
 					bIsDirty = true;
 				}
@@ -358,15 +364,18 @@ namespace UnrealBuildTool
 		/// </summary>
 		/// <param name="Target">Current build target</param>
 		/// <returns>Cache Path</returns>
-		public static string GetDependencyCachePathForTarget(UEBuildTarget Target)
+		public static FileReference GetDependencyCachePathForTarget(UEBuildTarget Target)
 		{
-			string PlatformIntermediatePath = BuildConfiguration.PlatformIntermediatePath;
+			DirectoryReference PlatformIntermediatePath;
 			if (UnrealBuildTool.HasUProjectFile())
 			{
-				PlatformIntermediatePath = Path.Combine(UnrealBuildTool.GetUProjectPath(), BuildConfiguration.PlatformIntermediateFolder);
+				PlatformIntermediatePath = DirectoryReference.Combine(UnrealBuildTool.GetUProjectPath(), BuildConfiguration.PlatformIntermediateFolder);
+			}
+			else
+			{
+				PlatformIntermediatePath = DirectoryReference.Combine(UnrealBuildTool.EngineDirectory, BuildConfiguration.PlatformIntermediatePath);
 			}
-			string CachePath = Path.Combine(PlatformIntermediatePath, Target.GetTargetName(), "DependencyCache.bin");
-			return CachePath;
+			return FileReference.Combine(PlatformIntermediatePath, Target.GetTargetName(), "DependencyCache.bin");
 		}
 	}
 }
diff --git a/Engine/Source/Programs/UnrealBuildTool/System/DirectoryLookupCache.cs b/Engine/Source/Programs/UnrealBuildTool/System/DirectoryLookupCache.cs
--- a/Engine/Source/Programs/UnrealBuildTool/System/DirectoryLookupCache.cs
+++ b/Engine/Source/Programs/UnrealBuildTool/System/DirectoryLookupCache.cs
@@ -8,47 +8,45 @@ namespace UnrealBuildTool
 {
 	class DirectoryCache
 	{
-		public DirectoryCache( string BaseDirectory )
+		public DirectoryCache( DirectoryReference BaseDirectory )
 		{
-			if( Directory.Exists( BaseDirectory ) )
+			if( BaseDirectory.Exists() )
 			{
-				BaseDirectory = Utils.CleanDirectorySeparators(BaseDirectory);
-				var FoundFiles = Directory.GetFiles( BaseDirectory );
-				foreach( var File in FoundFiles )
-				{
-					Files.Add(Utils.CleanDirectorySeparators(File).ToLowerInvariant());
-				}
+				Files = new HashSet<FileReference>(FileReference.Enumerate(BaseDirectory));
+			}
+			else
+			{
+				Files = new HashSet<FileReference>();
 			}
 		}
 
-		public bool HasFile( string File )
+		public bool HasFile( FileReference File )
 		{
-			return Files.Contains(Utils.CleanDirectorySeparators(File).ToLowerInvariant());
+			return Files.Contains(File);
 		}
 
-		HashSet<string> Files = new HashSet<string>();
+		HashSet<FileReference> Files;
 	}
 
 
 	public static class DirectoryLookupCache
 	{
-		static public bool FileExists( string FileFullPath )
+		static public bool FileExists( FileReference FileFullPath )
 		{
-			string FileDirectory = Utils.CleanDirectorySeparators(Path.GetDirectoryName(FileFullPath));
-			string DirectoriesKey = FileDirectory.ToLowerInvariant();
+			DirectoryReference FileDirectory = FileFullPath.Directory;
 
 			DirectoryCache FoundDirectoryCache;
-			if( !Directories.TryGetValue( DirectoriesKey, out FoundDirectoryCache ) )
+			if( !Directories.TryGetValue( FileDirectory, out FoundDirectoryCache ) )
 			{
 				// First time we've seen this directory.  Gather information about files.
 				FoundDirectoryCache = new DirectoryCache( FileDirectory );
-				Directories.Add( DirectoriesKey, FoundDirectoryCache );
+				Directories.Add( FileDirectory, FoundDirectoryCache );
 			}
 
 			// Try to find the file in this directory
 			return FoundDirectoryCache.HasFile( FileFullPath );
 		}
 
-		static Dictionary<string, DirectoryCache> Directories = new Dictionary<string, DirectoryCache>();
+		static Dictionary<DirectoryReference, DirectoryCache> Directories = new Dictionary<DirectoryReference, DirectoryCache>();
 	}
 }
diff --git a/Engine/Source/Programs/UnrealBuildTool/System/DynamicCompilation.cs b/Engine/Source/Programs/UnrealBuildTool/System/DynamicCompilation.cs
--- a/Engine/Source/Programs/UnrealBuildTool/System/DynamicCompilation.cs
+++ b/Engine/Source/Programs/UnrealBuildTool/System/DynamicCompilation.cs
@@ -3,6 +3,7 @@
 using System;
 using System.Collections.Generic;
 using System.IO;
+using System.Linq;
 using System.CodeDom.Compiler;
 using Microsoft.CSharp;
 using System.Reflection;
@@ -19,7 +20,7 @@ namespace UnrealBuildTool
 		/*
 		 * Checks to see if the assembly needs compilation
 		 */
-		private static bool RequiresCompilation( List<string> SourceFileNames, string AssemblySourceListFilePath, string OutputAssemblyPath )
+		private static bool RequiresCompilation( List<FileReference> SourceFileNames, FileReference AssemblySourceListFilePath, FileReference OutputAssemblyPath )
 		{
 			if ( UnrealBuildTool.RunningRocket() && ProjectFileGenerator.bGenerateProjectFiles )
 			{
@@ -28,7 +29,7 @@ namespace UnrealBuildTool
 			}
 
 			// Check to see if we already have a compiled assembly file on disk
-			FileInfo OutputAssemblyInfo = new FileInfo( OutputAssemblyPath );
+			FileInfo OutputAssemblyInfo = new FileInfo( OutputAssemblyPath.FullName );
 			if( OutputAssemblyInfo.Exists )
 			{
 				// Check the time stamp of the UnrealBuildTool.exe file.  If Unreal Build Tool was compiled more
@@ -46,7 +47,7 @@ namespace UnrealBuildTool
 				{
 					// Make sure we have a manifest of source files used to compile the output assembly.  If it doesn't exist
 					// for some reason (not an expected case) then we'll need to recompile.
-					var AssemblySourceListFile = new FileInfo( AssemblySourceListFilePath );
+					var AssemblySourceListFile = new FileInfo( AssemblySourceListFilePath.FullName );
 					if( !AssemblySourceListFile.Exists )
 					{
 						return true;
@@ -55,7 +56,7 @@ namespace UnrealBuildTool
 					{
 						// Make sure the source files we're compiling are the same as the source files that were compiled
 						// for the assembly that we want to load
-						var ExistingAssemblySourceFileNames = new List<string>();
+						var ExistingAssemblySourceFileNames = new List<FileReference>();
 						{
 							using( var Reader = AssemblySourceListFile.OpenRead() )
 							{
@@ -63,10 +64,12 @@ namespace UnrealBuildTool
 								{
 									for( var ExistingSourceFileName = TextReader.ReadLine(); ExistingSourceFileName != null; ExistingSourceFileName = TextReader.ReadLine() )
 									{
-										ExistingAssemblySourceFileNames.Add( ExistingSourceFileName );
+										FileReference FullExistingSourceFileName = new FileReference(ExistingSourceFileName);
+
+										ExistingAssemblySourceFileNames.Add( FullExistingSourceFileName );
 
 										// Was the existing assembly compiled with a source file that we aren't interested in?  If so, then it needs to be recompiled.
-										if( !SourceFileNames.Contains( ExistingSourceFileName ) )
+										if( !SourceFileNames.Contains( FullExistingSourceFileName ) )
 										{
 											return true;
 										}
@@ -84,7 +87,7 @@ namespace UnrealBuildTool
 								return true;
 							}
 
-							var SourceFileInfo = new FileInfo( Path.GetFullPath( SourceFileName ) );
+							var SourceFileInfo = new FileInfo( SourceFileName.FullName );
 
 							// Check to see if the source file exists
 							if( !SourceFileInfo.Exists )
@@ -122,7 +125,7 @@ namespace UnrealBuildTool
 		/*
 		 * Compiles an assembly from source files
 		 */
-		private static Assembly CompileAssembly(string OutputAssemblyPath, List<string> SourceFileNames, List<string> ReferencedAssembies, List<string> PreprocessorDefines = null, bool TreatWarningsAsErrors = false)
+		private static Assembly CompileAssembly(FileReference OutputAssemblyPath, List<FileReference> SourceFileNames, List<string> ReferencedAssembies, List<string> PreprocessorDefines = null, bool TreatWarningsAsErrors = false)
 		{
 			var TemporaryFiles = new TempFileCollection();
 
@@ -133,7 +136,7 @@ namespace UnrealBuildTool
 				CompileParams.GenerateInMemory = false;
 
 				// This is the full path to the assembly file we're generating
-				CompileParams.OutputAssembly = OutputAssemblyPath;
+				CompileParams.OutputAssembly = OutputAssemblyPath.FullName;
 
 				// We always want to generate a class library, not an executable
 				CompileParams.GenerateExecutable = false;
@@ -191,7 +194,7 @@ namespace UnrealBuildTool
 			}
 
 			// Create the output directory if it doesn't exist already
-			DirectoryInfo DirInfo = new DirectoryInfo( Path.GetDirectoryName( OutputAssemblyPath ) );
+			DirectoryInfo DirInfo = new DirectoryInfo( OutputAssemblyPath.Directory.FullName );
 			if( !DirInfo.Exists )
 			{
 				try
@@ -211,7 +214,7 @@ namespace UnrealBuildTool
 				// Enable .NET 4.0 as we want modern language features like 'var'
 				var ProviderOptions = new Dictionary<string, string>() { { "CompilerVersion", "v4.0" } };
 				var Compiler = new CSharpCodeProvider( ProviderOptions );
-				CompileResults = Compiler.CompileAssemblyFromFile( CompileParams, SourceFileNames.ToArray() );
+				CompileResults = Compiler.CompileAssemblyFromFile( CompileParams, SourceFileNames.Select(x => x.FullName).ToArray() );
 			}
 			catch( Exception Ex )
 			{
@@ -250,13 +253,10 @@ namespace UnrealBuildTool
 		/// <param name="SourceFileNames">List of source file name</param>
 		/// <param name="OutputAssemblyPath">Full path to the assembly to be created</param>
 		/// <returns>The assembly that was loaded</returns>
-		public static Assembly CompileAndLoadAssembly(string OutputAssemblyPath, List<string> SourceFileNames, List<string> ReferencedAssembies = null, List<string> PreprocessorDefines = null, bool DoNotCompile = false, bool TreatWarningsAsErrors = false)
+		public static Assembly CompileAndLoadAssembly(FileReference OutputAssemblyPath, List<FileReference> SourceFileNames, List<string> ReferencedAssembies = null, List<string> PreprocessorDefines = null, bool DoNotCompile = false, bool TreatWarningsAsErrors = false)
 		{
-			// Load assembly requires absolute paths
-			OutputAssemblyPath = Path.GetFullPath( OutputAssemblyPath );
-
 			// Check to see if the resulting assembly is compiled and up to date
-			var AssemblySourcesListFilePath = Path.Combine( Path.GetDirectoryName( OutputAssemblyPath ), Path.GetFileNameWithoutExtension( OutputAssemblyPath ) + "SourceFiles.txt" );
+			FileReference AssemblySourcesListFilePath = FileReference.Combine(OutputAssemblyPath.Directory, Path.GetFileNameWithoutExtension( OutputAssemblyPath.FullName ) + "SourceFiles.txt" );
 			bool bNeedsCompilation = false;
 			if (!DoNotCompile)
 			{
@@ -270,7 +270,7 @@ namespace UnrealBuildTool
 				try
 				{
 					// Load the previously-compiled assembly from disk
-					CompiledAssembly = Assembly.LoadFile( OutputAssemblyPath );
+					CompiledAssembly = Assembly.LoadFile( OutputAssemblyPath.FullName );
 				}
 				catch( FileLoadException Ex )
 				{
@@ -296,7 +296,7 @@ namespace UnrealBuildTool
 				// Save out a list of all the source files we compiled.  This is so that we can tell if whole files were added or removed
 				// since the previous time we compiled the assembly.  In that case, we'll always want to recompile it!
 				{
-					FileInfo AssemblySourcesListFile = new FileInfo( AssemblySourcesListFilePath );
+					FileInfo AssemblySourcesListFile = new FileInfo( AssemblySourcesListFilePath.FullName );
 					using( var Writer = AssemblySourcesListFile.CreateText() )
 					{
 						SourceFileNames.ForEach( x => Writer.WriteLine( x ) );
diff --git a/Engine/Source/Programs/UnrealBuildTool/System/ExternalExecution.cs b/Engine/Source/Programs/UnrealBuildTool/System/ExternalExecution.cs
--- a/Engine/Source/Programs/UnrealBuildTool/System/ExternalExecution.cs
+++ b/Engine/Source/Programs/UnrealBuildTool/System/ExternalExecution.cs
@@ -129,9 +129,9 @@ namespace UnrealBuildTool
 		{
 			public string       Name;
 			public string		ModuleType;
-			public string       BaseDirectory;
-			public string       IncludeBase;     // The include path which all UHT-generated includes should be relative to
-			public string       OutputDirectory;
+			public string		BaseDirectory;
+			public string		IncludeBase;     // The include path which all UHT-generated includes should be relative to
+			public string		OutputDirectory;
 			public List<string> ClassesHeaders;
 			public List<string> PublicHeaders;
 			public List<string> PrivateHeaders;
@@ -165,7 +165,7 @@ namespace UnrealBuildTool
 				PrivateHeaders           = Info.PrivateUObjectHeaders      .Select((Header) => Header.AbsolutePath).ToList(),
 				PCH                      = Info.PCH,
 				GeneratedCPPFilenameBase = Info.GeneratedCPPFilenameBase,
-				SaveExportedHeaders      = !UnrealBuildTool.IsEngineInstalled() || !Utils.IsFileUnderDirectory(Info.ModuleDirectory, BuildConfiguration.RelativeEnginePath),
+				SaveExportedHeaders      = !UnrealBuildTool.IsEngineInstalled() || !new DirectoryReference(Info.ModuleDirectory).IsUnderDirectory(UnrealBuildTool.EngineDirectory),
 				UHTGeneratedCodeVersion = Info.GeneratedCodeVersion,
 			}).ToList();
 		}
@@ -187,10 +187,10 @@ namespace UnrealBuildTool
 		/// Generates a UHTModuleInfo for a particular named module under a directory.
 		/// </summary>
 		/// <returns>
-		public static UHTModuleInfo CreateUHTModuleInfo(IEnumerable<string> HeaderFilenames, UEBuildTarget Target, string ModuleName, string ModuleDirectory, UEBuildModuleType ModuleType)
+		public static UHTModuleInfo CreateUHTModuleInfo(IEnumerable<string> HeaderFilenames, UEBuildTarget Target, string ModuleName, DirectoryReference ModuleDirectory, UEBuildModuleType ModuleType)
 		{
-			var ClassesFolder = Path.Combine(ModuleDirectory, "Classes");
-			var PublicFolder  = Path.Combine(ModuleDirectory, "Public");
+			var ClassesFolder = DirectoryReference.Combine(ModuleDirectory, "Classes");
+			var PublicFolder  = DirectoryReference.Combine(ModuleDirectory, "Public");
 			var BuildPlatform = UEBuildPlatform.GetBuildPlatform(Target.Platform);
 
 			var AllClassesHeaders     = new List<FileItem>();
@@ -205,11 +205,11 @@ namespace UnrealBuildTool
 
 				if (CPPEnvironment.DoesFileContainUObjects(UObjectHeaderFileItem.AbsolutePath))
 				{
-					if (UObjectHeaderFileItem.AbsolutePath.StartsWith(ClassesFolder))
+					if (new FileReference(UObjectHeaderFileItem.AbsolutePath).IsUnderDirectory(ClassesFolder))
 					{
 						AllClassesHeaders.Add(UObjectHeaderFileItem);
 					}
-					else if (UObjectHeaderFileItem.AbsolutePath.StartsWith(PublicFolder))
+					else if (new FileReference(UObjectHeaderFileItem.AbsolutePath).IsUnderDirectory(PublicFolder))
 					{
 						PublicUObjectHeaders.Add(UObjectHeaderFileItem);
 					}
@@ -223,7 +223,7 @@ namespace UnrealBuildTool
 			var Result = new UHTModuleInfo
 			{
 				ModuleName                  = ModuleName,
-				ModuleDirectory             = ModuleDirectory,
+				ModuleDirectory             = ModuleDirectory.FullName,
 				ModuleType                  = ModuleType.ToString(),
 				PublicUObjectClassesHeaders = AllClassesHeaders,
 				PublicUObjectHeaders        = PublicUObjectHeaders,
@@ -274,7 +274,7 @@ namespace UnrealBuildTool
 					Timestamp = DateTime.MaxValue;
 					return false;
 				}
-				Receipt.ExpandPathVariables(BuildConfiguration.RelativeEnginePath, BuildConfiguration.RelativeEnginePath);
+				Receipt.ExpandPathVariables(UnrealBuildTool.EngineDirectory, UnrealBuildTool.EngineDirectory);
 
 				// Check all the binaries exist, and that all the DLLs are built against the right version
 				if(!CheckBinariesExist(Receipt) || !CheckDynamicLibaryVersionsMatch(Receipt))
@@ -427,7 +427,7 @@ namespace UnrealBuildTool
 			foreach( var Module in UObjectModules )
 			{
 				// If the engine is installed, skip skip checking timestamps for modules that are under the engine directory
-				if (UnrealBuildTool.IsEngineInstalled() && Utils.IsFileUnderDirectory( Module.ModuleDirectory, BuildConfiguration.RelativeEnginePath))
+				if (UnrealBuildTool.IsEngineInstalled() && new DirectoryReference(Module.ModuleDirectory).IsUnderDirectory(UnrealBuildTool.EngineDirectory))
 				{
 					continue;
 				}
@@ -537,7 +537,7 @@ namespace UnrealBuildTool
 					if (GeneratedCodeDirectoryInfo.Exists)
 					{
 						// Don't write anything to the engine directory if we're running an installed build
-						if(UnrealBuildTool.IsEngineInstalled() && Utils.IsFileUnderDirectory(Module.ModuleDirectory, BuildConfiguration.RelativeEnginePath))
+						if(UnrealBuildTool.IsEngineInstalled() && new DirectoryReference(Module.ModuleDirectory).IsUnderDirectory(UnrealBuildTool.EngineDirectory))
 						{
 							continue;
 						}
@@ -546,7 +546,7 @@ namespace UnrealBuildTool
 						// However, the headers might not be touched at all since that would cause the compiler to recompile everything
 						// We can't alter the directory timestamp directly, because this may throw exceptions when the directory is
 						// open in visual studio or windows explorer, so instead we create a blank file that will change the timestamp for us
-						string TimestampFile = GeneratedCodeDirectoryInfo.FullName + Path.DirectorySeparatorChar + @"Timestamp";
+						FileReference TimestampFile = FileReference.Combine(new DirectoryReference(GeneratedCodeDirectoryInfo.FullName), "Timestamp");
 
 						if( !GeneratedCodeDirectoryInfo.Exists )
 						{
@@ -606,7 +606,7 @@ namespace UnrealBuildTool
 		 * Builds and runs the header tool and touches the header directories.
 		 * Performs any early outs if headers need no changes, given the UObject modules, tool path, game name, and configuration
 		 */
-		public static bool ExecuteHeaderToolIfNecessary( UEBuildTarget Target, CPPEnvironment GlobalCompileEnvironment, List<UHTModuleInfo> UObjectModules, string ModuleInfoFileName, ref ECompilationResult UHTResult )
+		public static bool ExecuteHeaderToolIfNecessary( UEBuildTarget Target, CPPEnvironment GlobalCompileEnvironment, List<UHTModuleInfo> UObjectModules, FileReference ModuleInfoFileName, ref ECompilationResult UHTResult )
 		{
 			if(ProgressWriter.bWriteMarkup)
 			{
@@ -710,8 +710,8 @@ namespace UnrealBuildTool
 					}
 
 					// Disable extensions when serializing to remove the $type fields
-					Directory.CreateDirectory(Path.GetDirectoryName(ModuleInfoFileName));
-					System.IO.File.WriteAllText(ModuleInfoFileName, fastJSON.JSON.Instance.ToJSON(Manifest, new fastJSON.JSONParameters{ UseExtensions = false }));
+					Directory.CreateDirectory(ModuleInfoFileName.Directory.FullName);
+					System.IO.File.WriteAllText(ModuleInfoFileName.FullName, fastJSON.JSON.Instance.ToJSON(Manifest, new fastJSON.JSONParameters{ UseExtensions = false }));
 
 					string CmdLine = (UnrealBuildTool.HasUProjectFile()) ? "\"" + UnrealBuildTool.GetUProjectFile() + "\"" : Target.GetTargetName();
 					CmdLine += " \"" + ModuleInfoFileName + "\" -LogCmds=\"loginit warning, logexit warning, logdatabase error\"";
diff --git a/Engine/Source/Programs/UnrealBuildTool/System/FileItem.cs b/Engine/Source/Programs/UnrealBuildTool/System/FileItem.cs
--- a/Engine/Source/Programs/UnrealBuildTool/System/FileItem.cs
+++ b/Engine/Source/Programs/UnrealBuildTool/System/FileItem.cs
@@ -12,7 +12,7 @@ namespace UnrealBuildTool
 {
 	/**
 	 * Represents a file on disk that is used as an input or output of a build action.
-	 * FileItems are created by calling FileItem.GetItemByPath, which creates a single FileItem for each unique file path.
+	 * FileItems are created by calling FileItem.GetItemByFileReference, which creates a single FileItem for each unique file path.
 	 */
 	[Serializable]
 	public class FileItem : ISerializable
@@ -24,8 +24,8 @@ namespace UnrealBuildTool
 		/** The action that produces the file. */
 		public Action ProducingAction = null;
 
-		/** The absolute path of the file. */
-		public readonly string AbsolutePath;
+		/** The file reference */
+		public FileReference Reference;
 
 		/** True if any DLLs produced by this  */
 		public bool bNeedsHotReloadNumbersDLLCleanUp = false;
@@ -33,6 +33,11 @@ namespace UnrealBuildTool
 		/** Whether or not this is a remote file, in which case we can't access it directly */
 		public bool bIsRemoteFile = false;
 
+		/** Accessor for the absolute path to the file */
+		public string AbsolutePath
+		{
+			get { return Reference.FullName; }
+		}
 
 		/** For C++ file items, this stores cached information about the include paths needed in order to include header files from these C++ files.  This is part of UBT's dependency caching optimizations. */
 		public CPPIncludeInfo CachedCPPIncludeInfo
@@ -98,7 +103,7 @@ namespace UnrealBuildTool
 		public string PCHHeaderNameInCode;
 
 		/** The PCH file that this file will use */
-		public string PrecompiledHeaderIncludeFilename;
+		public FileReference PrecompiledHeaderIncludeFilename;
 
 
 		///
@@ -169,7 +174,7 @@ namespace UnrealBuildTool
 		public static long MissingFileItemCount = 0;
 
 		/** A case-insensitive dictionary that's used to map each unique file name to a single FileItem object. */
-		static Dictionary<string, FileItem> UniqueSourceFileMap = new Dictionary<string, FileItem>();
+		static Dictionary<FileReference, FileItem> UniqueSourceFileMap = new Dictionary<FileReference, FileItem>();
 
 		/** A list of remote file items that have been created but haven't needed the remote info yet, so we can gang up many into one request */
 		static List<FileItem> DelayedRemoteLookupFiles = new List<FileItem>();
@@ -212,21 +217,20 @@ namespace UnrealBuildTool
 		/** @return The FileItem that represents the given file path. */
 		public static FileItem GetItemByPath(string FilePath)
 		{
-			var FullPath = Path.GetFullPath( FilePath );
-			return GetItemByFullPath( FullPath );
+			return GetItemByFileReference( new FileReference(FilePath) );
 		}
 
 		/** @return The FileItem that represents the given a full file path. */
-		public static FileItem GetItemByFullPath( string FullPath )
+		public static FileItem GetItemByFileReference(FileReference Reference)
 		{
 			FileItem Result = null;
-			if( UniqueSourceFileMap.TryGetValue( FullPath.ToLowerInvariant(), out Result ) )
+			if( UniqueSourceFileMap.TryGetValue( Reference, out Result ) )
 			{
 				return Result;
 			}
 			else
 			{
-				return new FileItem( FullPath );
+				return new FileItem( Reference );
 			}
 		}
 
@@ -238,22 +242,29 @@ namespace UnrealBuildTool
 				throw new BuildException("GetRemoteItemByPath must be passed an absolute path, not a relative path '{0}'", AbsoluteRemotePath);
 			}
 
-			string InvariantPath = AbsoluteRemotePath.ToLowerInvariant();
+			FileReference RemoteFileReference = FileReference.MakeRemote(AbsoluteRemotePath);
+
 			FileItem Result = null;
-			if (UniqueSourceFileMap.TryGetValue(InvariantPath, out Result))
+			if (UniqueSourceFileMap.TryGetValue(RemoteFileReference, out Result))
 			{
 				return Result;
 			}
 			else
 			{
-				return new FileItem(AbsoluteRemotePath, true, Platform);
+				return new FileItem(RemoteFileReference, true, Platform);
 			}
 		}
 
 		/** If the given file path identifies a file that already exists, returns the FileItem that represents it. */
-		public static FileItem GetExistingItemByPath(string Path)
+		public static FileItem GetExistingItemByPath(string FileName)
+		{
+			return GetExistingItemByFileReference(new FileReference(FileName));
+		}
+
+		/** If the given file path identifies a file that already exists, returns the FileItem that represents it. */
+		public static FileItem GetExistingItemByFileReference(FileReference FileRef)
 		{
-			FileItem Result = GetItemByPath(Path);
+			FileItem Result = GetItemByFileReference(FileRef);
 			if (Result.bExists)
 			{
 				return Result;
@@ -280,18 +291,18 @@ namespace UnrealBuildTool
 		 * Creates a text file with the given contents.  If the contents of the text file aren't changed, it won't write the new contents to
 		 * the file to avoid causing an action to be considered outdated.
 		 */
-		public static FileItem CreateIntermediateTextFile(string AbsolutePath, string Contents)
+		public static FileItem CreateIntermediateTextFile(FileReference AbsolutePath, string Contents)
 		{
 			// Create the directory if it doesn't exist.
-			Directory.CreateDirectory(Path.GetDirectoryName(AbsolutePath));
+			Directory.CreateDirectory(Path.GetDirectoryName(AbsolutePath.FullName));
 
 			// Only write the file if its contents have changed.
-			if (!File.Exists(AbsolutePath) || !String.Equals(Utils.ReadAllText(AbsolutePath), Contents, StringComparison.InvariantCultureIgnoreCase))
+			if (!AbsolutePath.Exists() || !String.Equals(Utils.ReadAllText(AbsolutePath.FullName), Contents, StringComparison.InvariantCultureIgnoreCase))
 			{
-				File.WriteAllText(AbsolutePath, Contents, GetEncodingForString(Contents));
+				File.WriteAllText(AbsolutePath.FullName, Contents, GetEncodingForString(Contents));
 			}
 
-			return GetItemByPath(AbsolutePath);
+			return GetItemByFileReference(AbsolutePath);
 		}
 
 		/** Deletes the file. */
@@ -341,9 +352,9 @@ namespace UnrealBuildTool
 		}
 
 		/** Initialization constructor. */
-		protected FileItem(string FileAbsolutePath)
+		protected FileItem(FileReference InFile)
 		{
-			AbsolutePath = FileAbsolutePath;
+			Reference = InFile;
 
 			ResetFileInfo();
 
@@ -354,7 +365,7 @@ namespace UnrealBuildTool
 				// Log.TraceInformation( "Missing: " + FileAbsolutePath );
 			}
 
-			UniqueSourceFileMap[ AbsolutePath.ToLowerInvariant() ] = this;
+			UniqueSourceFileMap[ Reference ] = this;
 		}
 
 
@@ -362,7 +373,7 @@ namespace UnrealBuildTool
 		protected FileItem( SerializationInfo SerializationInfo, StreamingContext StreamingContext )
 		{
 			ProducingAction = (Action)SerializationInfo.GetValue( "pa", typeof( Action ) );
-			AbsolutePath = SerializationInfo.GetString( "ap" );
+			Reference = (FileReference)SerializationInfo.GetValue( "fi", typeof(FileReference) );
 			bIsRemoteFile = SerializationInfo.GetBoolean( "rf" );
 			bNeedsHotReloadNumbersDLLCleanUp = SerializationInfo.GetBoolean( "hr" );
 			CachedCPPIncludeInfo = (CPPIncludeInfo)SerializationInfo.GetValue( "ci", typeof( CPPIncludeInfo ) );
@@ -387,7 +398,7 @@ namespace UnrealBuildTool
 				}
 				else
 				{ 
-					UniqueSourceFileMap[ AbsolutePath.ToLowerInvariant() ] = this;
+					UniqueSourceFileMap[ Reference ] = this;
 				}
 			}
 		}
@@ -397,7 +408,7 @@ namespace UnrealBuildTool
 		public void GetObjectData( SerializationInfo SerializationInfo, StreamingContext StreamingContext )
 		{
 			SerializationInfo.AddValue( "pa", ProducingAction );
-			SerializationInfo.AddValue( "ap", AbsolutePath );
+			SerializationInfo.AddValue( "fi", Reference );
 			SerializationInfo.AddValue( "rf", bIsRemoteFile );
 			SerializationInfo.AddValue( "hr", bNeedsHotReloadNumbersDLLCleanUp );
 			SerializationInfo.AddValue( "ci", CachedCPPIncludeInfo );
@@ -443,10 +454,10 @@ namespace UnrealBuildTool
 		}
 
 		/** Initialization constructor for optionally remote files. */
-		protected FileItem(string InAbsolutePath, bool InIsRemoteFile, UnrealTargetPlatform Platform)
+		protected FileItem(FileReference InReference, bool InIsRemoteFile, UnrealTargetPlatform Platform)
 		{
 			bIsRemoteFile = InIsRemoteFile;
-			AbsolutePath = InAbsolutePath;
+			Reference = InReference;
 
 			// @todo iosmerge: This doesn't handle remote directories (may be needed for compiling Mac from Windows)
 			if (bIsRemoteFile)
