Use FileReference and DirectoryReference classes instead of strings for paths. Massively reduces the amount of paranoid calls to GetFullPath(), CleanDirectorySeparators() et al., and improves performa

diff --git a/Engine/Source/Programs/UnrealBuildTool/System/FileSystemNames.cs b/Engine/Source/Programs/UnrealBuildTool/System/FileSystemNames.cs
--- a/Engine/Source/Programs/UnrealBuildTool/System/FileSystemNames.cs
+++ b/Engine/Source/Programs/UnrealBuildTool/System/FileSystemNames.cs
@@ -10,7 +10,8 @@ namespace UnrealBuildTool
 	/// <summary>
 	/// Base class for file system objects (files or directories).
 	/// </summary>
-	public abstract class FullFileSystemName
+	[Serializable]
+	public abstract class FileSystemReference
 	{
 		/// <summary>
 		/// The path to this object. Stored as an absolute path, with O/S preferred separator characters, and no trailing slash for directories.
@@ -25,7 +26,7 @@ namespace UnrealBuildTool
 		/// <summary>
 		/// Constructs a filesystem object for the given path.
 		/// </summary>
-		public FullFileSystemName(string InPath)
+		public FileSystemReference(string InPath)
 		{
 			FullName = Path.GetFullPath(InPath);
 			CanonicalName = FullName.ToLowerInvariant();
@@ -34,18 +35,113 @@ namespace UnrealBuildTool
 		/// <summary>
 		/// Direct constructor for a path
 		/// </summary>
-		protected FullFileSystemName(string InFullName, string InCanonicalName)
+		protected FileSystemReference(string InFullName, string InCanonicalName)
 		{
 			FullName = InFullName;
 			CanonicalName = InCanonicalName;
 		}
 
+		/// <summary>
+		/// Create a full path by concatenating multiple strings
+		/// </summary>
+		/// <returns></returns>
+		static protected string CombineStrings(DirectoryReference BaseDirectory, params string[] Fragments)
+		{
+			// Get the initial string to append to, and strip any root directory suffix from it
+			StringBuilder NewFullName = new StringBuilder(BaseDirectory.FullName);
+			if(NewFullName.Length > 0 && NewFullName[NewFullName.Length - 1] == Path.DirectorySeparatorChar)
+			{
+				NewFullName.Remove(NewFullName.Length - 1, 1);
+			}
+
+			// Scan through the fragments to append, appending them to a string and updating the base length as we go
+			foreach(string Fragment in Fragments)
+			{
+				// Check if this fragment is an absolute path
+				if((Fragment.Length >= 2 && Fragment[1] == ':') || (Fragment.Length >= 1 && (Fragment[0] == Path.DirectorySeparatorChar || Fragment[0] == Path.AltDirectorySeparatorChar)))
+				{
+					// It is. Reset the new name to the full version of this path.
+					NewFullName.Clear();
+					NewFullName.Append(Path.GetFullPath(Fragment));
+				}
+				else
+				{
+					// Append all the parts of this fragment to the end of the existing path.
+					int StartIdx = 0;
+					while(StartIdx < Fragment.Length)
+					{
+						// Find the end of this fragment. We may have been passed multiple paths in the same string.
+						int EndIdx = StartIdx;
+						while(EndIdx < Fragment.Length && Fragment[EndIdx] != Path.DirectorySeparatorChar && Fragment[EndIdx] != Path.AltDirectorySeparatorChar)
+						{
+							EndIdx++;
+						}
+
+						// Ignore any empty sections, like leading or trailing slashes, and '.' directory references.
+						int Length = EndIdx - StartIdx;
+						if(Length == 0)
+						{
+							// Multiple directory separators in a row; illegal.
+							throw new ArgumentException("Path fragment '{0}' contains invalid directory separators.");
+						}
+						else if(Length == 2 && Fragment[StartIdx] == '.' && Fragment[StartIdx + 1] == '.')
+						{
+							// Remove the last directory name
+							for(int SeparatorIdx = NewFullName.Length - 1; SeparatorIdx >= 0; SeparatorIdx--)
+							{
+								if(NewFullName[SeparatorIdx] == Path.DirectorySeparatorChar)
+								{
+									NewFullName.Remove(SeparatorIdx, NewFullName.Length - SeparatorIdx);
+									break;
+								}
+							}
+						}
+						else if(Length != 1 || Fragment[StartIdx] != '.')
+						{
+							// Append this fragment
+							NewFullName.Append(Path.DirectorySeparatorChar);
+							NewFullName.Append(Fragment, StartIdx, Length);
+						}
+
+						// Move to the next part
+						StartIdx = EndIdx + 1;
+					}
+				}
+			}
+
+			// Append the directory separator
+			if(NewFullName.Length == 0 || (NewFullName.Length == 2 && NewFullName[1] == ':'))
+			{
+				NewFullName.Append(Path.DirectorySeparatorChar);
+			}
+			
+			// Set the new path variables
+			return NewFullName.ToString();
+		}
+
+		/// <summary>
+		/// Checks whether this name has the given extension.
+		/// </summary>
+		/// <param name="Extension">The extension to check</param>
+		/// <returns>True if this name has the given extension, false otherwise</returns>
+		public bool HasExtension(string Extension)
+		{
+			if(Extension.Length > 0 && Extension[0] != '.')
+			{
+				return HasExtension("." + Extension);
+			}
+			else
+			{ 
+				return CanonicalName.EndsWith(Extension.ToLowerInvariant());
+			}
+		}
+
 		/// <summary>
 		/// Determines if the given object is at or under the given directory
 		/// </summary>
 		/// <param name="Directory"></param>
 		/// <returns></returns>
-		public bool IsUnderDirectory(FullDirectoryName Other)
+		public bool IsUnderDirectory(DirectoryReference Other)
 		{
 			return CanonicalName.StartsWith(Other.CanonicalName) && (CanonicalName.Length == Other.CanonicalName.Length || CanonicalName[Other.CanonicalName.Length] == Path.DirectorySeparatorChar);
 		}
@@ -55,7 +151,7 @@ namespace UnrealBuildTool
 		/// </summary>
 		/// <param name="Directory">The directory to create a relative path from</param>
 		/// <returns>A relative path from the given directory</returns>
-		public string MakeRelativeTo(FullDirectoryName Directory)
+		public string MakeRelativeTo(DirectoryReference Directory)
 		{
 			// Find how much of the path is common between the two paths. This length does not include a trailing directory separator character.
 			int CommonDirectoryLength = -1;
@@ -137,15 +233,16 @@ namespace UnrealBuildTool
 	}
 
 	/// <summary>
-	/// Representation of a directory in the filesystem
+	/// Representation of an absolute directory path. Allows fast hashing and comparisons.
 	/// </summary>
-	public class FullDirectoryName : FullFileSystemName
+	[Serializable]
+	public class DirectoryReference : FileSystemReference
 	{
 		/// <summary>
 		/// Default constructor.
 		/// </summary>
 		/// <param name="InPath">Path to this directory.</param>
-		public FullDirectoryName(string InPath)
+		public DirectoryReference(string InPath)
 			: base(InPath)
 		{
 		}
@@ -155,29 +252,41 @@ namespace UnrealBuildTool
 		/// </summary>
 		/// <param name="InFullName"></param>
 		/// <param name="InCanonicalName"></param>
-		protected FullDirectoryName(string InFullName, string InCanonicalName)
+		protected DirectoryReference(string InFullName, string InCanonicalName)
 			: base(InFullName, InCanonicalName)
 		{
 		}
 
+		/// <summary>
+		/// Gets the top level directory name
+		/// </summary>
+		/// <returns>The name of the directory</returns>
+		public string GetDirectoryName()
+		{
+			return Path.GetFileName(FullName);
+		}
+
 		/// <summary>
 		/// Gets the directory containing this object
 		/// </summary>
 		/// <returns>A new directory object representing the directory containing this object</returns>
-		public FullDirectoryName GetParentDirectory()
+		public DirectoryReference ParentDirectory
 		{
-			if (IsRootDirectory())
+			get
 			{
-				return null;
-			}
+				if (IsRootDirectory())
+				{
+					return null;
+				}
 
-			int ParentLength = CanonicalName.LastIndexOf(Path.DirectorySeparatorChar);
-			if (ParentLength == 2 && CanonicalName[1] == ':')
-			{
-				ParentLength++;
-			}
+				int ParentLength = CanonicalName.LastIndexOf(Path.DirectorySeparatorChar);
+				if (ParentLength == 2 && CanonicalName[1] == ':')
+				{
+					ParentLength++;
+				}
 
-			return new FullDirectoryName(FullName.Substring(0, ParentLength), CanonicalName.Substring(0, ParentLength));
+				return new DirectoryReference(FullName.Substring(0, ParentLength), CanonicalName.Substring(0, ParentLength));
+			}
 		}
 
 		/// <summary>
@@ -185,10 +294,27 @@ namespace UnrealBuildTool
 		/// </summary>
 		/// <param name="File">The file to get directory for</param>
 		/// <returns>The full directory name containing the given file</returns>
-		public static FullDirectoryName GetParentDirectory(FullFileName File)
+		public static DirectoryReference GetParentDirectory(FileReference File)
 		{
 			int ParentLength = File.CanonicalName.LastIndexOf(Path.DirectorySeparatorChar);
-			return new FullDirectoryName(File.FullName.Substring(0, ParentLength), File.CanonicalName.Substring(0, ParentLength));
+			return new DirectoryReference(File.FullName.Substring(0, ParentLength), File.CanonicalName.Substring(0, ParentLength));
+		}
+
+		/// <summary>
+		/// Creates the directory
+		/// </summary>
+		public void CreateDirectory()
+		{
+			Directory.CreateDirectory(FullName);
+		}
+
+		/// <summary>
+		/// Checks whether the directory exists
+		/// </summary>
+		/// <returns>True if this directory exists</returns>
+		public bool Exists()
+		{
+			return Directory.Exists(FullName);
 		}
 
 		/// <summary>
@@ -200,6 +326,17 @@ namespace UnrealBuildTool
 			return CanonicalName[CanonicalName.Length - 1] == Path.DirectorySeparatorChar;
 		}
 
+		/// <summary>
+		/// Combine several fragments with a base directory, to form a new directory name
+		/// </summary>
+		/// <param name="BaseDirectory">The base directory</param>
+		/// <param name="Fragments">Fragments to combine with the base directory</param>
+		/// <returns>The new directory name</returns>
+		public static DirectoryReference Combine(DirectoryReference BaseDirectory, params string[] Fragments)
+		{
+			string FullName = FileSystemReference.CombineStrings(BaseDirectory, Fragments);
+			return new DirectoryReference(FullName, FullName.ToLowerInvariant());
+		}
 
 		/// <summary>
 		/// Compares two filesystem object names for equality. Uses the canonical name representation, not the display name representation.
@@ -207,9 +344,16 @@ namespace UnrealBuildTool
 		/// <param name="A">First object to compare.</param>
 		/// <param name="B">Second object to compare.</param>
 		/// <returns>True if the names represent the same object, false otherwise</returns>
-		public static bool operator ==(FullDirectoryName A, FullDirectoryName B)
+		public static bool operator ==(DirectoryReference A, DirectoryReference B)
 		{
-			return A.CanonicalName == B.CanonicalName;
+			if((object)A == null)
+			{
+				return (object)B == null;
+			}
+			else
+			{
+				return (object)B != null && A.CanonicalName == B.CanonicalName;
+			}
 		}
 
 		/// <summary>
@@ -218,9 +362,9 @@ namespace UnrealBuildTool
 		/// <param name="A">First object to compare.</param>
 		/// <param name="B">Second object to compare.</param>
 		/// <returns>False if the names represent the same object, true otherwise</returns>
-		public static bool operator !=(FullDirectoryName A, FullDirectoryName B)
+		public static bool operator !=(DirectoryReference A, DirectoryReference B)
 		{
-			return A.CanonicalName != B.CanonicalName;
+			return !(A == B);
 		}
 
 		/// <summary>
@@ -231,7 +375,7 @@ namespace UnrealBuildTool
 		/// <returns>True if the names represent the same object, false otherwise</returns>
 		public override bool Equals(object Obj)
 		{
-			return (Obj is FullDirectoryName) && ((FullDirectoryName)Obj) == this;
+			return (Obj is DirectoryReference) && ((DirectoryReference)Obj) == this;
 		}
 
 		/// <summary>
@@ -242,31 +386,147 @@ namespace UnrealBuildTool
 		{
 			return CanonicalName.GetHashCode();
 		}
+
+		/// <summary>
+		/// Helper function to create a remote directory reference. Unlike normal DirectoryReference objects, these aren't converted to a full path in the local filesystem.
+		/// </summary>
+		/// <param name="AbsolutePath">The absolute path in the remote file system</param>
+		/// <returns>New directory reference</returns>
+		public static DirectoryReference MakeRemote(string AbsolutePath)
+		{
+			return new DirectoryReference(AbsolutePath, AbsolutePath.ToLowerInvariant());
+		}
 	}
 
 	/// <summary>
-	/// Representation of a file in the filesystem
+	/// Representation of an absolute file path. Allows fast hashing and comparisons.
 	/// </summary>
-	public class FullFileName : FullFileSystemName
+	[Serializable]
+	public class FileReference : FileSystemReference
 	{
 		/// <summary>
 		/// Default constructor.
 		/// </summary>
 		/// <param name="InPath">Path to this file</param>
-		public FullFileName(string InPath)
+		public FileReference(string InPath)
 			: base(InPath)
 		{
 		}
 
+		/// <summary>
+		/// Default constructor.
+		/// </summary>
+		/// <param name="InPath">Path to this file</param>
+		protected FileReference(string InFullName, string InCanonicalName)
+			: base(InFullName, InCanonicalName)
+		{
+		}
+
+		/// <summary>
+		/// Gets the file name without path information
+		/// </summary>
+		/// <returns>A string containing the file name</returns>
+		public string GetFileName()
+		{
+			return Path.GetFileName(FullName);
+		}
+
+		/// <summary>
+		/// Gets the file name without path information or an extension
+		/// </summary>
+		/// <returns>A string containing the file name without an extension</returns>
+		public string GetFileNameWithoutExtension()
+		{
+			return Path.GetFileNameWithoutExtension(FullName);
+		}
+
+		/// <summary>
+		/// Gets the file name without path or any extensions
+		/// </summary>
+		/// <returns>A string containing the file name without an extension</returns>
+		public string GetFileNameWithoutAnyExtensions()
+		{
+			int StartIdx = FullName.LastIndexOf(Path.DirectorySeparatorChar) + 1;
+
+			int EndIdx = FullName.IndexOf('.', StartIdx);
+			if(EndIdx < StartIdx)
+			{
+				return FullName.Substring(StartIdx);
+			}
+			else
+			{
+				return FullName.Substring(StartIdx, EndIdx - StartIdx);
+			}
+		}
+
+		/// <summary>
+		/// Gets the extension for this filename
+		/// </summary>
+		/// <returns>A string containing the extension of this filename</returns>
+		public string GetExtension()
+		{
+			return Path.GetExtension(FullName);
+		}
+
+		/// <summary>
+		/// Change the file's extension to something else
+		/// </summary>
+		/// <param name="Extension">The new extension</param>
+		/// <returns>A FileReference with the same path and name, but with the new extension</returns>
+		public FileReference ChangeExtension(string Extension)
+		{
+			string NewFullName = Path.ChangeExtension(FullName, Extension);
+			return new FileReference(NewFullName, NewFullName.ToLowerInvariant());
+		}
+
 		/// <summary>
 		/// Gets the directory containing this file
 		/// </summary>
 		/// <returns>A new directory object representing the directory containing this object</returns>
-		public FullDirectoryName GetDirectory()
+		public DirectoryReference Directory
 		{
-			return FullDirectoryName.GetParentDirectory(this);
+			get { return DirectoryReference.GetParentDirectory(this); }
 		}
 
+		/// <summary>
+		/// Determines whether the given filename exists
+		/// </summary>
+		/// <returns>True if it exists, false otherwise</returns>
+		public bool Exists()
+		{
+			return File.Exists(FullName);
+		}
+
+		/// <summary>
+		/// Deletes this file
+		/// </summary>
+		public void Delete()
+		{
+			File.Delete(FullName);
+		}
+
+		/// <summary>
+		/// Combine several fragments with a base directory, to form a new filename
+		/// </summary>
+		/// <param name="BaseDirectory">The base directory</param>
+		/// <param name="Fragments">Fragments to combine with the base directory</param>
+		/// <returns>The new file name</returns>
+		public static FileReference Combine(DirectoryReference BaseDirectory, params string[] Fragments)
+		{
+			string FullName = FileSystemReference.CombineStrings(BaseDirectory, Fragments);
+			return new FileReference(FullName, FullName.ToLowerInvariant());
+		}
+
+		/// <summary>
+		/// Append a string to the end of a filename
+		/// </summary>
+		/// <param name="A">The base file reference</param>
+		/// <param name="B">Suffix to be appended</param>
+		/// <returns>The new file reference</returns>
+		public static FileReference operator+(FileReference A, string B)
+		{
+			return new FileReference(A.FullName + B, A.CanonicalName + B.ToLowerInvariant());
+		}
 
 		/// <summary>
 		/// Compares two filesystem object names for equality. Uses the canonical name representation, not the display name representation.
@@ -274,9 +534,16 @@ namespace UnrealBuildTool
 		/// <param name="A">First object to compare.</param>
 		/// <param name="B">Second object to compare.</param>
 		/// <returns>True if the names represent the same object, false otherwise</returns>
-		public static bool operator ==(FullFileName A, FullFileName B)
+		public static bool operator ==(FileReference A, FileReference B)
 		{
-			return A.CanonicalName == B.CanonicalName;
+			if((object)A == null)
+			{
+				return (object)B == null;
+			}
+			else
+			{
+				return (object)B != null && A.CanonicalName == B.CanonicalName;
+			}
 		}
 
 		/// <summary>
@@ -285,9 +552,9 @@ namespace UnrealBuildTool
 		/// <param name="A">First object to compare.</param>
 		/// <param name="B">Second object to compare.</param>
 		/// <returns>False if the names represent the same object, true otherwise</returns>
-		public static bool operator !=(FullFileName A, FullFileName B)
+		public static bool operator !=(FileReference A, FileReference B)
 		{
-			return A.CanonicalName != B.CanonicalName;
+			return !(A == B);
 		}
 
 		/// <summary>
@@ -298,7 +565,7 @@ namespace UnrealBuildTool
 		/// <returns>True if the names represent the same object, false otherwise</returns>
 		public override bool Equals(object Obj)
 		{
-			return (Obj is FullFileName) && ((FullFileName)Obj) == this;
+			return (Obj is FileReference) && ((FileReference)Obj) == this;
 		}
 
 		/// <summary>
@@ -309,5 +576,55 @@ namespace UnrealBuildTool
 		{
 			return CanonicalName.GetHashCode();
 		}
+
+		/// <summary>
+		/// Enumerate files from a given directory
+		/// </summary>
+		/// <returns>Sequence of file references</returns>
+		public static IEnumerable<FileReference> Enumerate(DirectoryReference BaseDirectory)
+		{
+			foreach(string FileName in System.IO.Directory.EnumerateFiles(BaseDirectory.FullName))
+			{
+				yield return new FileReference(FileName, FileName.ToLowerInvariant());
+			}
+		}
+
+		/// <summary>
+		/// Enumerate files from a given directory
+		/// </summary>
+		/// <param name="Pattern">Pattern to match</param>
+		/// <returns>Sequence of file references</returns>
+		public static IEnumerable<FileReference> Enumerate(DirectoryReference BaseDirectory, string Pattern)
+		{
+			foreach(string FileName in System.IO.Directory.EnumerateFiles(BaseDirectory.FullName, Pattern))
+			{
+				yield return new FileReference(FileName, FileName.ToLowerInvariant());
+			}
+		}
+
+		/// <summary>
+		/// Enumerate files from a given directory
+		/// </summary>
+		/// <param name="Pattern">Pattern to match</param>
+		/// <param name="Search">Options for the search</param>
+		/// <returns>Sequence of file references</returns>
+		public static IEnumerable<FileReference> Enumerate(DirectoryReference BaseDirectory, string Pattern, SearchOption Option)
+		{
+			foreach(string FileName in System.IO.Directory.EnumerateFiles(BaseDirectory.FullName, Pattern, Option))
+			{
+				yield return new FileReference(FileName, FileName.ToLowerInvariant());
+			}
+		}
+
+		/// <summary>
+		/// Helper function to create a remote file reference. Unlike normal FileReference objects, these aren't converted to a full path in the local filesystem, but are
+		/// left as they are passed in.
+		/// </summary>
+		/// <param name="AbsolutePath">The absolute path in the remote file system</param>
+		/// <returns>New file reference</returns>
+		public static FileReference MakeRemote(string AbsolutePath)
+		{
+			return new FileReference(AbsolutePath, AbsolutePath.ToLowerInvariant());
+		}
 	}
 }
diff --git a/Engine/Source/Programs/UnrealBuildTool/System/FlatCPPIncludeDepencencyCache.cs b/Engine/Source/Programs/UnrealBuildTool/System/FlatCPPIncludeDepencencyCache.cs
--- a/Engine/Source/Programs/UnrealBuildTool/System/FlatCPPIncludeDepencencyCache.cs
+++ b/Engine/Source/Programs/UnrealBuildTool/System/FlatCPPIncludeDepencencyCache.cs
@@ -16,10 +16,10 @@ namespace UnrealBuildTool
 	class FlatCPPIncludeDependencyInfo
 	{
 		/// The PCH header this file is dependent on.
-		public string PCHName;
+		public FileReference PCHName;
 
 		/// List of files this file includes, excluding headers that were included from a PCH.
-		public List<string> Includes;
+		public List<FileReference> Includes;
 
 		/// Transient cached list of FileItems for all of the includes for a specific files.  This just saves a bunch of string
 		/// hash lookups as we locate FileItems for files that we've already requested dependencies for
@@ -41,10 +41,10 @@ namespace UnrealBuildTool
 		/// <returns>The new instance</returns>
 		public static FlatCPPIncludeDependencyCache Create( UEBuildTarget Target )
 		{
-			string CachePath = FlatCPPIncludeDependencyCache.GetDependencyCachePathForTarget( Target );;
+			FileReference CachePath = FlatCPPIncludeDependencyCache.GetDependencyCachePathForTarget( Target );;
 
 			// See whether the cache file exists.
-			FileItem Cache = FileItem.GetItemByPath(CachePath);
+			FileItem Cache = FileItem.GetItemByFileReference(CachePath);
 			if (Cache.bExists)
 			{
 				if (BuildConfiguration.bPrintPerformanceInfo)
@@ -113,7 +113,7 @@ namespace UnrealBuildTool
 			catch (Exception Ex)
 			{
 				// Don't bother failing if the file format has changed, simply abort the cache load
-				if (Ex.Message.Contains( "cannot be converted to type" ))	// To catch serialization differences added when we added the DependencyInfo struct
+				if (!Ex.Message.Contains( "cannot be converted to type" ))	// To catch serialization differences added when we added the DependencyInfo struct
 				{
 					Console.Error.WriteLine("Failed to read FlatCPPIncludeDependencyCache: {0}", Ex.Message);
 				}
@@ -129,7 +129,7 @@ namespace UnrealBuildTool
 		protected FlatCPPIncludeDependencyCache(FileItem Cache)
 		{
 			CacheFileItem = Cache;
-			DependencyMap = new Dictionary<string, FlatCPPIncludeDependencyInfo>();
+			DependencyMap = new Dictionary<FileReference, FlatCPPIncludeDependencyInfo>();
 		}
 
 
@@ -192,7 +192,7 @@ namespace UnrealBuildTool
 		/// <param name="AbsoluteFilePath">File to set</param>
 		/// <param name="PCHName">The PCH for this file</param>
 		/// <param name="Dependencies">List of source dependencies</param>
-		public void SetDependenciesForFile( string AbsoluteFilePath, string PCHName, List<string> Dependencies )
+		public void SetDependenciesForFile( FileReference AbsoluteFilePath, FileReference PCHName, List<FileReference> Dependencies )
 		{
 			var DependencyInfo = new FlatCPPIncludeDependencyInfo();
 			DependencyInfo.PCHName = PCHName;	// @todo ubtmake: Not actually used yet.  The idea is to use this to reduce the number of indirect includes we need to store in the cache.
@@ -200,7 +200,7 @@ namespace UnrealBuildTool
 			DependencyInfo.IncludeFileItems = null;
 
 			// @todo ubtmake: We could shrink this file by not storing absolute paths (project and engine relative paths only, except for system headers.)  May affect load times.
-			DependencyMap[ AbsoluteFilePath.ToLowerInvariant() ] = DependencyInfo;
+			DependencyMap[ AbsoluteFilePath ] = DependencyInfo;
 			bIsDirty = true;
 		}
 
@@ -210,18 +210,18 @@ namespace UnrealBuildTool
 		/// </summary>
 		/// <param name="AbsoluteFilePath">Path to the file</param>
 		/// <returns>The list of includes</returns>
-		public List<FileItem> GetDependenciesForFile( string AbsoluteFilePath )
+		public List<FileItem> GetDependenciesForFile( FileReference AbsoluteFilePath )
 		{
 			FlatCPPIncludeDependencyInfo DependencyInfo;
-			if( DependencyMap.TryGetValue( AbsoluteFilePath.ToLowerInvariant(), out DependencyInfo ) )
+			if( DependencyMap.TryGetValue( AbsoluteFilePath, out DependencyInfo ) )
 			{
 				// Update our transient cache of FileItems for each of the included files
 				if( DependencyInfo.IncludeFileItems == null )
 				{
 					DependencyInfo.IncludeFileItems = new List<FileItem>( DependencyInfo.Includes.Count );
-					foreach( string Dependency in DependencyInfo.Includes )
+					foreach( FileReference Dependency in DependencyInfo.Includes )
 					{
-						DependencyInfo.IncludeFileItems.Add( FileItem.GetItemByFullPath( Dependency ) );
+						DependencyInfo.IncludeFileItems.Add( FileItem.GetItemByFileReference( Dependency ) );
 					}
 				}
 				return DependencyInfo.IncludeFileItems;
@@ -235,15 +235,18 @@ namespace UnrealBuildTool
 		/// </summary>
 		/// <param name="Target">Current build target</param>
 		/// <returns>Cache Path</returns>
-		public static string GetDependencyCachePathForTarget(UEBuildTarget Target)
+		public static FileReference GetDependencyCachePathForTarget(UEBuildTarget Target)
 		{
-			string PlatformIntermediatePath = BuildConfiguration.PlatformIntermediatePath;
+			DirectoryReference PlatformIntermediatePath;
 			if (UnrealBuildTool.HasUProjectFile())
 			{
-				PlatformIntermediatePath = Path.Combine(UnrealBuildTool.GetUProjectPath(), BuildConfiguration.PlatformIntermediateFolder);
+				PlatformIntermediatePath = DirectoryReference.Combine(UnrealBuildTool.GetUProjectPath(), BuildConfiguration.PlatformIntermediateFolder);
 			}
-			string CachePath = Path.Combine(PlatformIntermediatePath, Target.GetTargetName(), "FlatCPPIncludes.bin" );
-			return CachePath;
+			else
+			{
+				PlatformIntermediatePath = new DirectoryReference(BuildConfiguration.PlatformIntermediatePath);
+			}
+			return FileReference.Combine(PlatformIntermediatePath, Target.GetTargetName(), "FlatCPPIncludes.bin" );
 		}
 
 
@@ -256,7 +259,7 @@ namespace UnrealBuildTool
 		private bool bIsDirty = false;
 
 		/// Dependency lists, keyed (case-insensitively) on file's absolute path.
-		private Dictionary<string, FlatCPPIncludeDependencyInfo> DependencyMap;
+		private Dictionary<FileReference, FlatCPPIncludeDependencyInfo> DependencyMap;
 
 
 	}
diff --git a/Engine/Source/Programs/UnrealBuildTool/System/KDevelopGenerator.cs b/Engine/Source/Programs/UnrealBuildTool/System/KDevelopGenerator.cs
--- a/Engine/Source/Programs/UnrealBuildTool/System/KDevelopGenerator.cs
+++ b/Engine/Source/Programs/UnrealBuildTool/System/KDevelopGenerator.cs
@@ -23,7 +23,7 @@ namespace UnrealBuildTool
 
 	public class KDevelopProjectFile : ProjectFile
 	{
-		public KDevelopProjectFile( string InitFilePath )
+		public KDevelopProjectFile( FileReference InitFilePath )
 			: base(InitFilePath)
 		{
 		}
@@ -186,7 +186,7 @@ namespace UnrealBuildTool
 						continue;
 					}
 
-					var TargetName = Utils.GetFilenameWithoutAnyExtensions(TargetFile.TargetFilePath);
+					var TargetName = TargetFile.TargetFilePath.GetFileNameWithoutAnyExtensions();
 
 					// Remove both ".cs" and ".
 					foreach (UnrealTargetConfiguration CurConfiguration in Enum.GetValues(typeof(UnrealTargetConfiguration))) 
@@ -278,7 +278,7 @@ namespace UnrealBuildTool
 
 				foreach (var CurPath in KDevelopProject.IntelliSenseIncludeSearchPaths) 
 				{
-					string FullProjectPath = Path.GetFullPath(ProjectFileGenerator.MasterProjectRelativePath);
+					string FullProjectPath = ProjectFileGenerator.MasterProjectPath.FullName;
 					string FullPath = "";
 
 					// need to test to see if this in the project souce tree
@@ -290,7 +290,7 @@ namespace UnrealBuildTool
 					}
 					else
 					{
-						FullPath = Path.GetFullPath(Path.Combine(Path.GetDirectoryName(KDevelopProject.ProjectFilePath), CurPath));
+						FullPath = Path.GetFullPath(Path.Combine(Path.GetDirectoryName(KDevelopProject.ProjectFilePath.FullName), CurPath));
 						FullPath = Utils.MakePathRelativeTo(FullPath, FullProjectPath);
 						FullPath = FullPath.TrimEnd('/');
 						FullPath = Path.Combine(UnrealEngineRootPath, FullPath);
@@ -307,7 +307,7 @@ namespace UnrealBuildTool
 
 				foreach (var CurPath in KDevelopProject.IntelliSenseSystemIncludeSearchPaths) 
 				{
-					string FullProjectPath = Path.GetFullPath(ProjectFileGenerator.MasterProjectRelativePath);
+					string FullProjectPath = ProjectFileGenerator.MasterProjectPath.FullName;
 					string FullPath = "";
 
 					if (CurPath.StartsWith("/") && !CurPath.StartsWith(FullProjectPath))
@@ -317,7 +317,7 @@ namespace UnrealBuildTool
 					}
 					else
 					{
-						FullPath = Path.GetFullPath(Path.Combine(Path.GetDirectoryName(KDevelopProject.ProjectFilePath), CurPath));
+						FullPath = Path.GetFullPath(Path.Combine(Path.GetDirectoryName(KDevelopProject.ProjectFilePath.FullName), CurPath));
 						FullPath = Utils.MakePathRelativeTo(FullPath, FullProjectPath);
 						FullPath = FullPath.TrimEnd('/');
 						FullPath = Path.Combine(UnrealEngineRootPath, FullPath);
@@ -431,14 +431,14 @@ namespace UnrealBuildTool
 			WriteExcludeSection (ref KDevelopFileContent);
 
 			// Write the master kdev file.
-			var FullMasterProjectPath = Path.Combine (MasterProjectRelativePath,".kdev4/");
+			var FullMasterProjectPath = Path.Combine (MasterProjectPath.FullName,".kdev4/");
 
 			if(!Directory.Exists(FullMasterProjectPath))
 			{
 				Directory.CreateDirectory (FullMasterProjectPath);
 			}
 
-			var FullKDevelopMasterFileName = Path.Combine(MasterProjectRelativePath, FileName);
+			var FullKDevelopMasterFileName = Path.Combine(MasterProjectPath.FullName, FileName);
 			var FullKDevelopFileName = Path.Combine (FullMasterProjectPath, FileName);
 
 			var FullDefinesFileName = Path.Combine (FullMasterProjectPath, DefinesFileName);
@@ -470,13 +470,13 @@ namespace UnrealBuildTool
 		/// </summary>
 		/// <param name="InitFilePath">Path to the project file</param>
 		/// <returns>The newly allocated project file object</returns>
-		protected override ProjectFile AllocateProjectFile( string InitFilePath )
+		protected override ProjectFile AllocateProjectFile( FileReference InitFilePath )
 		{
 			return new KDevelopProjectFile( InitFilePath );
 		}
 
 		/// ProjectFileGenerator interface
-		public override void CleanProjectFiles(string InMasterProjectRelativePath, string InMasterProjectName, string InIntermediateProjectFilesPath)
+		public override void CleanProjectFiles(DirectoryReference InMasterProjectDirectory, string InMasterProjectName, DirectoryReference InIntermediateProjectFilesDirectory)
 		{
 		}
 	}
diff --git a/Engine/Source/Programs/UnrealBuildTool/System/LinkEnvironment.cs b/Engine/Source/Programs/UnrealBuildTool/System/LinkEnvironment.cs
--- a/Engine/Source/Programs/UnrealBuildTool/System/LinkEnvironment.cs
+++ b/Engine/Source/Programs/UnrealBuildTool/System/LinkEnvironment.cs
@@ -13,19 +13,19 @@ namespace UnrealBuildTool
 	public class LinkEnvironmentConfiguration : NativeBuildEnvironmentConfiguration
 	{
 		/** The directory to put the non-executable files in (PDBs, import library, etc) */
-		public string OutputDirectory;
+		public DirectoryReference OutputDirectory;
 
 		/** Intermediate file directory */
-		public string IntermediateDirectory;
+		public DirectoryReference IntermediateDirectory;
 
 		/** The directory to shadow source files in for syncing to remote compile servers */
-		public string LocalShadowDirectory = null;
+		public DirectoryReference LocalShadowDirectory = null;
 
 		/** The file path for the executable file that is output by the linker. */
-		public List<string> OutputFilePaths = new List<string>();
+		public List<FileReference> OutputFilePaths = new List<FileReference>();
 
 		/** Returns the OutputFilePath is there is only one entry in OutputFilePaths */
-		public string OutputFilePath
+		public FileReference OutputFilePath
 		{
 			get
 			{
diff --git a/Engine/Source/Programs/UnrealBuildTool/System/MakefileGenerator.cs b/Engine/Source/Programs/UnrealBuildTool/System/MakefileGenerator.cs
--- a/Engine/Source/Programs/UnrealBuildTool/System/MakefileGenerator.cs
+++ b/Engine/Source/Programs/UnrealBuildTool/System/MakefileGenerator.cs
@@ -23,7 +23,7 @@ namespace UnrealBuildTool
 
 	public class MakefileProjectFile : ProjectFile
 	{
-		public MakefileProjectFile( string InitFilePath )
+		public MakefileProjectFile( FileReference InitFilePath )
 			: base(InitFilePath)
 		{
 		}
@@ -71,7 +71,7 @@ namespace UnrealBuildTool
 
             if (!String.IsNullOrEmpty(GameProjectName))
             {
-                GameProjectFile = UnrealBuildTool.GetUProjectFile();
+                GameProjectFile = UnrealBuildTool.GetUProjectFile().FullName;
                 MakeGameProjectFile = "GAMEPROJECTFILE =" + GameProjectFile + "\n";
                 ProjectBuildCommand = "PROJECTBUILD = mono $(UNREALROOTPATH)/Engine/Binaries/DotNET/UnrealBuildTool.exe\n";
             }
@@ -98,7 +98,7 @@ namespace UnrealBuildTool
 						continue;
 					}
 
-					string TargetFileName = Path.GetFileNameWithoutExtension(TargetFile.TargetFilePath);
+					string TargetFileName = TargetFile.TargetFilePath.GetFileNameWithoutExtension();
 					string Basename = TargetFileName.Substring(0, TargetFileName.LastIndexOf(".Target"));
 
 					foreach (UnrealTargetConfiguration CurConfiguration in Enum.GetValues(typeof(UnrealTargetConfiguration)))
@@ -133,7 +133,7 @@ namespace UnrealBuildTool
 						continue;
 					}
 
-					string TargetFileName = Path.GetFileNameWithoutExtension(TargetFile.TargetFilePath);
+					string TargetFileName = TargetFile.TargetFilePath.GetFileNameWithoutExtension();
 					string Basename = TargetFileName.Substring(0, TargetFileName.LastIndexOf(".Target"));
 
 					if (Basename == GameProjectName || Basename == (GameProjectName + "Editor")) 
@@ -185,8 +185,8 @@ namespace UnrealBuildTool
             }
 
             MakefileContent.Append("\n.PHONY: $(TARGETS)\n");
-            var FullFileName = Path.Combine(MasterProjectRelativePath, FileName);
-            return WriteFileIfChanged(FullFileName, MakefileContent.ToString());
+            FileReference FullFileName = FileReference.Combine(MasterProjectPath, FileName);
+            return WriteFileIfChanged(FullFileName.FullName, MakefileContent.ToString());
         }
 
         /// ProjectFileGenerator interface
@@ -208,13 +208,13 @@ namespace UnrealBuildTool
         /// </summary>
         /// <param name="InitFilePath">Path to the project file</param>
         /// <returns>The newly allocated project file object</returns>
-        protected override ProjectFile AllocateProjectFile( string InitFilePath )
+        protected override ProjectFile AllocateProjectFile( FileReference InitFilePath )
         {
             return new MakefileProjectFile( InitFilePath );
         }
 
         /// ProjectFileGenerator interface
-        public override void CleanProjectFiles(string InMasterProjectRelativePath, string InMasterProjectName, string InIntermediateProjectFilesPath)
+        public override void CleanProjectFiles(DirectoryReference InMasterProjectDirectory, string InMasterProjectName, DirectoryReference InIntermediateProjectFilesDirectory)
         {
         }
     }
diff --git a/Engine/Source/Programs/UnrealBuildTool/System/PluginDescriptor.cs b/Engine/Source/Programs/UnrealBuildTool/System/PluginDescriptor.cs
--- a/Engine/Source/Programs/UnrealBuildTool/System/PluginDescriptor.cs
+++ b/Engine/Source/Programs/UnrealBuildTool/System/PluginDescriptor.cs
@@ -86,9 +86,9 @@ namespace UnrealBuildTool
 		/// </summary>
 		/// <param name="FileName">The filename to read</param>
 		/// <returns>New plugin descriptor</returns>
-		public static PluginDescriptor FromFile(string FileName)
+		public static PluginDescriptor FromFile(FileReference FileName)
 		{
-			JsonObject RawObject = JsonObject.Read(FileName);
+			JsonObject RawObject = JsonObject.Read(FileName.FullName);
 			try
 			{
 				PluginDescriptor Descriptor = new PluginDescriptor();
diff --git a/Engine/Source/Programs/UnrealBuildTool/System/Plugins.cs b/Engine/Source/Programs/UnrealBuildTool/System/Plugins.cs
--- a/Engine/Source/Programs/UnrealBuildTool/System/Plugins.cs
+++ b/Engine/Source/Programs/UnrealBuildTool/System/Plugins.cs
@@ -24,10 +24,10 @@ namespace UnrealBuildTool
 		public readonly string Name;
 
 		// Path to the plugin
-		public readonly string FileName;
+		public readonly FileReference File;
 
 		// Path to the plugin's root directory
-		public readonly string Directory;
+		public readonly DirectoryReference Directory;
 
 		// The plugin descriptor
 		public PluginDescriptor Descriptor;
@@ -38,14 +38,14 @@ namespace UnrealBuildTool
 		/// <summary>
 		/// Constructs a PluginInfo object
 		/// </summary>
-		/// <param name="InFileName"></param>
+		/// <param name="InFile"></param>
 		/// <param name="InLoadedFrom">Where this pl</param>
-		public PluginInfo(string InFileName, PluginLoadedFrom InLoadedFrom)
+		public PluginInfo(FileReference InFile, PluginLoadedFrom InLoadedFrom)
 		{
-			Name = Path.GetFileNameWithoutExtension(InFileName);
-			FileName = Path.GetFullPath(InFileName);
-			Directory = Path.GetDirectoryName(FileName);
-			Descriptor = PluginDescriptor.FromFile(FileName);
+			Name = Path.GetFileNameWithoutExtension(InFile.FullName);
+			File = InFile;
+			Directory = File.Directory;
+			Descriptor = PluginDescriptor.FromFile(File);
 			LoadedFrom = InLoadedFrom;
 		}
 	}
@@ -58,23 +58,23 @@ namespace UnrealBuildTool
 		/// <param name="EngineDir">Path to the engine directory</param>
 		/// <param name="ProjectFileName">Path to the project file (or null)</param>
 		/// <returns>Sequence of PluginInfo objects, one for each discovered plugin</returns>
-		public static List<PluginInfo> ReadAvailablePlugins(string EngineDir, string ProjectFileName)
+		public static List<PluginInfo> ReadAvailablePlugins(DirectoryReference EngineDirectoryName, FileReference ProjectFileName)
 		{
 			List<PluginInfo> Plugins = new List<PluginInfo>();
 
 			// Read all the engine plugins
-			string EnginePluginsDir = Path.Combine(EngineDir, "Plugins");
-			foreach(string PluginFileName in EnumeratePlugins(EnginePluginsDir))
+			DirectoryReference EnginePluginsDirectoryName = DirectoryReference.Combine(EngineDirectoryName, "Plugins");
+			foreach(FileReference PluginFileName in EnumeratePlugins(EnginePluginsDirectoryName))
 			{
 				PluginInfo Plugin = new PluginInfo(PluginFileName, PluginLoadedFrom.Engine);
 				Plugins.Add(Plugin);
 			}
 
 			// Read all the project plugins
-			if(!String.IsNullOrEmpty(ProjectFileName))
+			if(ProjectFileName != null)
 			{
-				string ProjectPluginsDir = Path.Combine(Path.GetDirectoryName(ProjectFileName), "Plugins");
-				foreach(string PluginFileName in EnumeratePlugins(ProjectPluginsDir))
+				DirectoryReference ProjectPluginsDir = DirectoryReference.Combine(ProjectFileName.Directory, "Plugins");
+				foreach(FileReference PluginFileName in EnumeratePlugins(ProjectPluginsDir))
 				{
 					PluginInfo Plugin = new PluginInfo(PluginFileName, PluginLoadedFrom.GameProject);
 					Plugins.Add(Plugin);
@@ -88,13 +88,12 @@ namespace UnrealBuildTool
 		/// Find paths to all the plugins under a given parent directory (recursively)
 		/// </summary>
 		/// <param name="ParentDirectory">Parent directory to look in. Plugins will be found in any *subfolders* of this directory.</param>
-		/// <param name="FileNames">List of filenames. Will have all the discovered .uplugin files appended to it.</param>
-		public static IEnumerable<string> EnumeratePlugins(string ParentDirectory)
+		public static IEnumerable<FileReference> EnumeratePlugins(DirectoryReference ParentDirectory)
 		{
-			List<string> FileNames = new List<string>();
-			if(Directory.Exists(ParentDirectory))
+			List<FileReference> FileNames = new List<FileReference>();
+			if(ParentDirectory.Exists())
 			{
-				EnumeratePluginsInternal(new DirectoryInfo(ParentDirectory), FileNames);
+				EnumeratePluginsInternal(new DirectoryInfo(ParentDirectory.FullName), FileNames);
 			}
 			return FileNames;
 		}
@@ -104,14 +103,14 @@ namespace UnrealBuildTool
 		/// </summary>
 		/// <param name="ParentDirectory">Parent directory to look in. Plugins will be found in any *subfolders* of this directory.</param>
 		/// <param name="FileNames">List of filenames. Will have all the discovered .uplugin files appended to it.</param>
-		static void EnumeratePluginsInternal(DirectoryInfo ParentDirectory, List<string> FileNames)
+		static void EnumeratePluginsInternal(DirectoryInfo ParentDirectory, List<FileReference> FileNames)
 		{
 			foreach(DirectoryInfo ChildDirectory in ParentDirectory.EnumerateDirectories())
 			{
 				int InitialFileNamesCount = FileNames.Count;
 				foreach(FileInfo PluginFileInfo in ChildDirectory.EnumerateFiles("*.uplugin", SearchOption.TopDirectoryOnly))
 				{
-					FileNames.Add(PluginFileInfo.FullName);
+					FileNames.Add(new FileReference(PluginFileInfo.FullName));
 				}
 				if(FileNames.Count == InitialFileNamesCount)
 				{
diff --git a/Engine/Source/Programs/UnrealBuildTool/System/Project.cs b/Engine/Source/Programs/UnrealBuildTool/System/Project.cs
--- a/Engine/Source/Programs/UnrealBuildTool/System/Project.cs
+++ b/Engine/Source/Programs/UnrealBuildTool/System/Project.cs
@@ -35,7 +35,7 @@ namespace UnrealBuildTool
 	public class ProjectTarget
 	{
 		/// The target rules file path on disk, if we have one
-		public string TargetFilePath;
+		public FileReference TargetFilePath;
 
 		/// Optional target rules for this target.  If the target came from a *.Target.cs file on disk, then it will have one of these.
 		/// For targets that are synthetic (like UnrealBuildTool or other manually added project files) we won't have a rules object for those.
@@ -55,7 +55,7 @@ namespace UnrealBuildTool
 
 		public override string ToString()
 		{
-			return Path.GetFileNameWithoutExtension(TargetFilePath);
+			return TargetFilePath.GetFileNameWithoutExtension();
 		}
 	}
 
@@ -89,10 +89,10 @@ namespace UnrealBuildTool
 			/// </summary>
 			/// <param name="InitFilePath">Path to the source file on disk</param>
 			/// <param name="InitRelativeBaseFolder">The directory on this the path within the project will be relative to</param>
-			public SourceFile( string InitFilePath, string InitRelativeBaseFolder )
+			public SourceFile( FileReference InReference, DirectoryReference InBaseFolder )
 			{
-				FilePath = InitFilePath;
-				RelativeBaseFolder = InitRelativeBaseFolder;
+				Reference = InReference;
+				BaseFolder = InBaseFolder;
 			}
 
 			public SourceFile()
@@ -102,16 +102,16 @@ namespace UnrealBuildTool
 			/// <summary>
 			/// File path to file on disk
 			/// </summary>
-			public string FilePath
+			public FileReference Reference
 			{
 				get;
 				private set;
 			}
 
 			/// <summary>
-			/// Optional directory that overrides where files in this project are relative to when displayed in the IDE.  If null, will default to the project's RelativeBaseFolder.
+			/// Optional directory that overrides where files in this project are relative to when displayed in the IDE.  If null, will default to the project's BaseFolder.
 			/// </summary>
-			public string RelativeBaseFolder
+			public DirectoryReference BaseFolder
 			{
 				get;
 				private set;
@@ -123,25 +123,15 @@ namespace UnrealBuildTool
 		/// Constructs a new project file object
 		/// </summary>
 		/// <param name="InitFilePath">The path to the project file, relative to the master project file</param>
-		protected ProjectFile( string InitRelativeFilePath )
+		protected ProjectFile( FileReference InProjectFilePath )
 		{
-			RelativeProjectFilePath = InitRelativeFilePath;
+			ProjectFilePath = InProjectFilePath;
 			ShouldBuildByDefaultForSolutionTargets = true;
 		}
 
 
-		/// Full path to the project file on disk
-		public string ProjectFilePath
-		{
-			get
-			{
-				return Path.Combine( ProjectFileGenerator.MasterProjectRelativePath, RelativeProjectFilePath );
-			}
-		}
-
-
-		/// Project file path, relative to the master project
-		public string RelativeProjectFilePath
+		/// Project file path
+		public FileReference ProjectFilePath
 		{
 			get;
 			private set;
@@ -198,11 +188,11 @@ namespace UnrealBuildTool
 		/// </summary>
 		/// <param name="FilesToAdd">Files to add</param>
 		/// <param name="RelativeBaseFolder">The directory the path within the project will be relative to</param>
-		public void AddFilesToProject( List<string> FilesToAdd, string RelativeBaseFolder )
+		public void AddFilesToProject( List<FileReference> FilesToAdd, DirectoryReference BaseFolder )
 		{
 			foreach( var CurFile in FilesToAdd )
 			{
-				AddFileToProject( CurFile, RelativeBaseFolder );
+				AddFileToProject( CurFile, BaseFolder );
 			}
 		}
 
@@ -222,29 +212,21 @@ namespace UnrealBuildTool
 		/// Adds a file to this project, ignoring dupes
 		/// </summary>
 		/// <param name="FilePath">Path to the file on disk</param>
-		/// <param name="RelativeBaseFolder">The directory the path within the project will be relative to</param>
-		public void AddFileToProject( string FilePath, string RelativeBaseFolder )
+		/// <param name="BaseFolder">The directory the path within the project will be relative to</param>
+		public void AddFileToProject( FileReference FilePath, DirectoryReference BaseFolder )
 		{
 			// Don't add duplicates
 			SourceFile ExistingFile = null;
             if (SourceFileMap.TryGetValue(FilePath, out ExistingFile))
             {
-                if( ExistingFile.RelativeBaseFolder != RelativeBaseFolder )
+                if( ExistingFile.BaseFolder != BaseFolder )
                 {
-                    if( ( ExistingFile.RelativeBaseFolder != null ) != ( RelativeBaseFolder != null ) ||
-                        !ExistingFile.RelativeBaseFolder.Equals( RelativeBaseFolder, StringComparison.InvariantCultureIgnoreCase ) )
-                    {
-                        throw new BuildException( "Trying to add file '" + FilePath + "' to project '" + ProjectFilePath + "' when the file already exists, but with a different relative base folder '" + RelativeBaseFolder + "' is different than the current file's '" + ExistingFile.RelativeBaseFolder + "'!" );
-                    }
-                    else
-                    {
-                        throw new BuildException( "Trying to add file '" + FilePath + "' to project '" + ProjectFilePath + "' when the file already exists, but the specified project relative base folder is different than the current file's!" );
-                    }
+                    throw new BuildException( "Trying to add file '" + FilePath + "' to project '" + ProjectFilePath + "' when the file already exists, but with a different relative base folder '" + BaseFolder + "' is different than the current file's '" + ExistingFile.BaseFolder + "'!" );
                 }
 			}
 			else
 			{
-				SourceFile File = AllocSourceFile( FilePath, RelativeBaseFolder );
+				SourceFile File = AllocSourceFile( FilePath, BaseFolder );
 				if( File != null )
 				{
 					SourceFileMap[FilePath] = File;
@@ -412,7 +394,7 @@ namespace UnrealBuildTool
 		/// <param name="InitFilePath">Path to the source file on disk</param>
 		/// <param name="InitProjectSubFolder">Optional sub-folder to put the file in.  If empty, this will be determined automatically from the file's path relative to the project file</param>
 		/// <returns>The newly allocated source file object</returns>
-		public virtual SourceFile AllocSourceFile( string InitFilePath, string InitProjectSubFolder = null )
+		public virtual SourceFile AllocSourceFile( FileReference InitFilePath, DirectoryReference InitProjectSubFolder = null )
 		{
 			return new SourceFile( InitFilePath, InitProjectSubFolder );
 		}
@@ -425,27 +407,24 @@ namespace UnrealBuildTool
 			{
 				return InputPath;
 			}
-
-			// Otherwise make sure it's absolute
-			string FullInputPath = Utils.CleanDirectorySeparators(Path.GetFullPath(InputPath));
-
-			// Try to make it relative to the solution directory.
-			string FullSolutionPath = Utils.CleanDirectorySeparators(Path.GetFullPath(ProjectFileGenerator.MasterProjectRelativePath));
-			if (FullSolutionPath.Last() != Path.DirectorySeparatorChar)
+			else
 			{
-				FullSolutionPath += Path.DirectorySeparatorChar;
+				return NormalizeProjectPath(new FileReference(InputPath));
 			}
-			if (FullInputPath.StartsWith(FullSolutionPath))
+		}
+
+		/** Takes the given path and tries to rebase it relative to the project. */
+		public string NormalizeProjectPath(FileReference InputPath)
+		{
+			// Try to make it relative to the solution directory.
+			if(InputPath.IsUnderDirectory(ProjectFileGenerator.MasterProjectPath))
 			{
-				FullInputPath = Utils.MakePathRelativeTo(Utils.CleanDirectorySeparators(FullInputPath), Path.GetDirectoryName(Path.GetFullPath(ProjectFilePath)));
+				return InputPath.MakeRelativeTo(ProjectFileGenerator.IntermediateProjectFilesPath);
 			}
 			else
 			{
-				FullInputPath = Utils.CleanDirectorySeparators(FullInputPath);
+				return InputPath.FullName;
 			}
-
-			// Otherwise return the input
-			return FullInputPath;
 		}
 
 		/** Takes the given path, normalizes it, and quotes it if necessary. */
@@ -462,11 +441,11 @@ namespace UnrealBuildTool
 		/** Visualizer for the debugger */
 		public override string ToString()
 		{
-			return RelativeProjectFilePath;
+			return ProjectFilePath.ToString();
 		}
 
 		/// Map of file paths to files in the project.
-        private readonly Dictionary<string, SourceFile> SourceFileMap = new Dictionary<string, SourceFile>(StringComparer.InvariantCultureIgnoreCase);
+        private readonly Dictionary<FileReference, SourceFile> SourceFileMap = new Dictionary<FileReference, SourceFile>();
 
 		/// Files in this project
 		public readonly List<SourceFile> SourceFiles = new List<SourceFile>();
