Use FileReference and DirectoryReference classes instead of strings for paths. Massively reduces the amount of paranoid calls to GetFullPath(), CleanDirectorySeparators() et al., and improves performa

diff --git a/Engine/Source/Programs/UnrealBuildTool/System/ProjectFileGenerator.cs b/Engine/Source/Programs/UnrealBuildTool/System/ProjectFileGenerator.cs
--- a/Engine/Source/Programs/UnrealBuildTool/System/ProjectFileGenerator.cs
+++ b/Engine/Source/Programs/UnrealBuildTool/System/ProjectFileGenerator.cs
@@ -192,17 +192,17 @@ namespace UnrealBuildTool
 		public static readonly string EngineRelativePath = Path.Combine( RootRelativePath, "Engine" );
 
 		/// Relative path to the directory where the master project file will be saved to
-		public static string MasterProjectRelativePath = RootRelativePath;		// We'll save the master project to our "root" folder
+		public static DirectoryReference MasterProjectPath = UnrealBuildTool.RootDirectory; // We'll save the master project to our "root" folder
 
 		/// Name of the UE4 engine project name that contains all of the engine code, config files and other files
 		static readonly string EngineProjectFileNameBase = "UE4";
 
 		/// When ProjectsAreIntermediate is true, this is the directory to store generated project files
 		// @todo projectfiles: Ideally, projects for game modules/targets would be created in the game's Intermediate folder!
-		public static string IntermediateProjectFilesPath = Path.Combine( EngineRelativePath, "Intermediate", "ProjectFiles" );
+		public static DirectoryReference IntermediateProjectFilesPath = DirectoryReference.Combine( UnrealBuildTool.EngineDirectory, "Intermediate", "ProjectFiles" );
 
 		/// Path to timestamp file, recording when was the last time projects were created.
-		public static string ProjectTimestampFile = Path.Combine(IntermediateProjectFilesPath, "Timestamp");
+		public static string ProjectTimestampFile = Path.Combine(IntermediateProjectFilesPath.FullName, "Timestamp");
 
 		/// Global static new line string used by ProjectFileGenerator to generate project files.
 		public static readonly string NewLine = Environment.NewLine;
@@ -251,11 +251,11 @@ namespace UnrealBuildTool
 		{
 			var Folder = ProgramsFolder.AddSubFolder("Automation");
 			var AllGameFolders = UEBuildTarget.DiscoverAllGameFolders();
-			var BuildFolders = new List<string>(AllGameFolders.Count);
+			var BuildFolders = new List<DirectoryReference>(AllGameFolders.Count);
 			foreach (var GameFolder in AllGameFolders)
 			{
-				var GameBuildFolder = Path.Combine(GameFolder, "Build");
-				if (Directory.Exists(GameBuildFolder))
+				var GameBuildFolder = DirectoryReference.Combine(GameFolder, "Build");
+				if (GameBuildFolder.Exists())
 				{
 					BuildFolders.Add(GameBuildFolder);
 				}
@@ -265,11 +265,9 @@ namespace UnrealBuildTool
 			var ModuleFiles = RulesCompiler.FindAllRulesSourceFiles(RulesCompiler.RulesFileType.AutomationModule, null, ForeignPlugins:null, AdditionalSearchPaths: BuildFolders );
 			foreach (var ProjectFile in ModuleFiles)
 			{
-				FileInfo Info = new FileInfo(Path.Combine(ProjectFile));
-				if (Info.Exists)
+				if (ProjectFile.Exists())
 				{
-					var RelativeFileName = Utils.MakePathRelativeTo(ProjectFile, MasterProjectRelativePath);
-					var Project = new VCSharpProjectFile(RelativeFileName);
+					var Project = new VCSharpProjectFile(ProjectFile);
 					Project.ShouldBuildForAllSolutionTargets = true;
 					AddExistingProjectFile(Project, bForceDevelopmentConfiguration: true);
                     AutomationProjectFiles.Add( Project );
@@ -305,27 +303,27 @@ namespace UnrealBuildTool
 			{
 				Log.TraceInformation("Discovering modules, targets and source code for game...");
 
-				MasterProjectRelativePath = UnrealBuildTool.GetUProjectPath();
+				MasterProjectPath = UnrealBuildTool.GetUProjectPath();
 					
 					// Set the project file name
-				MasterProjectName = Path.GetFileNameWithoutExtension(UnrealBuildTool.GetUProjectFile());
+				MasterProjectName = Path.GetFileNameWithoutExtension(UnrealBuildTool.GetUProjectFile().FullName);
 
-				if (!Directory.Exists(MasterProjectRelativePath + "/Source"))
+				if (!DirectoryReference.Combine(MasterProjectPath, "Source").Exists())
 				{
-					if (!Directory.Exists(MasterProjectRelativePath + "/Intermediate/Source"))
+					if (!DirectoryReference.Combine(MasterProjectPath, "Intermediate", "Source").Exists())
 					{
 						if (BuildHostPlatform.Current.Platform == UnrealTargetPlatform.Mac)
 						{
-							MasterProjectRelativePath = Path.GetFullPath(Path.Combine(Path.GetDirectoryName(Assembly.GetEntryAssembly().GetOriginalLocation()), "..", "..", "..", "Engine"));
+							MasterProjectPath = UnrealBuildTool.EngineDirectory;
 							GameProjectName = "UE4Game";
 						}
-						if (!Directory.Exists(MasterProjectRelativePath + "/Source"))
+						if (!DirectoryReference.Combine(MasterProjectPath, "Source").Exists())
 						{
-							throw new BuildException("Directory '{0}' is missing 'Source' folder.", MasterProjectRelativePath);
+							throw new BuildException("Directory '{0}' is missing 'Source' folder.", MasterProjectPath);
 						}
 					}
 				}
-				IntermediateProjectFilesPath = Path.Combine(MasterProjectRelativePath, "Intermediate", "ProjectFiles");
+				IntermediateProjectFilesPath = DirectoryReference.Combine(MasterProjectPath, "Intermediate", "ProjectFiles");
 			}
 			else if( bGeneratingRocketProjectFiles )
 			{
@@ -336,17 +334,17 @@ namespace UnrealBuildTool
 				//       bIncludeEngineModulesInRocketProjects=true (defaults to false.)				
 				bExcludeNoRedistFiles = true;
 
-				MasterProjectRelativePath = UnrealBuildTool.GetUProjectPath();
-				IntermediateProjectFilesPath = Path.Combine( MasterProjectRelativePath, "Intermediate", "ProjectFiles" );
+				MasterProjectPath = UnrealBuildTool.GetUProjectPath();
+				IntermediateProjectFilesPath = DirectoryReference.Combine( MasterProjectPath, "Intermediate", "ProjectFiles" );
 
 				// Set the project file name
-				MasterProjectName = Path.GetFileNameWithoutExtension(UnrealBuildTool.GetUProjectFile());
+				MasterProjectName = Path.GetFileNameWithoutExtension(UnrealBuildTool.GetUProjectFile().FullName);
 
-				if (!Directory.Exists(MasterProjectRelativePath + "/Source"))
+				if (!DirectoryReference.Combine(MasterProjectPath, "Source").Exists())
 				{
-					if (!Directory.Exists(MasterProjectRelativePath + "/Intermediate/Source"))
+					if (!DirectoryReference.Combine(MasterProjectPath, "Intermediate", "Source").Exists())
 					{
-						throw new BuildException("Directory '{0}' is missing 'Source' folder.", MasterProjectRelativePath);
+						throw new BuildException("Directory '{0}' is missing 'Source' folder.", MasterProjectPath);
 					}
 				}
 			}
@@ -366,14 +364,14 @@ namespace UnrealBuildTool
 				foreach (string SortedPlatform in SortedPlatformNames)
 				{
 					MasterProjectName += SortedPlatform;
-					IntermediateProjectFilesPath += SortedPlatform;
+					IntermediateProjectFilesPath = new DirectoryReference(IntermediateProjectFilesPath.FullName + SortedPlatform);
 				}
 			}
 
 			bool bCleanProjectFiles = UnrealBuildTool.CommandLineContains( "-CleanProjects" );
 			if (bCleanProjectFiles)
 			{
-				CleanProjectFiles(MasterProjectRelativePath, MasterProjectName, IntermediateProjectFilesPath);
+				CleanProjectFiles(MasterProjectPath, MasterProjectName, IntermediateProjectFilesPath);
 			}
 
 			// Figure out which platforms we should generate project files for.
@@ -402,7 +400,7 @@ namespace UnrealBuildTool
 			var AllModuleFiles = DiscoverModules(AllGameProjects);
 
 			ProjectFile EngineProject = null;
-			Dictionary<string, ProjectFile> GameProjects = null;
+			Dictionary<DirectoryReference, ProjectFile> GameProjects = null;
 			Dictionary<string, ProjectFile> ProgramProjects = null;
 			HashSet<ProjectFile> TemplateGameProjects = null;
 			{
@@ -424,10 +422,10 @@ namespace UnrealBuildTool
 					// If we're still missing an engine project because we don't have any targets for it, make one up.
 					if( EngineProject == null )
 					{
-						string ProjectFilePath = Path.Combine(IntermediateProjectFilesPath, "UE4" + ProjectFileExtension);
+						FileReference ProjectFilePath = FileReference.Combine(IntermediateProjectFilesPath, "UE4" + ProjectFileExtension);
 
 						bool bAlreadyExisted;
-						EngineProject = FindOrAddProject(Utils.MakePathRelativeTo(ProjectFilePath, MasterProjectRelativePath), true, out bAlreadyExisted);
+						EngineProject = FindOrAddProject(ProjectFilePath, true, out bAlreadyExisted);
 
 						EngineProject.IsForeignProject = false;
 						EngineProject.IsGeneratedProject = true;
@@ -569,14 +567,14 @@ namespace UnrealBuildTool
 				{
 					// Figure out which targets we need about IntelliSense for.  We only need to worry about targets for projects
 					// that we're actually generating in this session.
-					var IntelliSenseTargetFiles = new List<Tuple<ProjectFile, string>>();
+					var IntelliSenseTargetFiles = new List<Tuple<ProjectFile, FileReference>>();
 					{
 						// Engine targets
 						if( EngineProject != null)
 						{
 							foreach( var ProjectTarget in EngineProject.ProjectTargets )
 							{
-								if( !String.IsNullOrEmpty( ProjectTarget.TargetFilePath ) )
+								if( ProjectTarget.TargetFilePath != null )
 								{
 									// Only bother with the editor target.  We want to make sure that definitions are setup to be as inclusive as possible
 									// for good quality IntelliSense.  For example, we want WITH_EDITORONLY_DATA=1, so using the editor targets works well.
@@ -593,7 +591,7 @@ namespace UnrealBuildTool
 						{
 							foreach( var ProjectTarget in ProgramProject.ProjectTargets )
 							{
-								if( !String.IsNullOrEmpty( ProjectTarget.TargetFilePath ) )
+								if( ProjectTarget.TargetFilePath != null )
 								{
 									IntelliSenseTargetFiles.Add( Tuple.Create( ProgramProject, ProjectTarget.TargetFilePath ) );
 								}
@@ -605,7 +603,7 @@ namespace UnrealBuildTool
 						{
 							foreach( var ProjectTarget in GameProject.ProjectTargets )
 							{
-								if( !String.IsNullOrEmpty( ProjectTarget.TargetFilePath ) )
+								if( ProjectTarget.TargetFilePath != null )
 								{
 									// Only bother with the editor target.  We want to make sure that definitions are setup to be as inclusive as possible
 									// for good quality IntelliSense.  For example, we want WITH_EDITORONLY_DATA=1, so using the editor targets works well.
@@ -675,7 +673,7 @@ namespace UnrealBuildTool
 		/// <param name="InMasterProjectRelativePath">The MasterProjectRelativePath</param>
 		/// <param name="InMasterProjectName">The name of the master project</param>
 		/// <param name="InIntermediateProjectFilesPath">The intermediate path of project files</param>
-		public abstract void CleanProjectFiles(string InMasterProjectRelativePath, string InMasterProjectName, string InIntermediateProjectFilesPath);
+		public abstract void CleanProjectFiles(DirectoryReference InMasterProjectDirectory, string InMasterProjectName, DirectoryReference InIntermediateProjectFilesDirectory);
 
 		/// <summary>
 		/// Configures project generator based on command-line options
@@ -806,7 +804,7 @@ namespace UnrealBuildTool
 				}
 
 				// Make sure we can get a valid game name out of this project
-				var GameName = Path.GetFileNameWithoutExtension(UnrealBuildTool.GetUProjectFile());
+				var GameName = UnrealBuildTool.GetUProjectFile().GetFileNameWithoutExtension();
 				if( String.IsNullOrEmpty( GameName ) )
 				{
 					throw new BuildException("A valid Rocket game project was not found in the specified location (" + UnrealBuildTool.GetUProjectPath() + ")");
@@ -822,7 +820,7 @@ namespace UnrealBuildTool
 			}
 			else if( bGeneratingGameProjectFiles )
 			{
-				GameProjectName = Path.GetFileNameWithoutExtension(UnrealBuildTool.GetUProjectFile());
+				GameProjectName = UnrealBuildTool.GetUProjectFile().GetFileNameWithoutExtension();
 				if (String.IsNullOrEmpty(GameProjectName))
 				{
 					throw new BuildException("A valid game project was not found in the specified location (" + UnrealBuildTool.GetUProjectPath() + ")");
@@ -846,27 +844,24 @@ namespace UnrealBuildTool
 		/// <summary>
 		/// Adds all game project files, including target projects and config files
 		/// </summary>
-		private void AddAllGameProjects(Dictionary<string, ProjectFile> GameProjects, string SupportedPlatformNames, MasterProjectFolder ProjectsFolder)
+		private void AddAllGameProjects(Dictionary<DirectoryReference, ProjectFile> GameProjects, string SupportedPlatformNames, MasterProjectFolder ProjectsFolder)
 		{
 			foreach( var GameFolderAndProjectFile in GameProjects )
 			{
-				var GameFolderName = GameFolderAndProjectFile.Key;
+				var GameProjectDirectory = GameFolderAndProjectFile.Key;
 
 				// @todo projectfiles: We have engine localization files, but should we also add GAME localization files?
 
-				string GameProjectDirectory = GameFolderName;
-				GameProjectDirectory = GameProjectDirectory.Replace("/", "\\");
-
 				// Game config files
 				if( IncludeConfigFiles )
 				{
-					var GameConfigDirectory = Path.Combine(GameProjectDirectory, "Config");
-					if( Directory.Exists( GameConfigDirectory ) )
+					var GameConfigDirectory = DirectoryReference.Combine(GameProjectDirectory, "Config");
+					if( GameConfigDirectory.Exists( ) )
 					{
 						var GameProjectFile = GameFolderAndProjectFile.Value;
-						var DirectoriesToSearch = new List<string>();
+						var DirectoriesToSearch = new List<DirectoryReference>();
 						DirectoriesToSearch.Add( GameConfigDirectory );
-						GameProjectFile.AddFilesToProject( SourceFileSearch.FindFiles( DirectoriesToSearch, ExcludeNoRedistFiles:bExcludeNoRedistFiles ), GameFolderName );
+						GameProjectFile.AddFilesToProject( SourceFileSearch.FindFiles( DirectoriesToSearch, ExcludeNoRedistFiles:bExcludeNoRedistFiles ), GameProjectDirectory );
 					}
 				}
 			}
@@ -876,12 +871,12 @@ namespace UnrealBuildTool
 		/// Adds all engine localization text files to the specified project
 		private void AddEngineLocalizationFiles( ProjectFile EngineProject )
 		{
-			var EngineLocalizationDirectory = Path.Combine( EngineRelativePath, "Content", "Localization" );
-			if( Directory.Exists( EngineLocalizationDirectory ) )
+			var EngineLocalizationDirectory = DirectoryReference.Combine( UnrealBuildTool.EngineDirectory, "Content", "Localization" );
+			if( EngineLocalizationDirectory.Exists( ) )
 			{
-				var DirectoriesToSearch = new List<string>();
+				var DirectoriesToSearch = new List<DirectoryReference>();
 				DirectoriesToSearch.Add( EngineLocalizationDirectory );
-				EngineProject.AddFilesToProject( SourceFileSearch.FindFiles( DirectoriesToSearch, ExcludeNoRedistFiles:bExcludeNoRedistFiles), EngineRelativePath );
+				EngineProject.AddFilesToProject( SourceFileSearch.FindFiles( DirectoriesToSearch, ExcludeNoRedistFiles:bExcludeNoRedistFiles), UnrealBuildTool.EngineDirectory );
 			}
 		}
 
@@ -889,12 +884,12 @@ namespace UnrealBuildTool
         /// Adds all engine template text files to the specified project
         private void AddEngineTemplateFiles( ProjectFile EngineProject )
         {
-            var EngineTemplateDirectory = Path.Combine(EngineRelativePath, "Content", "Editor", "Templates");
-            if (Directory.Exists(EngineTemplateDirectory))
+            var EngineTemplateDirectory = DirectoryReference.Combine(UnrealBuildTool.EngineDirectory, "Content", "Editor", "Templates");
+            if (EngineTemplateDirectory.Exists())
             {
-				var DirectoriesToSearch = new List<string>();
+				var DirectoriesToSearch = new List<DirectoryReference>();
 				DirectoriesToSearch.Add( EngineTemplateDirectory );
-				EngineProject.AddFilesToProject( SourceFileSearch.FindFiles( DirectoriesToSearch, ExcludeNoRedistFiles:bExcludeNoRedistFiles), EngineRelativePath );
+				EngineProject.AddFilesToProject( SourceFileSearch.FindFiles( DirectoriesToSearch, ExcludeNoRedistFiles:bExcludeNoRedistFiles), UnrealBuildTool.EngineDirectory );
 			}
         }
 
@@ -902,26 +897,26 @@ namespace UnrealBuildTool
 		/// Adds all engine config files to the specified project
 		private void AddEngineConfigFiles( ProjectFile EngineProject )
 		{
-            var EngineConfigDirectory = Path.Combine(EngineRelativePath, "Config" );
-			if( Directory.Exists( EngineConfigDirectory ) )
+            var EngineConfigDirectory = DirectoryReference.Combine(UnrealBuildTool.EngineDirectory, "Config" );
+			if( EngineConfigDirectory.Exists( ) )
 			{
-				var DirectoriesToSearch = new List<string>();
+				var DirectoriesToSearch = new List<DirectoryReference>();
 				DirectoriesToSearch.Add( EngineConfigDirectory );
 
-				EngineProject.AddFilesToProject( SourceFileSearch.FindFiles( DirectoriesToSearch, ExcludeNoRedistFiles:bExcludeNoRedistFiles), EngineRelativePath );
+				EngineProject.AddFilesToProject( SourceFileSearch.FindFiles( DirectoriesToSearch, ExcludeNoRedistFiles:bExcludeNoRedistFiles), UnrealBuildTool.EngineDirectory );
 			}
 		}
 
 		/// Adds UnrealHeaderTool config files to the specified project
 		private void AddUnrealHeaderToolConfigFiles(ProjectFile EngineProject)
 		{
-			var UHTConfigDirectory = Path.Combine(EngineRelativePath, "Programs", "UnrealHeaderTool", "Config");
-			if (Directory.Exists(UHTConfigDirectory))
+			var UHTConfigDirectory = DirectoryReference.Combine(UnrealBuildTool.EngineDirectory, "Programs", "UnrealHeaderTool", "Config");
+			if (UHTConfigDirectory.Exists())
 			{
-				var DirectoriesToSearch = new List<string>();
+				var DirectoriesToSearch = new List<DirectoryReference>();
 				DirectoriesToSearch.Add(UHTConfigDirectory);
 
-				EngineProject.AddFilesToProject(SourceFileSearch.FindFiles(DirectoriesToSearch, ExcludeNoRedistFiles: bExcludeNoRedistFiles), EngineRelativePath);
+				EngineProject.AddFilesToProject(SourceFileSearch.FindFiles(DirectoriesToSearch, ExcludeNoRedistFiles: bExcludeNoRedistFiles), UnrealBuildTool.EngineDirectory);
 			}
 		}
 
@@ -929,27 +924,25 @@ namespace UnrealBuildTool
 		/// Finds all module files (filtering by platform)
 		/// </summary>
 		/// <returns>Filtered list of module files</returns>
-		protected List<string> DiscoverModules(List<UProjectInfo> AllGameProjects)
+		protected List<FileReference> DiscoverModules(List<UProjectInfo> AllGameProjects)
 		{
-			var AllModuleFiles = new List<string>();
+			List<FileReference> AllModuleFiles = new List<FileReference>();
 
 			// Locate all modules (*.Build.cs files)
-			var FoundModuleFiles = RulesCompiler.FindAllRulesSourceFiles( RulesCompiler.RulesFileType.Module, GameFolders: AllGameProjects.Select(x => x.Folder).ToList(), ForeignPlugins:null, AdditionalSearchPaths:null );
+			List<FileReference> FoundModuleFiles = RulesCompiler.FindAllRulesSourceFiles( RulesCompiler.RulesFileType.Module, GameFolders: AllGameProjects.Select(x => x.Folder).ToList(), ForeignPlugins:null, AdditionalSearchPaths:null );
 			foreach( var BuildFileName in FoundModuleFiles )
 			{
-				var CleanBuildFileName = Utils.CleanDirectorySeparators( BuildFileName );
-
 				bool IncludeThisModule = true;
 
 				// Skip NoRedist files if necessary
 				if( bExcludeNoRedistFiles )
 				{
-					IncludeThisModule = !IsNoRedistModule(CleanBuildFileName);
+					IncludeThisModule = !IsNoRedistModule(BuildFileName);
 				}
 
 				if( IncludeThisModule )
 				{
-					AllModuleFiles.Add( CleanBuildFileName );
+					AllModuleFiles.Add( BuildFileName );
 				}
 			}
 			return AllModuleFiles;
@@ -970,11 +963,11 @@ namespace UnrealBuildTool
 		/// <param name="CleanBuildFileName"></param>
 		/// <param name="IncludeThisModule"></param>
 		/// <returns></returns>
-		private static bool IsNoRedistModule(string ModulePath)
+		private static bool IsNoRedistModule(FileReference ModulePath)
 		{
 			foreach (var NoRedistFolderName in NoRedistFolders)
 			{
-				if (ModulePath.IndexOf(NoRedistFolderName, StringComparison.InvariantCultureIgnoreCase) >= 0)
+				if (ModulePath.FullName.IndexOf(NoRedistFolderName, StringComparison.InvariantCultureIgnoreCase) >= 0)
 				{
 					return true;
 				}
@@ -986,18 +979,18 @@ namespace UnrealBuildTool
 		/// Finds all target files (filtering by platform)
 		/// </summary>
 		/// <returns>Filtered list of target files</returns>
-		protected List<string> DiscoverTargets(List<UProjectInfo> AllGameProjects)
+		protected List<FileReference> DiscoverTargets(List<UProjectInfo> AllGameProjects)
 		{
-			var AllTargetFiles = new List<string>();
+			List<FileReference> AllTargetFiles = new List<FileReference>();
 
 			// Make a list of all platform name strings that we're *not* including in the project files
 			var UnsupportedPlatformNameStrings = Utils.MakeListOfUnsupportedPlatforms( SupportedPlatforms );
 
 			// Locate all targets (*.Target.cs files)
-			var FoundTargetFiles = RulesCompiler.FindAllRulesSourceFiles( RulesCompiler.RulesFileType.Target, AllGameProjects.Select(x => x.Folder).ToList(), ForeignPlugins:null, AdditionalSearchPaths:null );
-			foreach( var CurTargetFile in FoundTargetFiles )
+			List<FileReference> FoundTargetFiles = RulesCompiler.FindAllRulesSourceFiles( RulesCompiler.RulesFileType.Target, AllGameProjects.Select(x => x.Folder).ToList(), ForeignPlugins:null, AdditionalSearchPaths:null );
+			foreach( FileReference CurTargetFile in FoundTargetFiles )
 			{
-				var CleanTargetFileName = Utils.CleanDirectorySeparators( CurTargetFile );
+				var CleanTargetFileName = Utils.CleanDirectorySeparators( CurTargetFile.FullName );
 
 				// remove the local root
 				var LocalRoot = Path.GetFullPath(RootRelativePath);
@@ -1016,7 +1009,7 @@ namespace UnrealBuildTool
 
 				if (UnrealBuildTool.HasUProjectFile())
 				{
-					string ProjectRoot = UnrealBuildTool.GetUProjectPath();
+					string ProjectRoot = UnrealBuildTool.GetUProjectPath().FullName;
 					if (Search.StartsWith(ProjectRoot, StringComparison.InvariantCultureIgnoreCase))
 					{
 						if (ProjectRoot.EndsWith("\\") || ProjectRoot.EndsWith("/"))
@@ -1044,12 +1037,12 @@ namespace UnrealBuildTool
 				// Skip NoRedist files if necessary
 				if( bExcludeNoRedistFiles )
 				{
-					IncludeThisTarget = !IsNoRedistModule(CleanTargetFileName);
+					IncludeThisTarget = !IsNoRedistModule(CurTargetFile);
 				}
 
 				if( IncludeThisTarget )
 				{
-					AllTargetFiles.Add( CleanTargetFileName );
+					AllTargetFiles.Add( CurTargetFile );
 				}
 			}
 
@@ -1098,7 +1091,7 @@ namespace UnrealBuildTool
 						SubFolder.SubFolders.Count == 0 )
 					{
 						// 3)
-						if (SubFolder.FolderName.Equals(Utils.GetFilenameWithoutAnyExtensions(SubFolder.ChildProjects[0].ProjectFilePath), StringComparison.InvariantCultureIgnoreCase))
+						if (SubFolder.FolderName.Equals(SubFolder.ChildProjects[0].ProjectFilePath.GetFileNameWithoutAnyExtensions(), StringComparison.InvariantCultureIgnoreCase))
 						{
 							CanCollapseFolder = true;
 						}
@@ -1165,9 +1158,9 @@ namespace UnrealBuildTool
 				{
 					if (CurChildProject != OtherChildProject)
 					{
-						if (Utils.GetFilenameWithoutAnyExtensions(CurChildProject.ProjectFilePath).Equals(Utils.GetFilenameWithoutAnyExtensions(OtherChildProject.ProjectFilePath), StringComparison.InvariantCultureIgnoreCase))
+						if (CurChildProject.ProjectFilePath.GetFileNameWithoutAnyExtensions().Equals(OtherChildProject.ProjectFilePath.GetFileNameWithoutAnyExtensions(), StringComparison.InvariantCultureIgnoreCase))
 						{
-							throw new BuildException("Detected collision between two project files with the same path in the same master project folder, " + OtherChildProject.ProjectFilePath + " and " + CurChildProject.ProjectFilePath + " (master project folder: " + MasterProjectFolderPath + ")");
+							throw new BuildException("Detected collision between two project files with the same path in the same master project folder, " + OtherChildProject.ProjectFilePath.FullName + " and " + CurChildProject.ProjectFilePath.FullName + " (master project folder: " + MasterProjectFolderPath + ")");
 						}
 					}
 				}
@@ -1180,7 +1173,7 @@ namespace UnrealBuildTool
 				// with project file names or file items, as that's not supported in Visual Studio.)
 				foreach (var CurChildProject in Folder.ChildProjects)
 				{
-					if (Utils.GetFilenameWithoutAnyExtensions(CurChildProject.ProjectFilePath).Equals(SubFolder.FolderName, StringComparison.InvariantCultureIgnoreCase))
+					if (CurChildProject.ProjectFilePath.GetFileNameWithoutAnyExtensions().Equals(SubFolder.FolderName, StringComparison.InvariantCultureIgnoreCase))
 					{
 						throw new BuildException("Detected collision between a master project sub-folder " + SubFolder.FolderName + " and a project within the outer folder " + CurChildProject.ProjectFilePath + " (master project folder: " + MasterProjectFolderPath + ")");
 					}
@@ -1210,7 +1203,7 @@ namespace UnrealBuildTool
 		/// </summary>
 		private void AddUnrealBuildToolProject( MasterProjectFolder ProgramsFolder, IEnumerable<ProjectFile> Dependencies )
 		{
-			var ProjectFileName = Utils.MakePathRelativeTo( Path.Combine( Path.Combine( EngineRelativePath, "Source" ), "Programs", "UnrealBuildTool", "UnrealBuildTool.csproj" ), MasterProjectRelativePath );
+			var ProjectFileName = FileReference.Combine(UnrealBuildTool.EngineSourceDirectory, "Programs", "UnrealBuildTool", "UnrealBuildTool.csproj" );
 			var UnrealBuildToolProject = new VCSharpProjectFile( ProjectFileName );
 			UnrealBuildToolProject.ShouldBuildForAllSolutionTargets = true;
 
@@ -1238,19 +1231,18 @@ namespace UnrealBuildTool
 		{
 			VCSharpProjectFile Project = null;
 
-			var ProjectFileName = Path.Combine( EngineRelativePath, "Source", "Programs", ProjectName, Path.GetFileName( ProjectName ) + ".csproj" );
-			FileInfo Info = new FileInfo( ProjectFileName );
+			var ProjectFileName = FileReference.Combine( UnrealBuildTool.EngineSourceDirectory, "Programs", ProjectName, Path.GetFileName( ProjectName ) + ".csproj" );
+			FileInfo Info = new FileInfo( ProjectFileName.FullName );
 			if( Info.Exists )
 			{
-				var FileNameRelativeToMasterProject = Utils.MakePathRelativeTo( ProjectFileName, MasterProjectRelativePath );
-				Project = new VCSharpProjectFile(FileNameRelativeToMasterProject);
+				Project = new VCSharpProjectFile(ProjectFileName);
 				Project.ShouldBuildForAllSolutionTargets = bShouldBuildForAllSolutionTargets;
 				Project.ShouldBuildByDefaultForSolutionTargets = bShouldBuildByDefaultForSolutionTargets;
 				AddExistingProjectFile(Project, bForceDevelopmentConfiguration: bForceDevelopmentConfiguration);
 			}
 			else
 			{
-				throw new BuildException( ProjectFileName + " doesn't exist!" );
+				throw new BuildException( ProjectFileName.FullName + " doesn't exist!" );
 			}
 
 			return Project;
@@ -1265,17 +1257,15 @@ namespace UnrealBuildTool
 			// We only need this for non-native projects
 			ProjectFile Project = null;
 
-			string ProjectFileName = Path.Combine( EngineRelativePath, "Source", "Programs", ProjectName, Path.GetFileName( ProjectName ) + ".shfbproj" );
-			FileInfo Info = new FileInfo( ProjectFileName );
-			if( Info.Exists )
+			FileReference ProjectFileName = FileReference.Combine( UnrealBuildTool.EngineSourceDirectory, "Programs", ProjectName, Path.GetFileName( ProjectName ) + ".shfbproj" );
+			if( ProjectFileName.Exists() )
 			{
-				string FileNameRelativeToMasterProject = Utils.MakePathRelativeTo( ProjectFileName, MasterProjectRelativePath );
-				Project = new VSHFBProjectFile( FileNameRelativeToMasterProject );
+				Project = new VSHFBProjectFile( ProjectFileName );
 				AddExistingProjectFile(Project);
 			}
 			else
 			{
-				throw new BuildException( ProjectFileName + " doesn't exist!" );
+				throw new BuildException( ProjectFileName.FullName + " doesn't exist!" );
 			}
 			
 			return Project;
@@ -1385,15 +1375,15 @@ namespace UnrealBuildTool
 					var ProgramFolder = ProjectFolderAndFile.Key;
 					var ProgramProjectFile = ProjectFolderAndFile.Value;
 
-					var ProgramName = Path.GetFileName( ProgramFolder );
+					var ProgramName = ProgramFolder;
 
 					// @todo projectfiles: The config folder for programs is kind of weird -- you end up going UP a few directories to get to it.  This stuff is not great.
 					// @todo projectfiles: Fragile assumption here about Programs always being under /Engine/Programs
-					var ProgramDirectory = Path.Combine( EngineRelativePath, "Programs", ProgramName );
-					var ProgramConfigDirectory = Path.Combine( ProgramDirectory, "Config" );
-					if( Directory.Exists( ProgramConfigDirectory ) )
+					var ProgramDirectory = DirectoryReference.Combine( UnrealBuildTool.EngineDirectory, "Programs", ProgramName );
+					var ProgramConfigDirectory = DirectoryReference.Combine( ProgramDirectory, "Config" );
+					if( ProgramConfigDirectory.Exists( ) )
 					{
-						var DirectoriesToSearch = new List<string>();
+						var DirectoriesToSearch = new List<DirectoryReference>();
 						DirectoriesToSearch.Add( ProgramConfigDirectory );
 						ProgramProjectFile.AddFilesToProject( SourceFileSearch.FindFiles( DirectoriesToSearch, ExcludeNoRedistFiles:bExcludeNoRedistFiles), ProgramDirectory );
 					}
@@ -1408,7 +1398,7 @@ namespace UnrealBuildTool
 		/// <param name="Arguments">Incoming command-line arguments to UBT</param>
 		/// <param name="TargetFiles">Target files</param>
 		/// <return>Whether the process was successful or not</return>
-		private bool GenerateIntelliSenseData( String[] Arguments, List<Tuple<ProjectFile, string>> TargetFiles )
+		private bool GenerateIntelliSenseData( String[] Arguments, List<Tuple<ProjectFile, FileReference>> TargetFiles )
 		{
 			var bSuccess = true;
 			if( ShouldGenerateIntelliSenseData() && TargetFiles.Count > 0 )
@@ -1421,7 +1411,7 @@ namespace UnrealBuildTool
 						var TargetProjectFile = TargetFiles[ TargetIndex ].Item1;
 						var CurTarget = TargetFiles[ TargetIndex ].Item2;
 
-						var TargetName = Utils.GetFilenameWithoutAnyExtensions( CurTarget );
+						var TargetName = CurTarget.GetFileNameWithoutAnyExtensions();
 
 						Log.TraceVerbose( "Found target: " + TargetName );
 
@@ -1531,12 +1521,11 @@ namespace UnrealBuildTool
 		/// </summary>
 		/// <param name="AllGameFolders">All game folders</param>
 		/// <param name="File">Full path of the file to search for</param>
-		protected UProjectInfo FindGameContainingFile(List<UProjectInfo> AllGames, string File)
+		protected UProjectInfo FindGameContainingFile(List<UProjectInfo> AllGames, FileReference File)
 		{
 			foreach (var Game in AllGames)
 			{
-				string GameFolderPath = Path.GetFullPath(Game.Folder);
-				if (Utils.IsFileUnderDirectory(File, GameFolderPath + Path.DirectorySeparatorChar))
+				if (File.IsUnderDirectory(Game.Folder))
 				{
 					return Game;
 				}
@@ -1551,8 +1540,10 @@ namespace UnrealBuildTool
 		/// <param name="ProgramProjects">All program projects</param>
 		/// <param name="AllModuleFiles">List of *.Build.cs files for all engine programs and games</param>
 		/// <param name="bGatherThirdPartySource">True to gather source code from third party projects too</param>
-		protected void AddProjectsForAllModules( List<UProjectInfo> AllGames, Dictionary<string, ProjectFile> ProgramProjects, List<string> AllModuleFiles, bool bGatherThirdPartySource )
+		protected void AddProjectsForAllModules( List<UProjectInfo> AllGames, Dictionary<string, ProjectFile> ProgramProjects, List<FileReference> AllModuleFiles, bool bGatherThirdPartySource )
 		{
+			DirectoryReference EngineSourceThirdPartyDirectory = DirectoryReference.Combine(UnrealBuildTool.EngineSourceDirectory, "ThirdParty");
+
 			HashSet<ProjectFile> ProjectsWithPlugins = new HashSet<ProjectFile>();
 			foreach( var CurModuleFile in AllModuleFiles )
 			{
@@ -1560,28 +1551,13 @@ namespace UnrealBuildTool
 
 				// The module's "base directory" is simply the directory where its xxx.Build.cs file is stored.  We'll always
 				// harvest source files for this module in this base directory directory and all of its sub-directories.
-				var ModuleName = Utils.GetFilenameWithoutAnyExtensions(CurModuleFile);		// Remove both ".cs" and ".Build"
+				var ModuleName = CurModuleFile.GetFileNameWithoutAnyExtensions();		// Remove both ".cs" and ".Build"
 
 				bool WantProjectFileForModule = true;
 
 				// We'll keep track of whether this is an "engine" or "external" module.  This is determined below while loading module rules.
-				bool IsEngineModule = false;
-				bool IsThirdPartyModule = false;
-
-				// Check to see if this is an Engine module.  That is, the module is located under the "Engine" folder
-				string ModuleFileRelativeToEngineDirectory = Utils.MakePathRelativeTo(CurModuleFile, Path.Combine(EngineRelativePath));
-				if (!ModuleFileRelativeToEngineDirectory.StartsWith( ".." ) && !Path.IsPathRooted( ModuleFileRelativeToEngineDirectory ))
-				{
-					// This is an engine module
-					IsEngineModule = true;
-				}
-
-				string ModuleFileRelativeToThirdPartyDirectory = Utils.MakePathRelativeTo(CurModuleFile, Path.Combine(EngineRelativePath, "Source", "ThirdParty" ));
-				if (!ModuleFileRelativeToThirdPartyDirectory.StartsWith( ".." ) && !Path.IsPathRooted( ModuleFileRelativeToThirdPartyDirectory ))
-				{
-					// This is a third partymodule
-					IsThirdPartyModule = true;
-				}
+				bool IsEngineModule = CurModuleFile.IsUnderDirectory(UnrealBuildTool.EngineDirectory);
+				bool IsThirdPartyModule = CurModuleFile.IsUnderDirectory(EngineSourceThirdPartyDirectory);
 
 				bool IncludePrivateSourceCode = true;
 				if( bGeneratingGameProjectFiles || bGeneratingRocketProjectFiles )
@@ -1606,20 +1582,20 @@ namespace UnrealBuildTool
 				if( WantProjectFileForModule )
 				{
 					string ProjectFileNameBase = null;
-					string BaseFolder = null;
+					DirectoryReference BaseFolder = null;
 
-					string PossibleProgramTargetName = Utils.GetFilenameWithoutAnyExtensions( CurModuleFile );
+					string PossibleProgramTargetName = CurModuleFile.GetFileNameWithoutAnyExtensions();
 
 					// @todo projectfiles: This works fine for now, but is pretty busted.  It assumes only one module per program and that it matches the program target file name. (see TTP 307091)
 					if( ProgramProjects != null && ProgramProjects.ContainsKey( PossibleProgramTargetName ) )	// @todo projectfiles: When building (in mem projects), ProgramProjects will be null so we are using the UE4 project instead
 					{
 						ProjectFileNameBase = PossibleProgramTargetName;
-						BaseFolder = Path.GetDirectoryName( CurModuleFile );
+						BaseFolder = CurModuleFile.Directory;
 					}
 					else if( IsEngineModule )
 					{
 						ProjectFileNameBase = EngineProjectFileNameBase;
-						BaseFolder = EngineRelativePath;
+						BaseFolder = UnrealBuildTool.EngineDirectory;
 					}
 					else
 					{
@@ -1634,7 +1610,7 @@ namespace UnrealBuildTool
 					}
 
 					// Setup a project file entry for this module's project.  Remember, some projects may host multiple modules!
-					string ProjectFileName = Utils.MakePathRelativeTo( Path.Combine( IntermediateProjectFilesPath, ProjectFileNameBase + ProjectFileExtension ), MasterProjectRelativePath );
+					FileReference ProjectFileName = FileReference.Combine( IntermediateProjectFilesPath, ProjectFileNameBase + ProjectFileExtension );
 					bool bProjectAlreadyExisted;
 					var ProjectFile = FindOrAddProject( ProjectFileName, IncludeInGeneratedProjects:true, bAlreadyExisted:out bProjectAlreadyExisted );
 
@@ -1659,20 +1635,20 @@ namespace UnrealBuildTool
 					// Check if there's a plugin directory here
 					if(!ProjectsWithPlugins.Contains(ProjectFile))
 					{
-						string PluginFolder = Path.Combine(BaseFolder, "Plugins");
-						if(Directory.Exists(PluginFolder))
+						DirectoryReference PluginFolder = DirectoryReference.Combine(BaseFolder, "Plugins");
+						if(PluginFolder.Exists())
 						{
 							// Add all the plugin files for this project
-							foreach(string PluginFileName in Plugins.EnumeratePlugins(Path.Combine(BaseFolder, "Plugins")))
+							foreach(FileReference PluginFileName in Plugins.EnumeratePlugins(PluginFolder))
 							{
 								// Add the .uplugin file
 								ProjectFile.AddFileToProject(PluginFileName, BaseFolder);
 
 								// Add plugin "resource" files if we have any
-								string PluginResourcesFolder = Path.Combine(Path.GetDirectoryName(PluginFileName), "Resources");
-								if(Directory.Exists(PluginResourcesFolder))
+								DirectoryReference PluginResourcesFolder = DirectoryReference.Combine(PluginFileName.Directory, "Resources");
+								if(PluginResourcesFolder.Exists())
 								{
-									ProjectFile.AddFilesToProject(SourceFileSearch.FindFiles(DirectoriesToSearch: new List<string>{ PluginResourcesFolder }, ExcludeNoRedistFiles: bExcludeNoRedistFiles ), BaseFolder );
+									ProjectFile.AddFilesToProject(SourceFileSearch.FindFiles(DirectoriesToSearch: new List<DirectoryReference>{ PluginResourcesFolder }, ExcludeNoRedistFiles: bExcludeNoRedistFiles ), BaseFolder );
 								}
 							}
 						}
@@ -1691,47 +1667,49 @@ namespace UnrealBuildTool
 		/// <param name="GameProjects">Map of game folder name to all of the game projects we created</param>
 		/// <param name="ProgramProjects">Map of program names to all of the program projects we created</param>
 		/// <param name="TemplateGameProjects">Set of template game projects we found.  These will also be in the GameProjects map</param>
-		private void AddProjectsForAllTargets( List<UProjectInfo> AllGames, out ProjectFile EngineProject, out Dictionary<string, ProjectFile> GameProjects, out Dictionary<string, ProjectFile> ProgramProjects, out HashSet<ProjectFile> TemplateGameProjects )
+		private void AddProjectsForAllTargets( List<UProjectInfo> AllGames, out ProjectFile EngineProject, out Dictionary<DirectoryReference, ProjectFile> GameProjects, out Dictionary<string, ProjectFile> ProgramProjects, out HashSet<ProjectFile> TemplateGameProjects )
 		{
 			// As we're creating project files, we'll also keep track of whether we created an "engine" project and return that if we have one
 			EngineProject = null;
-			GameProjects = new Dictionary<string,ProjectFile>( StringComparer.InvariantCultureIgnoreCase );
+			GameProjects = new Dictionary<DirectoryReference,ProjectFile>();
 			ProgramProjects = new Dictionary<string,ProjectFile>( StringComparer.InvariantCultureIgnoreCase );
 			TemplateGameProjects = new HashSet<ProjectFile>();
 
+			// Get some standard directories
+			DirectoryReference EngineSourceProgramsDirectory = DirectoryReference.Combine(UnrealBuildTool.EngineSourceDirectory, "Programs");
+			DirectoryReference TemplatesDirectory = DirectoryReference.Combine(UnrealBuildTool.EngineDirectory, "..", "Templates");
 
 			// Find all of the target files.  This will filter out any modules or targets that don't
 			// belong to platforms we're generating project files for.
 			var AllTargetFiles = DiscoverTargets(AllGames);
 			foreach( var TargetFilePath in AllTargetFiles )
 			{
-				var TargetName = Utils.GetFilenameWithoutAnyExtensions(TargetFilePath);		// Remove both ".cs" and ".Target"
+				var TargetName = TargetFilePath.GetFileNameWithoutAnyExtensions();		// Remove both ".cs" and ".Target"
 
 				// Check to see if this is an Engine target.  That is, the target is located under the "Engine" folder
 				bool IsEngineTarget = false;
-				string TargetFileRelativeToEngineDirectory = Utils.MakePathRelativeTo(TargetFilePath, Path.Combine(EngineRelativePath), AlwaysTreatSourceAsDirectory: false);
-				if (!TargetFileRelativeToEngineDirectory.StartsWith( ".." ) && !Path.IsPathRooted( TargetFileRelativeToEngineDirectory ))
+				bool WantProjectFileForTarget = true;
+				if(TargetFilePath.IsUnderDirectory(UnrealBuildTool.EngineDirectory))
 				{
 					// This is an engine target
 					IsEngineTarget = true;
-				}
 
-				bool WantProjectFileForTarget = true;
-				if(TargetFileRelativeToEngineDirectory.StartsWith(Path.Combine("Source", "Programs"), StringComparison.InvariantCultureIgnoreCase))
-				{
-					WantProjectFileForTarget = IncludeEnginePrograms;
-				}
-				else if(TargetFileRelativeToEngineDirectory.StartsWith(Path.Combine("Source"), StringComparison.InvariantCultureIgnoreCase))
-				{
-					WantProjectFileForTarget = IncludeEngineSource;
+					if(TargetFilePath.IsUnderDirectory(EngineSourceProgramsDirectory))
+					{
+						WantProjectFileForTarget = IncludeEnginePrograms;
+					}
+					else if(TargetFilePath.IsUnderDirectory(UnrealBuildTool.EngineSourceDirectory))
+					{
+						WantProjectFileForTarget = IncludeEngineSource;
+					}
 				}
 
 				if (WantProjectFileForTarget)
 				{
-					string CheckProjectFile = UProjectInfo.GetProjectForTarget(TargetName);
-
 					RulesAssembly RulesAssembly;
-					if(String.IsNullOrEmpty(CheckProjectFile))
+
+					FileReference CheckProjectFile;
+					if(!UProjectInfo.TryGetProjectForTarget(TargetName, out CheckProjectFile))
 					{
 						RulesAssembly = RulesCompiler.CreateEngineRulesAssembly();
 					}
@@ -1742,15 +1720,14 @@ namespace UnrealBuildTool
 
 					// Create target rules for all of the platforms and configuration combinations that we want to enable support for.
 					// Just use the current platform as we only need to recover the target type and both should be supported for all targets...
-					string UnusedTargetFilePath;
-					var TargetRulesObject = RulesAssembly.CreateTargetRules(TargetName, new TargetInfo(BuildHostPlatform.Current.Platform, UnrealTargetConfiguration.Development), false, out UnusedTargetFilePath);
+					var TargetRulesObject = RulesAssembly.CreateTargetRules(TargetName, new TargetInfo(BuildHostPlatform.Current.Platform, UnrealTargetConfiguration.Development), false);
 
 					// Exclude client and server targets under binary Rocket; it's impossible to build without precompiled engine binaries
 					if (!UnrealBuildTool.RunningRocket() || (TargetRulesObject.Type != TargetRules.TargetType.Client && TargetRulesObject.Type != TargetRules.TargetType.Server))
 					{
 						bool IsProgramTarget = false;
 
-						string GameFolder = null;
+						DirectoryReference GameFolder = null;
 						string ProjectFileNameBase = null;
 						if (TargetRulesObject.Type == TargetRules.TargetType.Program)
 						{
@@ -1775,7 +1752,7 @@ namespace UnrealBuildTool
 
 						// @todo projectfiles: We should move all of the Target.cs files out of sub-folders to clean up the project directories a bit (e.g. GameUncooked folder)
 
-						var ProjectFilePath = Path.Combine(IntermediateProjectFilesPath, ProjectFileNameBase + ProjectFileExtension);
+						FileReference ProjectFilePath = FileReference.Combine(IntermediateProjectFilesPath, ProjectFileNameBase + ProjectFileExtension);
 
 						if (TargetRules.IsGameType(TargetRulesObject.Type) &&
 							(TargetRules.IsEditorType(TargetRulesObject.Type) == false))
@@ -1784,39 +1761,31 @@ namespace UnrealBuildTool
 							UEPlatformProjectGenerator.GenerateGameProjectStubs(
 								InGenerator: this,
 								InTargetName: TargetName,
-								InTargetFilepath: TargetFilePath,
+								InTargetFilepath: TargetFilePath.FullName,
 								InTargetRules: TargetRulesObject,
 								InPlatforms: SupportedPlatforms,
 								InConfigurations: SupportedConfigurations);
 						}
 
-						ProjectFilePath = Utils.MakePathRelativeTo(ProjectFilePath, MasterProjectRelativePath);
-
 						bool bProjectAlreadyExisted;
 						var ProjectFile = FindOrAddProject(ProjectFilePath, IncludeInGeneratedProjects: true, bAlreadyExisted: out bProjectAlreadyExisted);
-						ProjectFile.IsForeignProject = bGeneratingGameProjectFiles && UnrealBuildTool.HasUProjectFile() && Utils.IsFileUnderDirectory(TargetFilePath, UnrealBuildTool.GetUProjectPath());
+						ProjectFile.IsForeignProject = bGeneratingGameProjectFiles && UnrealBuildTool.HasUProjectFile() && TargetFilePath.IsUnderDirectory(UnrealBuildTool.GetUProjectPath());
 						ProjectFile.IsGeneratedProject = true;
 						ProjectFile.IsStubProject = false;
 
-						bool IsTemplateTarget = false;
-						{
-							// Check to see if this is a template target.  That is, the target is located under the "Templates" folder
-							string TargetFileRelativeToTemplatesDirectory = Utils.MakePathRelativeTo(TargetFilePath, Path.Combine(RootRelativePath, "Templates"));
-							if (!TargetFileRelativeToTemplatesDirectory.StartsWith("..") && !Path.IsPathRooted(TargetFileRelativeToTemplatesDirectory))
-							{
-								IsTemplateTarget = true;
-							}
-						}
-						string BaseFolder = null;
+						// Check to see if this is a template target.  That is, the target is located under the "Templates" folder
+						bool IsTemplateTarget = TargetFilePath.IsUnderDirectory(TemplatesDirectory);
+
+						DirectoryReference BaseFolder = null;
 						if (IsProgramTarget)
 						{
 							ProgramProjects[TargetName] = ProjectFile;
-							BaseFolder = Path.GetDirectoryName(TargetFilePath);
+							BaseFolder = TargetFilePath.Directory;
 						}
 						else if (IsEngineTarget)
 						{
 							EngineProject = ProjectFile;
-							BaseFolder = EngineRelativePath;
+							BaseFolder = UnrealBuildTool.EngineDirectory;
 							if (bGeneratingRocketProjectFiles)
 							{
 								// Allow engine projects to be created but not built for Rocket
@@ -1837,8 +1806,8 @@ namespace UnrealBuildTool
 							if (!bProjectAlreadyExisted)
 							{
 								// Add the .uproject file for this game/template
-								var UProjectFilePath = Path.Combine(BaseFolder, ProjectFileNameBase + ".uproject");
-								if (File.Exists(UProjectFilePath))
+								var UProjectFilePath = FileReference.Combine(BaseFolder, ProjectFileNameBase + ".uproject");
+								if (UProjectFilePath.Exists())
 								{
 									ProjectFile.AddFileToProject(UProjectFilePath, BaseFolder);
 								}
@@ -1894,9 +1863,9 @@ namespace UnrealBuildTool
 		protected void AddEngineShaderSource( ProjectFile EngineProject )
 		{
 			// Setup a project file entry for this module's project.  Remember, some projects may host multiple modules!
-			var ShadersDirectory = Path.Combine( EngineRelativePath, "Shaders" );
+			var ShadersDirectory = DirectoryReference.Combine( UnrealBuildTool.EngineDirectory, "Shaders" );
 
-			var DirectoriesToSearch = new List<string>();
+			var DirectoriesToSearch = new List<DirectoryReference>();
 			DirectoriesToSearch.Add( ShadersDirectory );
 			var SubdirectoryNamesToExclude = new List<string>();
 			{
@@ -1909,16 +1878,16 @@ namespace UnrealBuildTool
 			EngineProject.AddFilesToProject( SourceFileSearch.FindFiles(
 				DirectoriesToSearch: DirectoriesToSearch,
 				ExcludeNoRedistFiles: bExcludeNoRedistFiles,
-				SubdirectoryNamesToExclude: SubdirectoryNamesToExclude ), EngineRelativePath );
+				SubdirectoryNamesToExclude: SubdirectoryNamesToExclude ), UnrealBuildTool.EngineDirectory );
 		}
 
 
 		/// Adds engine build infrastructure files to the specified project
 		protected void AddEngineBuildFiles( ProjectFile EngineProject )
 		{
-			var BuildDirectory = Path.Combine( EngineRelativePath, "Build" );
+			var BuildDirectory = DirectoryReference.Combine( UnrealBuildTool.EngineDirectory, "Build" );
 
-			var DirectoriesToSearch = new List<string>();
+			var DirectoriesToSearch = new List<DirectoryReference>();
 			DirectoriesToSearch.Add( BuildDirectory );
 			var SubdirectoryNamesToExclude = new List<string>();
 			SubdirectoryNamesToExclude.Add("Receipts");
@@ -1926,7 +1895,7 @@ namespace UnrealBuildTool
 			EngineProject.AddFilesToProject( SourceFileSearch.FindFiles(
 				DirectoriesToSearch: DirectoriesToSearch,
 				ExcludeNoRedistFiles: bExcludeNoRedistFiles,
-				SubdirectoryNamesToExclude: SubdirectoryNamesToExclude ), EngineRelativePath );
+				SubdirectoryNamesToExclude: SubdirectoryNamesToExclude ), UnrealBuildTool.EngineDirectory );
 		}
 
 
@@ -1934,14 +1903,14 @@ namespace UnrealBuildTool
 		protected void AddEngineDocumentation( ProjectFile EngineProject )
 		{
 			// NOTE: The project folder added here will actually be collapsed away later if not needed
-			var DocumentationProjectDirectory = Path.Combine( EngineRelativePath, "Documentation" );
-			var DocumentationSourceDirectory = Path.Combine( EngineRelativePath, "Documentation", "Source" );
-			DirectoryInfo DirInfo = new DirectoryInfo( DocumentationProjectDirectory );
-			if( DirInfo.Exists && Directory.Exists( DocumentationSourceDirectory ) )
+			var DocumentationProjectDirectory = DirectoryReference.Combine( UnrealBuildTool.EngineDirectory, "Documentation" );
+			var DocumentationSourceDirectory = DirectoryReference.Combine( UnrealBuildTool.EngineDirectory, "Documentation", "Source" );
+			DirectoryInfo DirInfo = new DirectoryInfo( DocumentationProjectDirectory.FullName );
+			if( DirInfo.Exists && DocumentationSourceDirectory.Exists() )
 			{
 				Log.TraceVerbose( "Adding documentation files..." );
 
-				var DirectoriesToSearch = new List<string>();
+				var DirectoriesToSearch = new List<DirectoryReference>();
 				DirectoriesToSearch.Add( DocumentationSourceDirectory );
 
 				var SubdirectoryNamesToExclude = new List<string>();
@@ -1965,13 +1934,13 @@ namespace UnrealBuildTool
 				// Filter out non-English documentation files if we were configured to do so
 				if( !bAllDocumentationLanguages )
 				{
-					var FilteredDocumentationFiles = new List<string>();
+					var FilteredDocumentationFiles = new List<FileReference>();
 					foreach( var DocumentationFile in DocumentationFiles )
 					{
 						bool bPassesFilter = true;
-						if( DocumentationFile.EndsWith( ".udn", StringComparison.InvariantCultureIgnoreCase ) )
+						if( DocumentationFile.FullName.EndsWith( ".udn", StringComparison.InvariantCultureIgnoreCase ) )
 						{
-							var LanguageSuffix = Path.GetExtension( Path.GetFileNameWithoutExtension( DocumentationFile ) );
+							var LanguageSuffix = Path.GetExtension( Path.GetFileNameWithoutExtension( DocumentationFile.FullName ) );
 							if( !String.IsNullOrEmpty( LanguageSuffix ) &&
 								!LanguageSuffix.Equals( ".int", StringComparison.InvariantCultureIgnoreCase ) )
 							{
@@ -1987,7 +1956,7 @@ namespace UnrealBuildTool
 					DocumentationFiles = FilteredDocumentationFiles;
 				}
 
-				EngineProject.AddFilesToProject( DocumentationFiles, EngineRelativePath );
+				EngineProject.AddFilesToProject( DocumentationFiles, UnrealBuildTool.EngineDirectory );
 			}
 			else
 			{
@@ -2003,9 +1972,9 @@ namespace UnrealBuildTool
 		/// <param name="IncludeInGeneratedProjects">True if this project should be included in the set of generated projects.  Only matters when actually generating project files.</param>
 		/// <param name="bAlreadyExisted">True if we already had this project file</param>
 		/// <returns>Object that represents this project file in Unreal Build Tool</returns>
-		public ProjectFile FindOrAddProject( string FilePath, bool IncludeInGeneratedProjects, out bool bAlreadyExisted )
+		public ProjectFile FindOrAddProject( FileReference FilePath, bool IncludeInGeneratedProjects, out bool bAlreadyExisted )
 		{
-			if( String.IsNullOrEmpty( FilePath ) )
+			if( FilePath == null )
 			{
 				throw new BuildException( "Not valid to call FindOrAddProject() with an empty file path!" );
 			}
@@ -2037,7 +2006,7 @@ namespace UnrealBuildTool
 		/// </summary>
 		/// <param name="InitFilePath">Path to the project file</param>
 		/// <returns>The newly allocated project file object</returns>
-		protected abstract ProjectFile AllocateProjectFile( string InitFilePath );
+		protected abstract ProjectFile AllocateProjectFile( FileReference InitFilePath );
 
 
 		/// <summary>
@@ -2098,7 +2067,7 @@ namespace UnrealBuildTool
             {
                 var NS = XNamespace.Get("http://schemas.microsoft.com/developer/msbuild/2003");
 
-                var AutomationToolDir = Path.Combine(EngineRelativePath, "Source", "Programs", "AutomationTool");
+                var AutomationToolDir = DirectoryReference.Combine(UnrealBuildTool.EngineSourceDirectory, "Programs", "AutomationTool");
                 new XDocument(
                     new XElement(NS + "Project",
                         new XAttribute("ToolsVersion", "4.0"),
@@ -2106,14 +2075,14 @@ namespace UnrealBuildTool
                         new XElement(NS + "ItemGroup",
                             from AutomationProject in AutomationProjectFiles
                             select new XElement(NS + "ProjectReference",
-                                new XAttribute("Include", Utils.MakePathRelativeTo(AutomationProject.ProjectFilePath, AutomationToolDir, true)),
+                                new XAttribute("Include", AutomationProject.ProjectFilePath.MakeRelativeTo(AutomationToolDir)),
                                 new XElement(NS + "Project", (AutomationProject as VCSharpProjectFile).ProjectGUID.ToString("B")),
-                                new XElement(NS + "Name", Path.GetFileNameWithoutExtension(AutomationProject.RelativeProjectFilePath)),
+                                new XElement(NS + "Name", AutomationProject.ProjectFilePath.GetFileNameWithoutExtension()),
                                 new XElement(NS + "Private", "false")
                             )
                         )
                     )
-                ).Save(Path.Combine( AutomationToolDir, "AutomationTool.csproj.References" ));
+                ).Save(FileReference.Combine( AutomationToolDir, "AutomationTool.csproj.References" ).FullName);
             }
 
 			return true;
@@ -2130,11 +2099,11 @@ namespace UnrealBuildTool
 			{
 				foreach( var SourceFile in CurProject.SourceFiles )
 				{
-					Manifest.AddBuildProduct( SourceFile.FilePath );
+					Manifest.AddBuildProduct( SourceFile.Reference.FullName );
 				}
 			}
 
-			string ManifestName = Path.Combine( ProjectFileGenerator.IntermediateProjectFilesPath, "UE4SourceFiles.xml" );
+			string ManifestName = Path.Combine( ProjectFileGenerator.IntermediateProjectFilesPath.FullName, "UE4SourceFiles.xml" );
 			Utils.WriteClass<BuildManifest>( Manifest, ManifestName, "" );
 			return true;
 		}
@@ -2326,7 +2295,7 @@ namespace UnrealBuildTool
 		/// Map of project file names to their project files.  This includes every single project file in memory or otherwise that
 		/// we know about so far.  Note that when generating project files, this map may even include project files that we won't
 		/// be including in the generated projects.
-		protected readonly Dictionary<string, ProjectFile> ProjectFileMap = new Dictionary<string, ProjectFile>( StringComparer.InvariantCultureIgnoreCase );
+		protected readonly Dictionary<FileReference, ProjectFile> ProjectFileMap = new Dictionary<FileReference, ProjectFile>();
 
 		/// List of project files that we'll be generating
 		protected readonly List<ProjectFile> GeneratedProjectFiles = new List<ProjectFile>();
diff --git a/Engine/Source/Programs/UnrealBuildTool/System/QMakefileGenerator.cs b/Engine/Source/Programs/UnrealBuildTool/System/QMakefileGenerator.cs
--- a/Engine/Source/Programs/UnrealBuildTool/System/QMakefileGenerator.cs
+++ b/Engine/Source/Programs/UnrealBuildTool/System/QMakefileGenerator.cs
@@ -23,7 +23,7 @@ namespace UnrealBuildTool
 
 	public class QMakefileProjectFile : ProjectFile
 	{
-		public QMakefileProjectFile( string InitFilePath )
+		public QMakefileProjectFile( FileReference InitFilePath )
 			: base(InitFilePath)
 		{
 		}
@@ -74,7 +74,7 @@ namespace UnrealBuildTool
 		/// Adds the include directory to the list, after converting it to relative to UE4 root
 		private void AddIncludeDirectory(ref List<string> IncludeDirectories, string IncludeDir, string ProjectDir)
 		{
-			string FullProjectPath = Path.GetFullPath(ProjectFileGenerator.MasterProjectRelativePath);
+			string FullProjectPath = ProjectFileGenerator.MasterProjectPath.FullName;
 			string FullPath = "";
 			if (IncludeDir.StartsWith("/") && !IncludeDir.StartsWith(FullProjectPath))
 			{
@@ -133,11 +133,11 @@ namespace UnrealBuildTool
 				}
 
 				foreach (var CurPath in QMakeProject.IntelliSenseIncludeSearchPaths) {
-					AddIncludeDirectory(ref IncludeDirectories, CurPath, Path.GetDirectoryName(QMakeProject.ProjectFilePath));
+					AddIncludeDirectory(ref IncludeDirectories, CurPath, Path.GetDirectoryName(QMakeProject.ProjectFilePath.FullName));
 					// System.Console.WriteLine ("Not empty now? CurPath == ", CurPath);
 				}
 				foreach (var CurPath in QMakeProject.IntelliSenseSystemIncludeSearchPaths) {
-					AddIncludeDirectory (ref SystemIncludeDirectories, CurPath, Path.GetDirectoryName (QMakeProject.ProjectFilePath));
+					AddIncludeDirectory (ref SystemIncludeDirectories, CurPath, Path.GetDirectoryName (QMakeProject.ProjectFilePath.FullName));
 				}
 
 			}
@@ -206,8 +206,8 @@ namespace UnrealBuildTool
 
 			if (!String.IsNullOrEmpty (GameProjectName))
 			{
-				GameProjectPath = UnrealBuildTool.GetUProjectPath ();
-				GameProjectFile = UnrealBuildTool.GetUProjectFile ();
+				GameProjectPath = UnrealBuildTool.GetUProjectPath ().FullName;
+				GameProjectFile = UnrealBuildTool.GetUProjectFile ().FullName;
 				QMakeGameProjectFile = "gameProjectFile=" + GameProjectFile + "\n";
 				BuildCommand = "build=mono $$unrealRootPath/Engine/Binaries/DotNET/UnrealBuildTool.exe\n\n";
 			}
@@ -263,12 +263,12 @@ namespace UnrealBuildTool
 
 			// Create SourceFiles, HeaderFiles, and ConfigFiles sections.
 			var AllModuleFiles = DiscoverModules(FindGameProjects());
-			foreach (string CurModuleFile in AllModuleFiles)
+			foreach (FileReference CurModuleFile in AllModuleFiles)
 			{
 				var FoundFiles = SourceFileSearch.FindModuleSourceFiles(CurModuleFile, ExcludeNoRedistFiles: bExcludeNoRedistFiles);
-				foreach (string CurSourceFile in FoundFiles)
+				foreach (FileReference CurSourceFile in FoundFiles)
 				{
-					string SourceFileRelativeToRoot = Utils.MakePathRelativeTo(CurSourceFile, Path.Combine(EngineRelativePath));
+					string SourceFileRelativeToRoot = CurSourceFile.MakeRelativeTo(UnrealBuildTool.EngineDirectory);
 					// Exclude some directories that we don't compile (note that we still want Windows/Mac etc for code navigation)
 					if (!SourceFileRelativeToRoot.Contains ("Source/ThirdParty/")) 
 					{
@@ -286,7 +286,7 @@ namespace UnrealBuildTool
 								} 
 								else if (!String.IsNullOrEmpty (GameProjectName)) 
 								{
-									QMakeSourceFilesList += ("\t\"$$" + GameProjectName + "RootPath/" + Utils.MakePathRelativeTo (CurSourceFile, GameProjectPath) + "\" \\\n");
+									QMakeSourceFilesList += ("\t\"$$" + GameProjectName + "RootPath/" + Utils.MakePathRelativeTo (CurSourceFile.FullName, GameProjectPath) + "\" \\\n");
 								} 
 								else 
 								{
@@ -310,7 +310,7 @@ namespace UnrealBuildTool
 								} 
 								else if (!String.IsNullOrEmpty (GameProjectName)) 
 								{
-									QMakeHeaderFilesList += ("\t\"$$" + GameProjectName + "RootPath/" + Utils.MakePathRelativeTo (CurSourceFile, GameProjectPath) + "\" \\\n");
+									QMakeHeaderFilesList += ("\t\"$$" + GameProjectName + "RootPath/" + Utils.MakePathRelativeTo (CurSourceFile.FullName, GameProjectPath) + "\" \\\n");
 								} 
 								else 
 								{
@@ -335,7 +335,7 @@ namespace UnrealBuildTool
 								} 
 								else if (!String.IsNullOrEmpty (GameProjectName)) 
 								{
-									QMakeConfigFilesList += ("\t\"$$" + GameProjectName + "RootPath/" + Utils.MakePathRelativeTo (CurSourceFile, GameProjectPath) + "\" \\\n");
+									QMakeConfigFilesList += ("\t\"$$" + GameProjectName + "RootPath/" + Utils.MakePathRelativeTo (CurSourceFile.FullName, GameProjectPath) + "\" \\\n");
 								} 
 								else 
 								{
@@ -369,7 +369,7 @@ namespace UnrealBuildTool
 						continue;
 					}
 
-					var TargetName = Utils.GetFilenameWithoutAnyExtensions(TargetFile.TargetFilePath);		// Remove both ".cs" and ".
+					var TargetName = TargetFile.TargetFilePath.GetFileNameWithoutAnyExtensions();		// Remove both ".cs" and ".
 
 					foreach (UnrealTargetConfiguration CurConfiguration in Enum.GetValues(typeof(UnrealTargetConfiguration)))
 					{
@@ -401,13 +401,13 @@ namespace UnrealBuildTool
 
 			QMakeFileContent.Append (QMakeTargetList.TrimEnd('\\'));
 
-			var FullFileName = Path.Combine(MasterProjectRelativePath, FileName);
+			var FullFileName = Path.Combine(MasterProjectPath.FullName, FileName);
 
-			var FullQMakeDefinesFileName = Path.Combine (MasterProjectRelativePath, QMakeDefinesFileName);
-			var FullQMakeIncludesFileName = Path.Combine(MasterProjectRelativePath, QMakeIncludesFileName);
-			var FullQMakeSourcePriFileName = Path.Combine(MasterProjectRelativePath, QMakeSourcePriFileName);
-			var FullQMakeHeaderPriFileName = Path.Combine(MasterProjectRelativePath, QMakeHeaderPriFileName);
-			var FullQMakeConfigPriFileName = Path.Combine(MasterProjectRelativePath, QMakeConfigPriFileName);
+			var FullQMakeDefinesFileName = Path.Combine (MasterProjectPath.FullName, QMakeDefinesFileName);
+			var FullQMakeIncludesFileName = Path.Combine(MasterProjectPath.FullName, QMakeIncludesFileName);
+			var FullQMakeSourcePriFileName = Path.Combine(MasterProjectPath.FullName, QMakeSourcePriFileName);
+			var FullQMakeHeaderPriFileName = Path.Combine(MasterProjectPath.FullName, QMakeHeaderPriFileName);
+			var FullQMakeConfigPriFileName = Path.Combine(MasterProjectPath.FullName, QMakeConfigPriFileName);
 
 			WriteFileIfChanged (FullQMakeDefinesFileName, QMakeDefinesPriFileContent.ToString ());
 			WriteFileIfChanged (FullQMakeIncludesFileName, QMakeIncludesPriFileContent.ToString());
@@ -438,13 +438,13 @@ namespace UnrealBuildTool
 		/// </summary>
 		/// <param name="InitFilePath">Path to the project file</param>
 		/// <returns>The newly allocated project file object</returns>
-		protected override ProjectFile AllocateProjectFile( string InitFilePath )
+		protected override ProjectFile AllocateProjectFile( FileReference InitFilePath )
 		{
 			return new QMakefileProjectFile( InitFilePath );
 		}
 
 		/// ProjectFileGenerator interface
-		public override void CleanProjectFiles(string InMasterProjectRelativePath, string InMasterProjectName, string InIntermediateProjectFilesPath)
+		public override void CleanProjectFiles(DirectoryReference InMasterProjectDirectory, string InMasterProjectName, DirectoryReference InIntermediateProjectFilesDirectory)
 		{
 		}
 	}
diff --git a/Engine/Source/Programs/UnrealBuildTool/System/ResponseFile.cs b/Engine/Source/Programs/UnrealBuildTool/System/ResponseFile.cs
--- a/Engine/Source/Programs/UnrealBuildTool/System/ResponseFile.cs
+++ b/Engine/Source/Programs/UnrealBuildTool/System/ResponseFile.cs
@@ -14,9 +14,9 @@ namespace UnrealBuildTool
 		/// </summary>
 		/// <param name="TempFileName">Name of response file</param>
 		/// <param name="Lines">List of lines to write to the response file</param>
-		public static string Create(string TempFileName, List<string> Lines)
+		public static FileReference Create(FileReference TempFileName, List<string> Lines)
 		{
-			FileInfo TempFileInfo = new FileInfo( TempFileName );
+			FileInfo TempFileInfo = new FileInfo( TempFileName.FullName );
 			// Delete the existing file if it exists
 			if( TempFileInfo.Exists )
 			{
diff --git a/Engine/Source/Programs/UnrealBuildTool/System/RulesCompiler.cs b/Engine/Source/Programs/UnrealBuildTool/System/RulesCompiler.cs
--- a/Engine/Source/Programs/UnrealBuildTool/System/RulesCompiler.cs
+++ b/Engine/Source/Programs/UnrealBuildTool/System/RulesCompiler.cs
@@ -394,10 +394,10 @@ namespace UnrealBuildTool
 		 * @param ProjectFile A path to the .uproject file.
 		 * @param ModuleName Name of the module.
 		 */
-		public void ReadAdditionalDependencies(string ProjectFile, string ModuleName)
+		public void ReadAdditionalDependencies(FileReference ProjectFile, string ModuleName)
 		{
 			// Create a case-insensitive dictionary of the contents
-			Dictionary<string, object> Descriptor = fastJSON.JSON.Instance.ToObject<Dictionary<string, object>>(File.ReadAllText(ProjectFile));
+			Dictionary<string, object> Descriptor = fastJSON.JSON.Instance.ToObject<Dictionary<string, object>>(File.ReadAllText(ProjectFile.FullName));
 			Descriptor = new Dictionary<string,object>(Descriptor, StringComparer.InvariantCultureIgnoreCase);
 
 			// Get the list of plugins
@@ -1085,12 +1085,12 @@ namespace UnrealBuildTool
 		/// <summary>
 		/// Maps module names to their actual xxx.Module.cs file on disk
 		/// </summary>
-		private Dictionary<string, string> ModuleNameToFileName = new Dictionary<string, string>(StringComparer.InvariantCultureIgnoreCase);
+		private Dictionary<string, FileReference> ModuleNameToFileName = new Dictionary<string, FileReference>(StringComparer.InvariantCultureIgnoreCase);
 
 		/// <summary>
 		/// Maps target names to their actual xxx.Target.cs file on disk
 		/// </summary>
-		private Dictionary<string, string> TargetNameToFileName = new Dictionary<string, string>(StringComparer.InvariantCultureIgnoreCase);
+		private Dictionary<string, FileReference> TargetNameToFileName = new Dictionary<string, FileReference>(StringComparer.InvariantCultureIgnoreCase);
 
 		/// <summary>
 		/// Constructor. Compiles a rules assembly from the given source files.
@@ -1098,12 +1098,12 @@ namespace UnrealBuildTool
 		/// <param name="ModuleFileNames">List of module files to compile</param>
 		/// <param name="TargetFileNames">List of target files to compile</param>
 		/// <param name="AssemblyFileName">The output path for the compiled assembly</param>
-		public RulesAssembly(List<string> ModuleFileNames, List<string> TargetFileNames, string AssemblyFileName, RulesAssembly InParent)
+		public RulesAssembly(List<FileReference> ModuleFileNames, List<FileReference> TargetFileNames, FileReference AssemblyFileName, RulesAssembly InParent)
 		{
 			Parent = InParent;
 
 			// Find all the source files
-			List<string> AssemblySourceFiles = new List<string>();
+			List<FileReference> AssemblySourceFiles = new List<FileReference>();
 			AssemblySourceFiles.AddRange( ModuleFileNames );
 			AssemblySourceFiles.AddRange( TargetFileNames );
 
@@ -1114,10 +1114,9 @@ namespace UnrealBuildTool
 			}
 
 			// Setup the module map
-			foreach( var CurModuleFileName in ModuleFileNames )
+			foreach( FileReference CurModuleFileName in ModuleFileNames )
 			{
-				var CleanFileName = Utils.CleanDirectorySeparators( CurModuleFileName );
-				var ModuleName = Path.GetFileNameWithoutExtension( Path.GetFileNameWithoutExtension( CleanFileName ) );	// Strip both extensions
+				string ModuleName = Path.GetFileNameWithoutExtension( Path.GetFileNameWithoutExtension( CurModuleFileName.FullName ) );	// Strip both extensions
 				if( !ModuleNameToFileName.ContainsKey( ModuleName ) )
 				{
 					ModuleNameToFileName.Add( ModuleName, CurModuleFileName );
@@ -1125,10 +1124,9 @@ namespace UnrealBuildTool
 			}
 
 			// Setup the target map
-			foreach( var CurTargetFileName in TargetFileNames )
+			foreach( FileReference CurTargetFileName in TargetFileNames )
 			{
-				var CleanFileName = Utils.CleanDirectorySeparators( CurTargetFileName );
-				var TargetName = Path.GetFileNameWithoutExtension( Path.GetFileNameWithoutExtension( CleanFileName ) );	// Strip both extensions
+				string TargetName = Path.GetFileNameWithoutExtension( Path.GetFileNameWithoutExtension( CurTargetFileName.FullName ) );	// Strip both extensions
 				if( !TargetNameToFileName.ContainsKey( TargetName ) )
 				{
 					TargetNameToFileName.Add( TargetName, CurTargetFileName );
@@ -1142,7 +1140,7 @@ namespace UnrealBuildTool
 		/// <param name="ExistingType"></param>
 		/// <param name="FileName"></param>
 		/// <returns>True if the type was found, false otherwise</returns>
-		public bool TryGetFileNameFromType(Type ExistingType, out string FileName)
+		public bool TryGetFileNameFromType(Type ExistingType, out FileReference FileName)
 		{
 			if(ExistingType.Assembly == CompiledAssembly)
 			{
@@ -1173,16 +1171,16 @@ namespace UnrealBuildTool
 		/// </summary>
 		/// <param name="ModuleName">The name of the module</param>
 		/// <returns>The filename containing rules for this module, or an empty string if not found</returns>
-		public string GetModuleFileName(string ModuleName)
+		public FileReference GetModuleFileName(string ModuleName)
 		{
-			string FileName;
+			FileReference FileName;
 			if(ModuleNameToFileName.TryGetValue(ModuleName, out FileName))
 			{
 				return FileName;
 			}
 			else
 			{
-				return (Parent == null)? "" : Parent.GetModuleFileName(ModuleName);
+				return (Parent == null)? null : Parent.GetModuleFileName(ModuleName);
 			}
 		}
 
@@ -1191,16 +1189,16 @@ namespace UnrealBuildTool
 		/// </summary>
 		/// <param name="TargetName">The name of the target</param>
 		/// <returns>The filename containing rules for this target, or an empty string if not found</returns>
-		public string GetTargetFileName(string TargetName)
+		public FileReference GetTargetFileName(string TargetName)
 		{
-			string FileName;
+			FileReference FileName;
 			if(TargetNameToFileName.TryGetValue(TargetName, out FileName))
 			{
 				return FileName;
 			}
 			else
 			{
-				return (Parent == null)? "" : Parent.GetTargetFileName(TargetName);
+				return (Parent == null)? null : Parent.GetTargetFileName(TargetName);
 			}
 		}
 
@@ -1210,9 +1208,9 @@ namespace UnrealBuildTool
 		/// <param name="ModuleName">Name of the module</param>
 		/// <param name="Target">Information about the target associated with this module</param>
 		/// <returns>Compiled module rule info</returns>
-		public ModuleRules CreateModuleRules( string ModuleName, TargetInfo Target )
+		public ModuleRules CreateModuleRules(string ModuleName, TargetInfo Target)
 		{
-			string ModuleFileName;
+			FileReference ModuleFileName;
 			return CreateModuleRules(ModuleName, Target, out ModuleFileName);
 		}
 
@@ -1223,7 +1221,7 @@ namespace UnrealBuildTool
 		/// <param name="Target">Information about the target associated with this module</param>
 		/// <param name="ModuleFileName">The original source file name for the Module.cs file for this module</param>
 		/// <returns>Compiled module rule info</returns>
-		public ModuleRules CreateModuleRules(string ModuleName, TargetInfo Target, out string ModuleFileName)
+		public ModuleRules CreateModuleRules(string ModuleName, TargetInfo Target, out FileReference ModuleFileName)
 		{
 			// Currently, we expect the user's rules object type name to be the same as the module name
 			var ModuleTypeName = ModuleName;
@@ -1273,13 +1271,7 @@ namespace UnrealBuildTool
 			}
 
 			// Have to do absolute here as this could be a project that is under the root
-			var FullUProjectPath = string.IsNullOrWhiteSpace(UnrealBuildTool.GetUProjectPath())
-				? ""
-				: Path.GetFullPath(UnrealBuildTool.GetUProjectPath());
-			var bProjectModule = string.IsNullOrWhiteSpace(FullUProjectPath)
-				? false
-				: Utils.IsFileUnderDirectory(ModuleFileName, FullUProjectPath);
-
+			bool bProjectModule = UnrealBuildTool.HasUProjectFile() && ModuleFileName.IsUnderDirectory(UnrealBuildTool.GetUProjectPath());
 			if (bProjectModule)
 			{
 				RulesObject.ReadAdditionalDependencies(UnrealBuildTool.GetUProjectFile(), ModuleName);
@@ -1302,8 +1294,8 @@ namespace UnrealBuildTool
 
 					// Choose code optimization options based on module type (game/engine) if
 					// default optimization method is selected.
-					bool bIsEngineModule = Utils.IsFileUnderDirectory(ModuleFileName, ProjectFileGenerator.EngineRelativePath);
-					bool IsGamePluginModule = bProjectModule ? Utils.IsFileUnderDirectory(ModuleFileName, Path.Combine(UnrealBuildTool.GetUProjectPath(), "Plugins")) : false;
+					bool bIsEngineModule = ModuleFileName.IsUnderDirectory(UnrealBuildTool.EngineDirectory);
+					bool IsGamePluginModule = bProjectModule ? ModuleFileName.IsUnderDirectory(DirectoryReference.Combine(UnrealBuildTool.GetUProjectPath(), "Plugins")) : false;
 					if (RulesObject.OptimizeCode == ModuleRules.CodeOptimization.Default)
 					{
 						// Engine/Source and Engine/Plugins are considered 'Engine' code...
@@ -1354,9 +1346,9 @@ namespace UnrealBuildTool
 		/// <param name="Target">Information about the target associated with this module</param>
 		/// <param name="Rules">Output </param>
 		/// <returns>Compiled module rule info</returns>
-		public bool TryCreateModuleRules( string ModuleName, TargetInfo Target, out ModuleRules Rules )
+		public bool TryCreateModuleRules(string ModuleName, TargetInfo Target, out ModuleRules Rules)
 		{
-			if(GetModuleFileName( ModuleName ) == "")
+			if(GetModuleFileName( ModuleName ) == null)
 			{
 				Rules = null;
 				return false;
@@ -1373,8 +1365,8 @@ namespace UnrealBuildTool
 		/// </summary>
 		public bool IsGameModule(string InModuleName)
 		{
-			string ModuleFileName = GetModuleFileName(InModuleName);
-			return ModuleFileName.Length > 0 && !Utils.IsFileUnderDirectory(ModuleFileName, BuildConfiguration.RelativeEnginePath);
+			FileReference ModuleFileName = GetModuleFileName(InModuleName);
+			return (ModuleFileName != null && !ModuleFileName.IsUnderDirectory(UnrealBuildTool.EngineDirectory));
 		}
 
 		protected bool GetTargetTypeAndRulesInstance(string InTargetName, TargetInfo InTarget, out System.Type OutRulesObjectType, out TargetRules OutRulesObject)
@@ -1406,6 +1398,18 @@ namespace UnrealBuildTool
 			return true;
 		}
 
+		/// <summary>
+		/// Creates a target rules object for the specified target name.
+		/// </summary>
+		/// <param name="TargetName">Name of the target</param>
+		/// <param name="Target">Information about the target associated with this target</param>
+		/// <returns>The build target rules for the specified target</returns>
+		public TargetRules CreateTargetRules(string TargetName, TargetInfo Target, bool bInEditorRecompile)
+		{
+			FileReference TargetFileName;
+			return CreateTargetRules(TargetName, Target, bInEditorRecompile, out TargetFileName);
+		}
+
 		/// <summary>
 		/// Creates a target rules object for the specified target name.
 		/// </summary>
@@ -1413,7 +1417,7 @@ namespace UnrealBuildTool
 		/// <param name="Target">Information about the target associated with this target</param>
 		/// <param name="TargetFileName">The original source file name of the Target.cs file for this target</param>
 		/// <returns>The build target rules for the specified target</returns>
-		public TargetRules CreateTargetRules(string TargetName, TargetInfo Target, bool bInEditorRecompile, out string TargetFileName)
+		public TargetRules CreateTargetRules(string TargetName, TargetInfo Target, bool bInEditorRecompile, out FileReference TargetFileName)
 		{
 			// Make sure the target file is known to us
 			bool bFoundTargetName = TargetNameToFileName.ContainsKey(TargetName);
@@ -1446,7 +1450,7 @@ namespace UnrealBuildTool
 					ExceptionMessage += "Location: " + CompiledAssembly.Location + Environment.NewLine;
 
 					ExceptionMessage += "Target rules found:" + Environment.NewLine;
-					foreach (KeyValuePair<string, string> entry in TargetNameToFileName)
+					foreach (KeyValuePair<string, FileReference> entry in TargetNameToFileName)
 					{
 						ExceptionMessage += "\t" + entry.Key + " - " + entry.Value + Environment.NewLine;
 					}
@@ -1477,24 +1481,24 @@ namespace UnrealBuildTool
 					if (RulesObject.Type != TargetRules.TargetType.Editor)
 					{
 						// Not the editor... determine the editor project
-						string TargetSourceFolder = TargetFileName;
+						string TargetSourceFolderString = TargetFileName.FullName;
 						Int32 SourceFolderIndex = -1;
 						if (Utils.IsRunningOnMono)
 						{
-							TargetSourceFolder = TargetSourceFolder.Replace("\\", "/");
-							SourceFolderIndex = TargetSourceFolder.LastIndexOf("/Source/", StringComparison.InvariantCultureIgnoreCase);
+							TargetSourceFolderString = TargetSourceFolderString.Replace("\\", "/");
+							SourceFolderIndex = TargetSourceFolderString.LastIndexOf("/Source/", StringComparison.InvariantCultureIgnoreCase);
 						}
 						else
 						{
-							TargetSourceFolder = TargetSourceFolder.Replace("/", "\\");
-							SourceFolderIndex = TargetSourceFolder.LastIndexOf("\\Source\\", StringComparison.InvariantCultureIgnoreCase);
+							TargetSourceFolderString = TargetSourceFolderString.Replace("/", "\\");
+							SourceFolderIndex = TargetSourceFolderString.LastIndexOf("\\Source\\", StringComparison.InvariantCultureIgnoreCase);
 						}
 						if (SourceFolderIndex != -1)
 						{
-							TargetSourceFolder = TargetSourceFolder.Substring(0, SourceFolderIndex + 8);
-							foreach (KeyValuePair<string, string> CheckEntry in TargetNameToFileName)
+							DirectoryReference TargetSourceFolder = new DirectoryReference(TargetSourceFolderString.Substring(0, SourceFolderIndex + 8));
+							foreach (KeyValuePair<string, FileReference> CheckEntry in TargetNameToFileName)
 							{
-								if (CheckEntry.Value.StartsWith(TargetSourceFolder, StringComparison.InvariantCultureIgnoreCase))
+								if (CheckEntry.Value.IsUnderDirectory(TargetSourceFolder))
 								{
 									if (CheckEntry.Key.Equals(TargetName, StringComparison.InvariantCultureIgnoreCase) == false)
 									{
@@ -1565,7 +1569,7 @@ namespace UnrealBuildTool
 		class RulesFileCache
 		{
 			/// List of rules file paths for each of the known types in RulesFileType
-			public List<string>[] RulesFilePaths = new List<string>[ typeof( RulesFileType ).GetEnumValues().Length ];
+			public List<FileReference>[] RulesFilePaths = new List<FileReference>[ typeof( RulesFileType ).GetEnumValues().Length ];
 		}
 
 
@@ -1597,11 +1601,11 @@ namespace UnrealBuildTool
 							{
 								if (RulesFileCache.RulesFilePaths[(int)CurRulesType] == null)
 								{
-									RulesFileCache.RulesFilePaths[(int)CurRulesType] = new List<string>();
+									RulesFileCache.RulesFilePaths[(int)CurRulesType] = new List<FileReference>();
 								}
 
 								// Convert file info to the full file path for this file and update our cache
-								RulesFileCache.RulesFilePaths[(int)CurRulesType].Add(RuleFile.FullName);
+								RulesFileCache.RulesFilePaths[(int)CurRulesType].Add(new FileReference(RuleFile.FullName));
 
 								// NOTE: Multiple rules files in the same folder are supported.  We'll continue iterating along.
 								if( CurRulesType == RulesFileType.Module )
@@ -1637,25 +1641,25 @@ namespace UnrealBuildTool
 
 		/// Map of root folders to a cached list of all UBT-related source files in that folder or any of its sub-folders.
 		/// We cache these file names so we can avoid searching for them later on.
-		static Dictionary<string, RulesFileCache> RootFolderToRulesFileCache = new Dictionary<string, RulesFileCache>();
+		static Dictionary<DirectoryReference, RulesFileCache> RootFolderToRulesFileCache = new Dictionary<DirectoryReference, RulesFileCache>();
 
-		public static List<string> FindAllRulesSourceFiles( RulesFileType RulesFileType, List<string> GameFolders, List<string> ForeignPlugins, List<string> AdditionalSearchPaths, bool bIncludeEngine = true )
+		public static List<FileReference> FindAllRulesSourceFiles( RulesFileType RulesFileType, List<DirectoryReference> GameFolders, List<FileReference> ForeignPlugins, List<DirectoryReference> AdditionalSearchPaths, bool bIncludeEngine = true )
 		{
-			List<string> Folders = new List<string>();
+			List<DirectoryReference> Folders = new List<DirectoryReference>();
 
 			// Add all engine source (including third party source)
 			if(bIncludeEngine)
 			{
-				Folders.Add( Path.Combine( ProjectFileGenerator.EngineRelativePath, "Source" ) );
+				Folders.Add(UnrealBuildTool.EngineSourceDirectory);
 			}
 
 			// @todo plugin: Disallow modules from including plugin modules as dependency modules? (except when the module is part of that plugin)
 
 			// Get all the root folders for plugins
-			List<string> RootFolders = new List<string>();
+			List<DirectoryReference> RootFolders = new List<DirectoryReference>();
 			if(bIncludeEngine)
 			{
-				RootFolders.Add(ProjectFileGenerator.EngineRelativePath);
+				RootFolders.Add(UnrealBuildTool.EngineDirectory);
 			}
 			if(GameFolders != null)
 			{
@@ -1663,34 +1667,32 @@ namespace UnrealBuildTool
 			}
 
 			// Find all the plugin source directories
-			foreach(string RootFolder in RootFolders)
+			foreach(DirectoryReference RootFolder in RootFolders)
 			{
-				string PluginsFolder = Path.Combine(RootFolder, "Plugins");
-				foreach(string PluginFile in Plugins.EnumeratePlugins(PluginsFolder))
+				DirectoryReference PluginsFolder = DirectoryReference.Combine(RootFolder, "Plugins");
+				foreach(FileReference PluginFile in Plugins.EnumeratePlugins(PluginsFolder))
 				{
-					string PluginDirectory = Path.GetDirectoryName(PluginFile);
-					Folders.Add(Path.Combine(PluginDirectory, "Source"));
+					Folders.Add(DirectoryReference.Combine(PluginFile.Directory, "Source"));
 				}
 			}
 
 			// Add all the extra plugin folders
 			if( ForeignPlugins != null )
 			{
-				foreach(string ForeignPlugin in ForeignPlugins)
+				foreach(FileReference ForeignPlugin in ForeignPlugins)
 				{
-					string PluginDirectory = Path.GetDirectoryName(Path.GetFullPath(ForeignPlugin));
-					Folders.Add(Path.Combine(PluginDirectory, "Source"));
+					Folders.Add(DirectoryReference.Combine(ForeignPlugin.Directory, "Source"));
 				}
 			}
 
 			// Add in the game folders to search
 			if( GameFolders != null )
 			{
-				foreach( var GameFolder in GameFolders )
+				foreach( DirectoryReference GameFolder in GameFolders )
 				{
-					var GameSourceFolder = Path.GetFullPath(Path.Combine( GameFolder, "Source" ));
-					Folders.Add( GameSourceFolder );
-					var GameIntermediateSourceFolder = Path.GetFullPath(Path.Combine(GameFolder, "Intermediate", "Source"));
+					DirectoryReference GameSourceFolder = DirectoryReference.Combine(GameFolder, "Source");
+					Folders.Add(GameSourceFolder);
+					DirectoryReference GameIntermediateSourceFolder = DirectoryReference.Combine(GameFolder, "Intermediate", "Source");
 					Folders.Add(GameIntermediateSourceFolder);
 				}
 			}
@@ -1700,9 +1702,9 @@ namespace UnrealBuildTool
 			{
 				foreach( var AdditionalSearchPath in AdditionalSearchPaths )
 				{
-					if (!string.IsNullOrEmpty(AdditionalSearchPath))
+					if ( AdditionalSearchPath != null )
 					{
-						if (Directory.Exists(AdditionalSearchPath))
+						if (AdditionalSearchPath.Exists())
 						{
 							Folders.Add(AdditionalSearchPath);
 						}
@@ -1714,17 +1716,16 @@ namespace UnrealBuildTool
 				}
 			}
 
-			var SourceFiles = new List<string>();
-
 			// Iterate over all the folders to check
-			foreach( string Folder in Folders )
+			List<FileReference> SourceFiles = new List<FileReference>();
+			foreach( DirectoryReference Folder in Folders )
 			{
 				// Check to see if we've already cached source files for this folder
 				RulesFileCache FolderRulesFileCache;
 				if (!RootFolderToRulesFileCache.TryGetValue(Folder, out FolderRulesFileCache))
 				{
 					FolderRulesFileCache = new RulesFileCache();
-					FindAllRulesFilesRecursively(new DirectoryInfo(Folder), FolderRulesFileCache);
+					FindAllRulesFilesRecursively(new DirectoryInfo(Folder.FullName), FolderRulesFileCache);
 					RootFolderToRulesFileCache[Folder] = FolderRulesFileCache;
 
 					if (BuildConfiguration.bPrintDebugInfo)
@@ -1743,7 +1744,7 @@ namespace UnrealBuildTool
 				var RulesFilePathsForType = FolderRulesFileCache.RulesFilePaths[(int)RulesFileType];
 				if (RulesFilePathsForType != null)
 				{
-					foreach (string RulesFilePath in RulesFilePathsForType)
+					foreach (FileReference RulesFilePath in RulesFilePathsForType)
 					{
 						if (!SourceFiles.Contains(RulesFilePath))
 						{
@@ -1752,7 +1753,6 @@ namespace UnrealBuildTool
 					}
 				}
 			}
-
 			return SourceFiles;
 		}
 
@@ -1763,7 +1763,7 @@ namespace UnrealBuildTool
 
 		/// Map of assembly names we've already compiled and loaded to their Assembly and list of game folders.  This is used to prevent
 		/// trying to recompile the same assembly when ping-ponging between different types of targets
-		private static Dictionary<string, RulesAssembly> LoadedAssemblyMap = new Dictionary<string, RulesAssembly>(StringComparer.InvariantCultureIgnoreCase);
+		private static Dictionary<FileReference, RulesAssembly> LoadedAssemblyMap = new Dictionary<FileReference, RulesAssembly>();
 
 		/// <summary>
 		/// Creates the engine rules assembly
@@ -1775,11 +1775,11 @@ namespace UnrealBuildTool
 			if(EngineRulesAssembly == null)
 			{
 				// Find all the rules files
-				List<string> ModuleFileNames = FindAllRulesSourceFiles(RulesFileType.Module, null, null, null, true);
-				List<string> TargetFileNames = FindAllRulesSourceFiles(RulesFileType.Target, null, null, null, true);
+				List<FileReference> ModuleFileNames = FindAllRulesSourceFiles(RulesFileType.Module, null, null, null, true);
+				List<FileReference> TargetFileNames = FindAllRulesSourceFiles(RulesFileType.Target, null, null, null, true);
 
 				// Create a path to the assembly that we'll either load or compile
-				string AssemblyFileName = Path.Combine(BuildConfiguration.BaseIntermediatePath, "BuildRules", "UE4Rules.dll");
+				FileReference AssemblyFileName = FileReference.Combine(UnrealBuildTool.EngineDirectory, BuildConfiguration.BaseIntermediateFolder, "BuildRules", "UE4Rules.dll");
 				EngineRulesAssembly = new RulesAssembly(ModuleFileNames, TargetFileNames, AssemblyFileName, null);
 			}
 			return EngineRulesAssembly;
@@ -1790,39 +1790,37 @@ namespace UnrealBuildTool
 		/// </summary>
 		/// <param name="ProjectFileName">The project file to create rules for. Null for the engine.</param>
 		/// <param name="ForeignPlugins">List of foreign plugin folders to include in the assembly. May be null.</param>
-		public static RulesAssembly CreateProjectRulesAssembly(string ProjectFileName)
+		public static RulesAssembly CreateProjectRulesAssembly(FileReference ProjectFileName)
 		{
-			// Normalize the project filename, so we can look it up in the cache
-			string FullProjectFileName = Path.GetFullPath(ProjectFileName);
-
 			// Check if there's an existing assembly for this project
 			RulesAssembly ProjectRulesAssembly;
-			if(!LoadedAssemblyMap.TryGetValue(FullProjectFileName, out ProjectRulesAssembly))
+			if(!LoadedAssemblyMap.TryGetValue(ProjectFileName, out ProjectRulesAssembly))
 			{
+				DirectoryReference ProjectDirectoryName = ProjectFileName.Directory;
+
 				// Create the engine rules assembly
 				RulesAssembly Parent = CreateEngineRulesAssembly();
 
 				// Get the folders to search through
-				List<string> GameFolders = new List<string>();
-				string ProjectDirectoryName = Path.GetDirectoryName(FullProjectFileName);
+				List<DirectoryReference> GameFolders = new List<DirectoryReference>();
 				GameFolders.Add(ProjectDirectoryName);
 
 				// Add the games project's intermediate source folder
-				List<string> AdditionalSearchPaths = new List<string>();
-				var ProjectIntermediateSourceDirectory = Path.Combine(ProjectDirectoryName, "Intermediate", "Source");
-				if (Directory.Exists(ProjectIntermediateSourceDirectory))
+				List<DirectoryReference> AdditionalSearchPaths = new List<DirectoryReference>();
+				DirectoryReference ProjectIntermediateSourceDirectory = DirectoryReference.Combine(ProjectDirectoryName, "Intermediate", "Source");
+				if (ProjectIntermediateSourceDirectory.Exists())
 				{
 					AdditionalSearchPaths.Add(ProjectIntermediateSourceDirectory);
 				}
 
 				// Find all the rules source files
-				List<string> ModuleFileNames = FindAllRulesSourceFiles(RulesFileType.Module, GameFolders, null, AdditionalSearchPaths, false);
-				List<string> TargetFileNames = FindAllRulesSourceFiles(RulesFileType.Target, GameFolders, null, AdditionalSearchPaths, false);
+				List<FileReference> ModuleFileNames = FindAllRulesSourceFiles(RulesFileType.Module, GameFolders, null, AdditionalSearchPaths, false);
+				List<FileReference> TargetFileNames = FindAllRulesSourceFiles(RulesFileType.Target, GameFolders, null, AdditionalSearchPaths, false);
 
 				// Compile the assembly
-				string AssemblyFileName = Path.Combine(ProjectDirectoryName, BuildConfiguration.BaseIntermediateFolder, "BuildRules", Path.GetFileNameWithoutExtension(ProjectFileName) + "ModuleRules.dll");
+				FileReference AssemblyFileName = FileReference.Combine(ProjectDirectoryName, BuildConfiguration.BaseIntermediateFolder, "BuildRules", ProjectFileName.GetFileNameWithoutExtension() + "ModuleRules.dll");
 				ProjectRulesAssembly = new RulesAssembly(ModuleFileNames, TargetFileNames, AssemblyFileName, Parent);
-				LoadedAssemblyMap.Add(FullProjectFileName, ProjectRulesAssembly);
+				LoadedAssemblyMap.Add(ProjectFileName, ProjectRulesAssembly);
 			}
 			return ProjectRulesAssembly;
 		}
@@ -1832,23 +1830,20 @@ namespace UnrealBuildTool
 		/// </summary>
 		/// <param name="ProjectFileName">The project file to create rules for. Null for the engine.</param>
 		/// <param name="ForeignPlugins">List of foreign plugin folders to include in the assembly. May be null.</param>
-		public static RulesAssembly CreatePluginRulesAssembly(string PluginFileName, RulesAssembly Parent)
+		public static RulesAssembly CreatePluginRulesAssembly(FileReference PluginFileName, RulesAssembly Parent)
 		{
-			// Normalize the project filename, so we can look it up in the cache
-			string FullPluginFileName = Path.GetFullPath(PluginFileName);
-
 			// Check if there's an existing assembly for this project
 			RulesAssembly PluginRulesAssembly;
-			if(!LoadedAssemblyMap.TryGetValue(FullPluginFileName, out PluginRulesAssembly))
+			if(!LoadedAssemblyMap.TryGetValue(PluginFileName, out PluginRulesAssembly))
 			{
 				// Find all the rules source files
-				List<string> ModuleFileNames = FindAllRulesSourceFiles(RulesFileType.Module, null, new List<string>{ PluginFileName }, null, false);
-				List<string> TargetFileNames = FindAllRulesSourceFiles(RulesFileType.Target, null, new List<string>{ PluginFileName }, null, false);
+				List<FileReference> ModuleFileNames = FindAllRulesSourceFiles(RulesFileType.Module, null, new List<FileReference>{ PluginFileName }, null, false);
+				List<FileReference> TargetFileNames = FindAllRulesSourceFiles(RulesFileType.Target, null, new List<FileReference>{ PluginFileName }, null, false);
 
 				// Compile the assembly
-				string AssemblyFileName = Path.Combine(Path.GetDirectoryName(PluginFileName), BuildConfiguration.BaseIntermediateFolder, "BuildRules", Path.GetFileNameWithoutExtension(PluginFileName) + "ModuleRules.dll");
+				FileReference AssemblyFileName = FileReference.Combine(PluginFileName.Directory, BuildConfiguration.BaseIntermediateFolder, "BuildRules", Path.GetFileNameWithoutExtension(PluginFileName.FullName) + "ModuleRules.dll");
 				PluginRulesAssembly = new RulesAssembly(ModuleFileNames, TargetFileNames, AssemblyFileName, Parent);
-				LoadedAssemblyMap.Add(FullPluginFileName, PluginRulesAssembly);
+				LoadedAssemblyMap.Add(PluginFileName, PluginRulesAssembly);
 			}
 			return PluginRulesAssembly;
 		}
@@ -1860,16 +1855,16 @@ namespace UnrealBuildTool
 		/// <returns>The filename that declared the given type, or null</returns>
 		public static string GetFileNameFromType(Type ExistingType)
 		{
-			string FileName;
+			FileReference FileName;
 			if(EngineRulesAssembly != null && EngineRulesAssembly.TryGetFileNameFromType(ExistingType, out FileName))
 			{
-				return FileName;
+				return FileName.FullName;
 			}
 			foreach(RulesAssembly RulesAssembly in LoadedAssemblyMap.Values)
 			{
 				if(RulesAssembly.TryGetFileNameFromType(ExistingType, out FileName))
 				{
-					return FileName;
+					return FileName.FullName;
 				}
 			}
 			return null;
diff --git a/Engine/Source/Programs/UnrealBuildTool/System/SourceFileSearch.cs b/Engine/Source/Programs/UnrealBuildTool/System/SourceFileSearch.cs
--- a/Engine/Source/Programs/UnrealBuildTool/System/SourceFileSearch.cs
+++ b/Engine/Source/Programs/UnrealBuildTool/System/SourceFileSearch.cs
@@ -15,14 +15,14 @@ namespace UnrealBuildTool
 
 
 		/// Finds mouse source files
-		public static List<string> FindModuleSourceFiles( string ModuleRulesFile, bool ExcludeNoRedistFiles, bool SearchSubdirectories = true, bool IncludePrivateSourceCode = true )
+		public static List<FileReference> FindModuleSourceFiles( FileReference ModuleRulesFile, bool ExcludeNoRedistFiles, bool SearchSubdirectories = true, bool IncludePrivateSourceCode = true )
 		{
 			// The module's "base directory" is simply the directory where its xxx.Build.cs file is stored.  We'll always
 			// harvest source files for this module in this base directory directory and all of its sub-directories.
-			var ModuleBaseDirectory = Path.GetDirectoryName( ModuleRulesFile );
+			var ModuleBaseDirectory = ModuleRulesFile.Directory;
 
 			// Find all of the source files (and other files) and add them to the project
-			var DirectoriesToSearch = new List<string>();
+			var DirectoriesToSearch = new List<DirectoryReference>();
 			DirectoriesToSearch.Add( ModuleBaseDirectory );
 			var FoundFiles = SourceFileSearch.FindFiles( DirectoriesToSearch, ExcludeNoRedistFiles: ExcludeNoRedistFiles, SubdirectoryNamesToExclude:null, SearchSubdirectories:SearchSubdirectories, IncludePrivateSourceCode:IncludePrivateSourceCode );
 
@@ -38,7 +38,7 @@ namespace UnrealBuildTool
 		/// <param name="SubdirectoryNamesToExclude">Directory base names to ignore when searching subdirectories.  Can be null.</param>
 		/// <param name="SearchSubdirectories">True to include subdirectories, otherwise we only search the list of base directories</param>
 		/// <param name="IncludePrivateSourceCode">True if source files in the 'Private' directory should be included</param>
-		public static List<string> FindFiles( List<string> DirectoriesToSearch, bool ExcludeNoRedistFiles, List<string> SubdirectoryNamesToExclude = null, bool SearchSubdirectories = true, bool IncludePrivateSourceCode = true )
+		public static List<FileReference> FindFiles( List<DirectoryReference> DirectoriesToSearch, bool ExcludeNoRedistFiles, List<string> SubdirectoryNamesToExclude = null, bool SearchSubdirectories = true, bool IncludePrivateSourceCode = true )
 		{
 			// Certain file types should never be added to project files
 			var ExcludedFileTypes = new string[] 
@@ -61,7 +61,7 @@ namespace UnrealBuildTool
 							"DO_NOT_DELETE.txt",	// Placeholder .txt file in P4
 						};
 
-			var FoundFiles = new List<string>();
+			HashSet<FileReference> FoundFiles = new HashSet<FileReference>();
 			foreach( var SearchDirectory in DirectoriesToSearch )
 			{
 				var FoundFileUnderIntermediateDirectory = string.Empty;
@@ -75,7 +75,7 @@ namespace UnrealBuildTool
 					foreach(var CurSourceFilePath in FileList)
 					{
 						// Get the part of the path that is relative to the module folder
-						var RelativeSourceFilePath = CurSourceFilePath.Substring( SearchDirectory.Length + 1 );
+						var RelativeSourceFilePath = new FileReference(CurSourceFilePath).MakeRelativeTo(SearchDirectory);
 
 						bool IncludeThisFile = true;
 
@@ -167,9 +167,9 @@ namespace UnrealBuildTool
 						if( IncludeThisFile )
 						{
 							// Don't add duplicates
-							if( !FoundFiles.Contains( CurSourceFilePath ) )
+							if( !FoundFiles.Contains( new FileReference(CurSourceFilePath) ) )
 							{
-								FoundFiles.Add( CurSourceFilePath );
+								FoundFiles.Add( new FileReference(CurSourceFilePath) );
 							}
 						}
 					}
@@ -178,11 +178,11 @@ namespace UnrealBuildTool
 				// call the delegate to do the work, based on platform (see note above)
 				if(!Utils.IsRunningOnMono)
 				{
-					DoSearch(Directory.EnumerateFiles( Utils.CleanDirectorySeparators( SearchDirectory ), "*.*", SearchSubdirectories ? SearchOption.AllDirectories : SearchOption.TopDirectoryOnly ));
+					DoSearch(Directory.EnumerateFiles(SearchDirectory.FullName, "*.*", SearchSubdirectories ? SearchOption.AllDirectories : SearchOption.TopDirectoryOnly ));
 				}
 				else
 				{
-					DoSearch(Directory.GetFiles(Utils.CleanDirectorySeparators(SearchDirectory), "*.*", SearchSubdirectories ? SearchOption.AllDirectories : SearchOption.TopDirectoryOnly).AsEnumerable());
+					DoSearch(Directory.GetFiles(SearchDirectory.FullName, "*.*", SearchSubdirectories ? SearchOption.AllDirectories : SearchOption.TopDirectoryOnly).AsEnumerable());
 				}
 
 				if( !string.IsNullOrEmpty( FoundFileUnderIntermediateDirectory ) )
@@ -191,7 +191,7 @@ namespace UnrealBuildTool
 				}
 			}
 
-			return FoundFiles;
+			return FoundFiles.ToList();
 		}
 
 	}
diff --git a/Engine/Source/Programs/UnrealBuildTool/System/TargetReceipt.cs b/Engine/Source/Programs/UnrealBuildTool/System/TargetReceipt.cs
--- a/Engine/Source/Programs/UnrealBuildTool/System/TargetReceipt.cs
+++ b/Engine/Source/Programs/UnrealBuildTool/System/TargetReceipt.cs
@@ -196,6 +196,17 @@ namespace UnrealBuildTool
 			return NewBuildProduct;
 		}
 
+		/// <summary>
+		/// Adds a build product to the receipt. Does not check whether it already exists.
+		/// </summary>
+		/// <param name="Path">Path to the build product.</param>
+		/// <param name="Type">Type of build product.</param>
+		/// <returns>The BuildProduct object that was created</returns>
+		public BuildProduct AddBuildProduct(FileReference Path, BuildProductType Type)
+		{
+			return AddBuildProduct(Path.FullName, Type);
+		}
+
 		/// <summary>
 		/// Constructs a runtime dependency object and adds it to the receipt.
 		/// </summary>
@@ -233,7 +244,7 @@ namespace UnrealBuildTool
 		/// </summary>
 		/// <param name="EngineDir">Value for the $(EngineDir) variable</param>
 		/// <param name="ProjectDir">Value for the $(ProjectDir) variable</param>
-		public void ExpandPathVariables(string EngineDir, string ProjectDir)
+		public void ExpandPathVariables(DirectoryReference EngineDir, DirectoryReference ProjectDir)
 		{
 			ExpandPathVariables(EngineDir, ProjectDir, new Dictionary<string, string>());
 		}
@@ -243,12 +254,12 @@ namespace UnrealBuildTool
 		/// </summary>
 		/// <param name="EngineDir">Value for the $(EngineDir) variable</param>
 		/// <param name="ProjectDir">Value for the $(ProjectDir) variable</param>
-		public void ExpandPathVariables(string EngineDir, string ProjectDir, IDictionary<string, string> OtherVariables)
+		public void ExpandPathVariables(DirectoryReference EngineDir, DirectoryReference ProjectDir, IDictionary<string, string> OtherVariables)
 		{
 			// Build a dictionary containing the standard variable expansions
 			Dictionary<string, string> Variables = new Dictionary<string, string>(OtherVariables);
-			Variables["EngineDir"] = Path.GetFullPath(EngineDir).TrimEnd(Path.DirectorySeparatorChar);
-			Variables["ProjectDir"] = Path.GetFullPath(ProjectDir).TrimEnd(Path.DirectorySeparatorChar);
+			Variables["EngineDir"] = EngineDir.FullName;
+			Variables["ProjectDir"] = ProjectDir.FullName;
 
 			// Replace all the variables in the paths
 			foreach(BuildProduct BuildProduct in BuildProducts)
@@ -270,10 +281,10 @@ namespace UnrealBuildTool
 		/// </summary>
 		/// <param name="EngineDir">The engine directory. Relative paths are ok.</param>
 		/// <param name="ProjectDir">The project directory. Relative paths are ok.</param>
-		public void InsertStandardPathVariables(string EngineDir, string ProjectDir)
+		public void InsertStandardPathVariables(DirectoryReference EngineDir, DirectoryReference ProjectDir)
 		{
-			string EnginePrefix = Path.GetFullPath(EngineDir).TrimEnd(Path.DirectorySeparatorChar) + Path.DirectorySeparatorChar;
-			string ProjectPrefix = Path.GetFullPath(ProjectDir).TrimEnd(Path.DirectorySeparatorChar) + Path.DirectorySeparatorChar;
+			string EnginePrefix = EngineDir.FullName.TrimEnd(Path.DirectorySeparatorChar) + Path.DirectorySeparatorChar;
+			string ProjectPrefix = ProjectDir.FullName.TrimEnd(Path.DirectorySeparatorChar) + Path.DirectorySeparatorChar;
 
 			foreach(BuildProduct BuildProduct in BuildProducts)
 			{
