Use FileReference and DirectoryReference classes instead of strings for paths. Massively reduces the amount of paranoid calls to GetFullPath(), CleanDirectorySeparators() et al., and improves performa

diff --git a/Engine/Source/Programs/UnrealBuildTool/System/ThirdPartyHeaderFinder.cs b/Engine/Source/Programs/UnrealBuildTool/System/ThirdPartyHeaderFinder.cs
--- a/Engine/Source/Programs/UnrealBuildTool/System/ThirdPartyHeaderFinder.cs
+++ b/Engine/Source/Programs/UnrealBuildTool/System/ThirdPartyHeaderFinder.cs
@@ -94,14 +94,13 @@ namespace UnrealBuildTool
 			RulesAssembly EngineRulesAssembly = RulesCompiler.CreateEngineRulesAssembly();
 
 			// Find all modules referenced by the current target
-			List<string> ModuleFileNames = RulesCompiler.FindAllRulesSourceFiles( RulesCompiler.RulesFileType.Module, GameFolders:null, ForeignPlugins:null, AdditionalSearchPaths:null );
-			foreach (string ModuleFileName in ModuleFileNames)
+			List<FileReference> ModuleFileNames = RulesCompiler.FindAllRulesSourceFiles( RulesCompiler.RulesFileType.Module, GameFolders:null, ForeignPlugins:null, AdditionalSearchPaths:null );
+			foreach (FileReference ModuleFileName in ModuleFileNames)
 			{
-				string ModuleName = Path.GetFileNameWithoutExtension(Path.GetFileNameWithoutExtension(ModuleFileName));
+				string ModuleName = Path.GetFileNameWithoutExtension(ModuleFileName.GetFileNameWithoutExtension());
 				try
 				{
-					string UnusedModuleFilename;
-					ModuleRules RulesObject = EngineRulesAssembly.CreateModuleRules(ModuleName, Target, out UnusedModuleFilename);
+					ModuleRules RulesObject = EngineRulesAssembly.CreateModuleRules(ModuleName, Target);
 					bool bEngineHeaders = RulesObject.Type != ModuleRules.ModuleType.External;
 					foreach (string SystemIncludePath in RulesObject.PublicSystemIncludePaths)
 					{
diff --git a/Engine/Source/Programs/UnrealBuildTool/System/UEPlatformProjectGenerator.cs b/Engine/Source/Programs/UnrealBuildTool/System/UEPlatformProjectGenerator.cs
--- a/Engine/Source/Programs/UnrealBuildTool/System/UEPlatformProjectGenerator.cs
+++ b/Engine/Source/Programs/UnrealBuildTool/System/UEPlatformProjectGenerator.cs
@@ -87,7 +87,7 @@ namespace UnrealBuildTool
 		/// </summary>
 		/// <param name="InPlatform"></param>
 		/// <returns></returns>
-		public static bool GenerateGamePlatformSpecificProperties(UnrealTargetPlatform InPlatform, UnrealTargetConfiguration Configuration, TargetRules.TargetType TargetType, StringBuilder VCProjectFileContent, string RootDirectory, string TargetFilePath)
+		public static bool GenerateGamePlatformSpecificProperties(UnrealTargetPlatform InPlatform, UnrealTargetConfiguration Configuration, TargetRules.TargetType TargetType, StringBuilder VCProjectFileContent, DirectoryReference RootDirectory, FileReference TargetFilePath)
 		{
 			if (ProjectGeneratorDictionary.ContainsKey(InPlatform) == true)
 			{
@@ -121,7 +121,7 @@ namespace UnrealBuildTool
 			// Do nothing
 		}
 
-		public virtual void GenerateGameProperties(UnrealTargetConfiguration Configuration, StringBuilder VCProjectFileContent, TargetRules.TargetType TargetType, string RootDirectory, string TargetFilePath)
+		public virtual void GenerateGameProperties(UnrealTargetConfiguration Configuration, StringBuilder VCProjectFileContent, TargetRules.TargetType TargetType, DirectoryReference RootDirectory, FileReference TargetFilePath)
 		{
 			// Do nothing
 		}
@@ -209,7 +209,7 @@ namespace UnrealBuildTool
 		 *	
 		 *	@return	string				The custom path lines for the project file; Empty string if it doesn't require one
 		 */
-		public virtual string GetVisualStudioPathsEntries(UnrealTargetPlatform InPlatform, UnrealTargetConfiguration InConfiguration, TargetRules.TargetType TargetType, string TargetRulesPath, string ProjectFilePath, string NMakeOutputPath)
+		public virtual string GetVisualStudioPathsEntries(UnrealTargetPlatform InPlatform, UnrealTargetConfiguration InConfiguration, TargetRules.TargetType TargetType, FileReference TargetRulesPath, FileReference ProjectFilePath, FileReference NMakeOutputPath)
 		{
 			// NOTE: We are intentionally overriding defaults for these paths with empty strings.  We never want Visual Studio's
 			//       defaults for these fields to be propagated, since they are version-sensitive paths that may not reflect
@@ -247,7 +247,7 @@ namespace UnrealBuildTool
 		 *	
 		 *	@return	string				The custom property import lines for the project file; Empty string if it doesn't require one
 		 */
-		public virtual string GetVisualStudioLayoutDirSection(UnrealTargetPlatform InPlatform, UnrealTargetConfiguration InConfiguration, string InConditionString, TargetRules.TargetType TargetType, string TargetRulesPath, string ProjectFilePath, string NMakeOutputPath)
+		public virtual string GetVisualStudioLayoutDirSection(UnrealTargetPlatform InPlatform, UnrealTargetConfiguration InConfiguration, string InConditionString, TargetRules.TargetType TargetType, FileReference TargetRulesPath, FileReference ProjectFilePath, FileReference NMakeOutputPath)
 		{
 			return "";
 		}
@@ -259,7 +259,7 @@ namespace UnrealBuildTool
 		 *	
 		 *	@return	string				The output manifest section for the project file; Empty string if it doesn't require one
 		 */
-        public virtual string GetVisualStudioOutputManifestSection(UnrealTargetPlatform InPlatform, TargetRules.TargetType TargetType, string TargetRulesPath, string ProjectFilePath)
+        public virtual string GetVisualStudioOutputManifestSection(UnrealTargetPlatform InPlatform, TargetRules.TargetType TargetType, FileReference TargetRulesPath, FileReference ProjectFilePath)
 		{
 			return "";
 		}
@@ -285,7 +285,7 @@ namespace UnrealBuildTool
 		/// <param name="ProjectFilePath">The project file path</param>
 		/// <returns>The string to append to the user file</returns>
 		public virtual string GetVisualStudioUserFileStrings(UnrealTargetPlatform InPlatform, UnrealTargetConfiguration InConfiguration,
-			string InConditionString, TargetRules InTargetRules, string TargetRulesPath, string ProjectFilePath)
+			string InConditionString, TargetRules InTargetRules, FileReference TargetRulesPath, FileReference ProjectFilePath)
 		{
 			return "";
 		}
diff --git a/Engine/Source/Programs/UnrealBuildTool/System/UProjectInfo.cs b/Engine/Source/Programs/UnrealBuildTool/System/UProjectInfo.cs
--- a/Engine/Source/Programs/UnrealBuildTool/System/UProjectInfo.cs
+++ b/Engine/Source/Programs/UnrealBuildTool/System/UProjectInfo.cs
@@ -16,44 +16,44 @@ namespace UnrealBuildTool
 	{
 		public string GameName;
 		public string FileName;
-		public string FilePath;
-		public string Folder;
+		public FileReference FilePath;
+		public DirectoryReference Folder;
 		public bool bIsCodeProject;
 
-		UProjectInfo(string InFilePath, bool bInIsCodeProject)
+		UProjectInfo(FileReference InFilePath, bool bInIsCodeProject)
 		{
-			GameName = Path.GetFileNameWithoutExtension(InFilePath);
-			FileName = Path.GetFileName(InFilePath);
+			GameName = InFilePath.GetFileNameWithoutExtension();
+			FileName = InFilePath.GetFileName();
 			FilePath = InFilePath;
-			Folder = Path.GetDirectoryName(InFilePath).TrimEnd('\\', '/');
+			Folder = FilePath.Directory;
 			bIsCodeProject = bInIsCodeProject;
 		}
 
 		/** Map of relative or complete project file names to the project info */
-		static Dictionary<string, UProjectInfo> ProjectInfoDictionary = new Dictionary<string, UProjectInfo>( StringComparer.InvariantCultureIgnoreCase );
+		static Dictionary<FileReference, UProjectInfo> ProjectInfoDictionary = new Dictionary<FileReference, UProjectInfo>();
 		/** Map of short project file names to the relative or complete project file name */
-		static Dictionary<string, string> ShortProjectNameDictionary = new Dictionary<string,string>( StringComparer.InvariantCultureIgnoreCase );
-		/** Map of targetnames to the relative or complete project file name */
-		static Dictionary<string, string> TargetToProjectDictionary = new Dictionary<string,string>( StringComparer.InvariantCultureIgnoreCase );
+		static Dictionary<string, FileReference> ShortProjectNameDictionary = new Dictionary<string,FileReference>( StringComparer.InvariantCultureIgnoreCase );
+		/** Map of target names to the relative or complete project file name */
+		static Dictionary<string, FileReference> TargetToProjectDictionary = new Dictionary<string,FileReference>( StringComparer.InvariantCultureIgnoreCase );
 
-		public static bool FindTargetFilesInFolder(string InTargetFolder)
+		public static bool FindTargetFilesInFolder(DirectoryReference InTargetFolder)
 		{
 			bool bFoundTargetFiles = false;
 			IEnumerable<string> Files;
 			if (!Utils.IsRunningOnMono)
 			{
-				Files = Directory.EnumerateFiles (InTargetFolder, "*.target.cs", SearchOption.TopDirectoryOnly);
+				Files = Directory.EnumerateFiles (InTargetFolder.FullName, "*.target.cs", SearchOption.TopDirectoryOnly);
 			}
 			else
 			{
-				Files = Directory.GetFiles (InTargetFolder, "*.Target.cs", SearchOption.TopDirectoryOnly).AsEnumerable();
+				Files = Directory.GetFiles (InTargetFolder.FullName, "*.Target.cs", SearchOption.TopDirectoryOnly).AsEnumerable();
 			}
 			foreach (var TargetFilename in Files)
 			{
 				bFoundTargetFiles = true;
-				foreach (KeyValuePair<string, UProjectInfo> Entry in ProjectInfoDictionary)
+				foreach (KeyValuePair<FileReference, UProjectInfo> Entry in ProjectInfoDictionary)
 				{
-					FileInfo ProjectFileInfo = new FileInfo(Entry.Key);
+					FileInfo ProjectFileInfo = new FileInfo(Entry.Key.FullName);
 					string ProjectDir = ProjectFileInfo.DirectoryName.TrimEnd(Path.DirectorySeparatorChar) + Path.DirectorySeparatorChar;
 					if (TargetFilename.StartsWith(ProjectDir, StringComparison.InvariantCultureIgnoreCase))
 					{
@@ -70,17 +70,16 @@ namespace UnrealBuildTool
 			return bFoundTargetFiles;
 		}
 
-		public static bool FindTargetFiles(string InCurrentTopDirectory, ref bool bOutFoundTargetFiles)
+		public static bool FindTargetFiles(DirectoryReference CurrentTopDirectory, ref bool bOutFoundTargetFiles)
 		{
 			// We will only search as deep as the first target file found
-			string CurrentTopDirectory = InCurrentTopDirectory;
-			List<string> SubFolderList = new List<string>();
+			List<DirectoryReference> SubFolderList = new List<DirectoryReference>();
 
 			// Check the root directory
 			bOutFoundTargetFiles |= FindTargetFilesInFolder(CurrentTopDirectory);
 			if (bOutFoundTargetFiles == false)
 			{
-				foreach (var TargetFolder in Directory.EnumerateDirectories(CurrentTopDirectory, "*", SearchOption.TopDirectoryOnly))
+				foreach (var TargetFolder in Directory.EnumerateDirectories(CurrentTopDirectory.FullName, "*", SearchOption.TopDirectoryOnly).Select(x => new DirectoryReference(x)))
 				{
 					SubFolderList.Add(TargetFolder);
 					bOutFoundTargetFiles |= FindTargetFilesInFolder(TargetFolder);
@@ -105,39 +104,38 @@ namespace UnrealBuildTool
 		/// <summary>
 		/// Add a single project to the project info dictionary
 		/// </summary>
-		public static void AddProject(string ProjectFile)
+		public static void AddProject(FileReference ProjectFile)
 		{
-			string RelativePath = Utils.MakePathRelativeTo(ProjectFile, EngineSourceDirectory);
-			if(!ProjectInfoDictionary.ContainsKey(RelativePath))
+			if(!ProjectInfoDictionary.ContainsKey(ProjectFile))
 			{
-				string ProjectDirectory = Path.GetDirectoryName(ProjectFile);
+				DirectoryReference ProjectDirectory = ProjectFile.Directory;
 
 				// Check if it's a code project
-				string SourceFolder = Path.Combine(ProjectDirectory, "Source");
-				string IntermediateSourceFolder = Path.Combine(ProjectDirectory, "Intermediate", "Source");
-				bool bIsCodeProject = Directory.Exists(SourceFolder) || Directory.Exists(IntermediateSourceFolder);
+				DirectoryReference SourceFolder = DirectoryReference.Combine(ProjectDirectory, "Source");
+				DirectoryReference IntermediateSourceFolder = DirectoryReference.Combine(ProjectDirectory, "Intermediate", "Source");
+				bool bIsCodeProject = SourceFolder.Exists() || IntermediateSourceFolder.Exists();
 
 				// Create the project, and check the name is unique
-				UProjectInfo NewProjectInfo = new UProjectInfo(RelativePath, bIsCodeProject);
+				UProjectInfo NewProjectInfo = new UProjectInfo(ProjectFile, bIsCodeProject);
 				if(ShortProjectNameDictionary.ContainsKey(NewProjectInfo.GameName))
 				{
 					var FirstProject = ProjectInfoDictionary[ShortProjectNameDictionary[NewProjectInfo.GameName]];
-					throw new BuildException("There are multiple projects with name {0}\n\t* {1}\n\t* {2}\nThis is not currently supported.", NewProjectInfo.GameName, Path.GetFullPath(FirstProject.FilePath), Path.GetFullPath(NewProjectInfo.FilePath));
+					throw new BuildException("There are multiple projects with name {0}\n\t* {1}\n\t* {2}\nThis is not currently supported.", NewProjectInfo.GameName, FirstProject.FilePath.FullName, NewProjectInfo.FilePath.FullName);
 				}
 
 				// Add it to the name -> project lookups
-				ProjectInfoDictionary.Add(RelativePath, NewProjectInfo);
-				ShortProjectNameDictionary.Add(NewProjectInfo.GameName, RelativePath);
+				ProjectInfoDictionary.Add(ProjectFile, NewProjectInfo);
+				ShortProjectNameDictionary.Add(NewProjectInfo.GameName, ProjectFile);
 
 				// Find all Target.cs files if it's a code project
 				if(bIsCodeProject)
 				{
 					bool bFoundTargetFiles = false;
-					if (Directory.Exists(SourceFolder) && !FindTargetFiles(SourceFolder, ref bFoundTargetFiles))
+					if (SourceFolder.Exists() && !FindTargetFiles(SourceFolder, ref bFoundTargetFiles))
 					{
 						Log.TraceVerbose("No target files found under " + SourceFolder);
 					}
-					if (Directory.Exists(IntermediateSourceFolder) && !FindTargetFiles(IntermediateSourceFolder, ref bFoundTargetFiles))
+					if (IntermediateSourceFolder.Exists() && !FindTargetFiles(IntermediateSourceFolder, ref bFoundTargetFiles))
 					{
 						Log.TraceVerbose("No target files found under " + IntermediateSourceFolder);
 					}
@@ -200,7 +198,7 @@ namespace UnrealBuildTool
 						foreach (string UProjFile in SubDirFiles)
 						{
 							Log.TraceVerbose("\t\t\t\t{0}", UProjFile);
-							AddProject(UProjFile);
+							AddProject(new FileReference(UProjFile));
 						}
 					}
 				}
@@ -228,7 +226,7 @@ namespace UnrealBuildTool
 		{
 			Log.TraceInformation("Dumping project info...");
 			Log.TraceInformation("\tProjectInfo");
-			foreach (KeyValuePair<string, UProjectInfo> InfoEntry in ProjectInfoDictionary)
+			foreach (KeyValuePair<FileReference, UProjectInfo> InfoEntry in ProjectInfoDictionary)
 			{
 				Log.TraceInformation("\t\t" + InfoEntry.Key);
 				Log.TraceInformation("\t\t\tName          : " + InfoEntry.Value.FileName);
@@ -236,13 +234,13 @@ namespace UnrealBuildTool
 				Log.TraceInformation("\t\t\tCode Project  : " + (InfoEntry.Value.bIsCodeProject ? "YES" : "NO"));
 			}
 			Log.TraceInformation("\tShortName to Project");
-			foreach (KeyValuePair<string, string> ShortEntry in ShortProjectNameDictionary)
+			foreach (KeyValuePair<string, FileReference> ShortEntry in ShortProjectNameDictionary)
 			{
 				Log.TraceInformation("\t\tShort Name : " + ShortEntry.Key);
 				Log.TraceInformation("\t\tProject    : " + ShortEntry.Value);
 			}
 			Log.TraceInformation("\tTarget to Project");
-			foreach (KeyValuePair<string, string> TargetEntry in TargetToProjectDictionary)
+			foreach (KeyValuePair<string, FileReference> TargetEntry in TargetToProjectDictionary)
 			{
 				Log.TraceInformation("\t\tTarget     : " + TargetEntry.Key);
 				Log.TraceInformation("\t\tProject    : " + TargetEntry.Value);
@@ -257,7 +255,7 @@ namespace UnrealBuildTool
 		public static List<UProjectInfo> FilterGameProjects(bool bOnlyCodeProjects, string GameNameFilter)
 		{
 			List<UProjectInfo> Projects = new List<UProjectInfo>();
-			foreach (KeyValuePair<string, UProjectInfo> Entry in ProjectInfoDictionary)
+			foreach (KeyValuePair<FileReference, UProjectInfo> Entry in ProjectInfoDictionary)
 			{
 				if (!bOnlyCodeProjects || Entry.Value.bIsCodeProject)
 				{
@@ -274,30 +272,22 @@ namespace UnrealBuildTool
 		/// Get the project folder for the given target name
 		/// </summary>
 		/// <param name="InTargetName">Name of the target of interest</param>
-		/// <returns>The project filename, empty string if not found</returns>
-		public static string GetProjectForTarget(string InTargetName)
+		/// <param name="OutProjectFileName">The project filename</param>
+		/// <returns>True if the target was found</returns>
+		public static bool TryGetProjectForTarget(string InTargetName, out FileReference OutProjectFileName)
 		{
-			string ProjectName;
-			if (TargetToProjectDictionary.TryGetValue(InTargetName, out ProjectName) == true)
-			{
-				return ProjectName;
-			}
-			return "";
+			return TargetToProjectDictionary.TryGetValue(InTargetName, out OutProjectFileName);
 		}
 
 		/// <summary>
 		/// Get the project folder for the given project name
 		/// </summary>
 		/// <param name="InProjectName">Name of the project of interest</param>
-		/// <returns>The project filename, empty string if not found</returns>
-		public static string GetProjectFilePath(string InProjectName)
+		/// <param name="OutProjectFileName">The project filename</param>
+		/// <returns>True if the target was found</returns>
+		public static bool TryGetProjectFileName(string InProjectName, out FileReference OutProjectFileName)
 		{
-			string ProjectFilePath;
-			if (ShortProjectNameDictionary.TryGetValue(InProjectName, out ProjectFilePath))
-			{
-				return ProjectFilePath;
-			}
-			return "";
+			return ShortProjectNameDictionary.TryGetValue(InProjectName, out OutProjectFileName);
 		}
 
 		/// <summary>
diff --git a/Engine/Source/Programs/UnrealBuildTool/System/Unity.cs b/Engine/Source/Programs/UnrealBuildTool/System/Unity.cs
--- a/Engine/Source/Programs/UnrealBuildTool/System/Unity.cs
+++ b/Engine/Source/Programs/UnrealBuildTool/System/Unity.cs
@@ -164,8 +164,8 @@ namespace UnrealBuildTool
 				    foreach( var CPPFile in SortedCPPFiles )
 				    {
 					    // Don't include writable source files into unity blobs
-						if( !Utils.IsFileUnderDirectory( CPPFile.AbsolutePath, Target.EngineIntermediateDirectory ) &&
-							!Utils.IsFileUnderDirectory( CPPFile.AbsolutePath, Target.ProjectIntermediateDirectory ) )
+						if( !CPPFile.Reference.IsUnderDirectory(Target.EngineIntermediateDirectory ) &&
+							!CPPFile.Reference.IsUnderDirectory(Target.ProjectIntermediateDirectory ) )
 						{
 							++CandidateWorkingSetSourceFileCount;
 
@@ -241,7 +241,7 @@ namespace UnrealBuildTool
 			string PCHHeaderNameInCode = CPPFiles.Count > 0 ? CPPFiles[0].PCHHeaderNameInCode : "";
 			if( CompileEnvironment.Config.PrecompiledHeaderIncludeFilename != null )
 			{
-				PCHHeaderNameInCode = ToolChain.ConvertPath( CompileEnvironment.Config.PrecompiledHeaderIncludeFilename );
+				PCHHeaderNameInCode = ToolChain.ConvertPath( CompileEnvironment.Config.PrecompiledHeaderIncludeFilename.FullName );
 
 				// Generated unity .cpp files always include the PCH using an absolute path, so we need to update
 				// our compile environment's PCH header name to use this instead of the text it pulled from the original
@@ -289,16 +289,16 @@ namespace UnrealBuildTool
                 }
 
 				// Determine unity file path name
-				string UnityCPPFilePath;
+				string UnityCPPFileName;
 				if (AllUnityFiles.Count > 1)
 				{
-					UnityCPPFilePath = string.Format("Module.{0}.{1}_of_{2}.cpp", BaseName, CurrentUnityFileCount, AllUnityFiles.Count);
+					UnityCPPFileName = string.Format("Module.{0}.{1}_of_{2}.cpp", BaseName, CurrentUnityFileCount, AllUnityFiles.Count);
 				}
 				else
 				{
-					UnityCPPFilePath = string.Format("Module.{0}.cpp", BaseName);
+					UnityCPPFileName = string.Format("Module.{0}.cpp", BaseName);
 				}
-				UnityCPPFilePath = Path.Combine(CompileEnvironment.Config.OutputDirectory, UnityCPPFilePath);
+				FileReference UnityCPPFilePath = FileReference.Combine(CompileEnvironment.Config.OutputDirectory, UnityCPPFileName);
 
 				// Write the unity file to the intermediate folder.
 				FileItem UnityCPPFile = FileItem.CreateIntermediateTextFile(UnityCPPFilePath, OutputUnityCPPWriter.ToString());
diff --git a/Engine/Source/Programs/UnrealBuildTool/System/UnrealBuildTool.cs b/Engine/Source/Programs/UnrealBuildTool/System/UnrealBuildTool.cs
--- a/Engine/Source/Programs/UnrealBuildTool/System/UnrealBuildTool.cs
+++ b/Engine/Source/Programs/UnrealBuildTool/System/UnrealBuildTool.cs
@@ -57,13 +57,22 @@ namespace UnrealBuildTool
         static public bool bRunningRocket = false;
 
         /** The project file */
-        static string UProjectFile = "";
+        static FileReference UProjectFile;
 
         /** The Unreal project file path.  This should always be valid when compiling game projects.  Engine targets and program targets may not have a UProject set. */
         static bool bHasUProjectFile = false;
 
         /** The Unreal project directory.  This should always be valid when compiling game projects.  Engine targets and program targets may not have a UProject set. */
-        static string UProjectPath = "";
+        static DirectoryReference UProjectPath;
+
+		/** The full name of the Root UE4 directory */
+		public static readonly DirectoryReference RootDirectory = new DirectoryReference(Path.Combine(Path.GetDirectoryName(Assembly.GetExecutingAssembly().GetOriginalLocation()), "..", "..", ".."));
+
+		/** The full name of the Engine directory */
+		public static readonly DirectoryReference EngineDirectory = DirectoryReference.Combine(RootDirectory, "Engine");
+
+		/** The full name of the Engine/Source directory */
+		public static readonly DirectoryReference EngineSourceDirectory = DirectoryReference.Combine(EngineDirectory, "Source");
 
         /** This is set to true during UBT startup, after it is safe to be accessing the IsGatheringBuild and IsAssemblingBuild properties.  Never access this directly. */
         private static bool bIsSafeToCheckIfGatheringOrAssemblingBuild = false;
@@ -163,7 +172,7 @@ namespace UnrealBuildTool
         /// The Unreal project file path.  This should always be valid when compiling game projects.  Engine targets and program targets may not have a UProject set.
         /// </summary>
         /// <returns>The path to the file</returns>
-        static public string GetUProjectFile()
+        static public FileReference GetUProjectFile()
         {
             return UProjectFile;
         }
@@ -172,7 +181,7 @@ namespace UnrealBuildTool
         /// The Unreal project directory.  This should always be valid when compiling game projects.  Engine targets and program targets may not have a UProject set. */
         /// </summary>
         /// <returns>The directory path</returns>
-        static public string GetUProjectPath()
+        static public DirectoryReference GetUProjectPath()
         {
             return UProjectPath;
         }
@@ -181,12 +190,12 @@ namespace UnrealBuildTool
         /// The Unreal project source directory.  This should always be valid when compiling game projects.  Engine targets and program targets may not have a UProject set. */
         /// </summary>
         /// <returns>Relative path to the project's source directory</returns>
-        static public string GetUProjectSourcePath()
+        static public DirectoryReference GetUProjectSourcePath()
         {
-            string ProjectPath = GetUProjectPath();
+            DirectoryReference ProjectPath = GetUProjectPath();
             if (ProjectPath != null)
             {
-                ProjectPath = Path.Combine(ProjectPath, "Source");
+                ProjectPath = DirectoryReference.Combine(ProjectPath, "Source");
             }
             return ProjectPath;
         }
@@ -487,20 +496,20 @@ namespace UnrealBuildTool
             }
         }
 
-        public static void SetProjectFile(string InProjectFile)
+        public static void SetProjectFile(FileReference InProjectFile)
         {
-            if (HasUProjectFile() == false || String.IsNullOrEmpty(InProjectFile))
+            if (HasUProjectFile() == false || InProjectFile == null)
             {
                 UProjectFile = InProjectFile;
-				if (!String.IsNullOrEmpty(UProjectFile))
+				if (UProjectFile != null)
                 {
                     bHasUProjectFile = true;
-                    UProjectPath = Path.GetDirectoryName(UProjectFile);
+                    UProjectPath = UProjectFile.Directory;
                 }
                 else
                 {
                     bHasUProjectFile = false;
-					UProjectPath = String.Empty;
+					UProjectPath = null;
                 }
             }
             else
@@ -512,16 +521,16 @@ namespace UnrealBuildTool
 		public static void ResetProjectFile()
 		{
 			// Reset the project settings to their original state. Used by project file generator, since it mutates the active project setting.
-			UProjectFile = "";
+			UProjectFile = null;
 			bHasUProjectFile = false;
-			UProjectPath = "";
+			UProjectPath = null;
 		}
 
         private static bool ParseRocketCommandlineArg(string InArg, ref string OutGameName)
         {
             string LowercaseArg = InArg.ToLowerInvariant();
-            bool bSetupProject = false;
             
+			string ProjectArg = null;
             if (LowercaseArg == "-rocket")
             {
                 bRunningRocket = true;
@@ -530,31 +539,28 @@ namespace UnrealBuildTool
             {
 				if (BuildHostPlatform.Current.Platform == UnrealTargetPlatform.Linux) 
 				{
-					UProjectFile = InArg.Substring(9).Trim(new Char[] { ' ', '"', '\'' });
+					ProjectArg = InArg.Substring(9).Trim(new Char[] { ' ', '"', '\'' });
 				}
 				else 
 				{
-					UProjectFile = InArg.Substring(9);
+					ProjectArg = InArg.Substring(9);
 				}
-
-				bSetupProject = true;
             }
             else if (LowercaseArg.EndsWith(".uproject"))
             {
-                UProjectFile = InArg;
-                bSetupProject = true;
+                ProjectArg = InArg;
             }
             else
             {
                 return false;
             }
 
-            if (bSetupProject && (UProjectFile.Length > 0))
+            if (ProjectArg != null && ProjectArg.Length > 0)
             {
                 bHasUProjectFile = true;
     
-                string ProjectPath = Path.GetDirectoryName(UProjectFile);
-                OutGameName = Path.GetFileNameWithoutExtension(UProjectFile);
+                string ProjectPath = Path.GetDirectoryName(ProjectArg);
+                OutGameName = Path.GetFileNameWithoutExtension(ProjectArg);
 
                 if (Path.IsPathRooted(ProjectPath) == false)
                 {
@@ -568,14 +574,15 @@ namespace UnrealBuildTool
                             if (Directory.Exists(AlternativeProjectpath))
                             {
                                 ProjectPath = AlternativeProjectpath;
-                                UProjectFile = UProjectFile.Substring(3);
-                                Debug.Assert(UProjectFile.Length > 0);
+                                ProjectArg = ProjectArg.Substring(3);
+                                Debug.Assert(ProjectArg.Length > 0);
                             }
                         }
                     }
                 }
 
-                UProjectPath = Path.GetFullPath(ProjectPath);
+				UProjectFile = new FileReference(ProjectArg);
+                UProjectPath = UProjectFile.Directory;
             }
 
             return true;
@@ -1700,13 +1707,13 @@ namespace UnrealBuildTool
                     // If we're generating project files, then go ahead and wipe out the existing UBTMakefile for every target, to make sure that
                     // it gets a full dependency scan next time.
                     // NOTE: This is just a safeguard and doesn't have to be perfect.  We also check for newer project file timestamps in LoadUBTMakefile()
-                    string UBTMakefilePath = UnrealBuildTool.GetUBTMakefilePath( TargetDescs );
+                    FileReference UBTMakefilePath = UnrealBuildTool.GetUBTMakefilePath( TargetDescs );
                     if( ProjectFileGenerator.bGenerateProjectFiles )	// @todo ubtmake: This is only hit when generating IntelliSense for project files.  Probably should be done right inside ProjectFileGenerator.bat
                     {													// @todo ubtmake: Won't catch multi-target cases as GPF always builds one target at a time for Intellisense
                         // Delete the UBTMakefile
-                        if (File.Exists(UBTMakefilePath))
+                        if (UBTMakefilePath.Exists())
                         {
-                            UEBuildTarget.CleanFile(UBTMakefilePath);
+                            UEBuildTarget.CleanFile(UBTMakefilePath.FullName);
                         }
                     }
 
@@ -1937,7 +1944,7 @@ namespace UnrealBuildTool
                                     if( TargetUObjectModules.Count > 0 )
                                     {
                                         // Execute the header tool
-                                        string ModuleInfoFileName = Path.GetFullPath( Path.Combine( Target.ProjectIntermediateDirectory, "UnrealHeaderTool.manifest" ) );
+                                        FileReference ModuleInfoFileName = FileReference.Combine(Target.ProjectIntermediateDirectory, "UnrealHeaderTool.manifest");
                                         ECompilationResult UHTResult = ECompilationResult.OtherCompilationError;
                                         if (!ExternalExecution.ExecuteHeaderToolIfNecessary(Target, GlobalCompileEnvironment:null, UObjectModules:TargetUObjectModules, ModuleInfoFileName:ModuleInfoFileName, UHTResult:ref UHTResult))
                                         {
@@ -2116,11 +2123,10 @@ namespace UnrealBuildTool
 
 			foreach (var MakefileName in MakefileNames)
 			{
-				var MakefilePath = Path.Combine(BaseDir, MakefileName);
-
-				if (File.Exists(MakefilePath))
+				var MakefileRef = FileReference.Combine(BaseDir, MakefileName);
+				if (MakefileRef.Exists())
 				{
-					File.Delete(MakefilePath);
+					MakefileRef.Delete();
 				}
 			}
 		}
@@ -2147,17 +2153,17 @@ namespace UnrealBuildTool
 
 			if (!ProjectFileGenerator.bGenerateProjectFiles && !UEBuildConfiguration.bGenerateManifest && bIsEditorTarget )
 			{
-				var EditorProcessFilenames = UEBuildTarget.MakeExecutablePaths("..", "UE4Editor", TargetDesc.Platform, TargetDesc.Configuration, UnrealTargetConfiguration.Development, false, null);
+				var EditorProcessFilenames = UEBuildTarget.MakeExecutablePaths(UnrealBuildTool.EngineDirectory, "UE4Editor", TargetDesc.Platform, TargetDesc.Configuration, UnrealTargetConfiguration.Development, false, null);
 				if (EditorProcessFilenames.Count != 1)
 				{
 					throw new BuildException("ShouldDoHotReload cannot handle multiple binaries returning from UEBuildTarget.MakeExecutablePaths");
 				}
 
-				var EditorProcessFilename = EditorProcessFilenames[0];
+				var EditorProcessFilename = EditorProcessFilenames[0].FullName;
 				var EditorProcessName = Path.GetFileNameWithoutExtension(EditorProcessFilename);
 				var EditorProcesses = BuildHostPlatform.Current.GetProcessesByName(EditorProcessName);
 				var BinariesPath = Path.GetFullPath(Path.GetDirectoryName(EditorProcessFilename));
-				var PerProjectBinariesPath = Path.Combine(UnrealBuildTool.GetUProjectPath(), BinariesPath.Substring(BinariesPath.LastIndexOf("Binaries")));
+				var PerProjectBinariesPath = UnrealBuildTool.HasUProjectFile()? Path.Combine(UnrealBuildTool.GetUProjectPath().FullName, BinariesPath.Substring(BinariesPath.LastIndexOf("Binaries"))) : "";
 				bIsRunning = EditorProcesses.FirstOrDefault(EditorProc =>
 					{
 						if(!Path.GetFullPath(EditorProc.Filename).StartsWith(BinariesPath, StringComparison.InvariantCultureIgnoreCase))
@@ -2340,7 +2346,7 @@ namespace UnrealBuildTool
 
             var TimerStartTime = DateTime.UtcNow;
 
-            var UBTMakefileItem = FileItem.GetItemByFullPath( GetUBTMakefilePath( TargetDescs ) );
+            var UBTMakefileItem = FileItem.GetItemByFileReference( GetUBTMakefilePath( TargetDescs ) );
 
             // @todo ubtmake: Optimization: The UBTMakefile saved for game projects is upwards of 9 MB.  We should try to shrink its content if possible
             // @todo ubtmake: Optimization: C# Serialization may be too slow for these big Makefiles.  Loading these files often shows up as the slower part of the assembling phase.
@@ -2374,11 +2380,11 @@ namespace UnrealBuildTool
         /// <param name="MakefilePath">Path to the makefile to load</param>
 		/// <param name="ReasonNotLoaded">If the function returns null, this string will contain the reason why</param>
 		/// <returns>The loaded makefile, or null if it failed for some reason.  On failure, the 'ReasonNotLoaded' variable will contain information about why</returns>
-		static UBTMakefile LoadUBTMakefile( string MakefilePath, out string ReasonNotLoaded )
+		static UBTMakefile LoadUBTMakefile( FileReference MakefilePath, out string ReasonNotLoaded )
 		{
 			// Check the directory timestamp on the project files directory.  If the user has generated project files more
 			// recently than the UBTMakefile, then we need to consider the file to be out of date
-			var UBTMakefileInfo = new FileInfo( MakefilePath );
+			var UBTMakefileInfo = new FileInfo( MakefilePath.FullName );
 			if( !UBTMakefileInfo.Exists )
 			{
 				// UBTMakefile doesn't even exist, so we won't bother loading it
@@ -2393,7 +2399,7 @@ namespace UnrealBuildTool
 			// Rocket doesn't need to check engine projects for outdatedness
 			if( !UnrealBuildTool.RunningRocket() )
 			{
-				if( Directory.Exists( ProjectFileGenerator.IntermediateProjectFilesPath ) )
+				if( ProjectFileGenerator.IntermediateProjectFilesPath.Exists() )
 				{
 					var EngineProjectFilesLastUpdateTime = new FileInfo(ProjectFileGenerator.ProjectTimestampFile).LastWriteTime;
 					if( UBTMakefileInfo.LastWriteTime.CompareTo( EngineProjectFilesLastUpdateTime ) < 0 )
@@ -2409,7 +2415,7 @@ namespace UnrealBuildTool
 			// Check the game project directory too
 			if( UnrealBuildTool.HasUProjectFile() )
 			{
-				string ProjectFilename = UnrealBuildTool.GetUProjectFile();
+				string ProjectFilename = UnrealBuildTool.GetUProjectFile().FullName;
 				var ProjectFileInfo = new FileInfo(ProjectFilename);
 				if( !ProjectFileInfo.Exists || UBTMakefileInfo.LastWriteTime.CompareTo(ProjectFileInfo.LastWriteTime ) < 0 )
 				{
@@ -2420,7 +2426,7 @@ namespace UnrealBuildTool
 				}
 
 				var MasterProjectRelativePath = UnrealBuildTool.GetUProjectPath();
-				var GameIntermediateProjectFilesPath = Path.Combine( MasterProjectRelativePath, "Intermediate", "ProjectFiles" );
+				var GameIntermediateProjectFilesPath = Path.Combine( MasterProjectRelativePath.FullName, "Intermediate", "ProjectFiles" );
 				if( Directory.Exists( GameIntermediateProjectFilesPath ) )
 				{
 					var GameProjectFilesLastUpdateTime = new DirectoryInfo( GameIntermediateProjectFilesPath ).LastWriteTime;
@@ -2479,7 +2485,7 @@ namespace UnrealBuildTool
 			// Check if any of the target's Build.cs files are newer than the makefile
 			foreach (var Target in LoadedUBTMakefile.Targets)
 			{
-				string TargetCsFilename = Target.TargetCsFilename;
+				string TargetCsFilename = Target.TargetCsFilename.FullName;
 				if (TargetCsFilename != null)
 				{
 					var TargetCsFile = new FileInfo(TargetCsFilename);
@@ -2518,7 +2524,7 @@ namespace UnrealBuildTool
 				var HFilesNewerThanMakefile =
 					new HashSet<string>(
 						Target.FlatModuleCsData
-						.SelectMany(x => x.Value.ModuleSourceFolder != null ? Directory.EnumerateFiles(x.Value.ModuleSourceFolder, "*.h", SearchOption.AllDirectories) : Enumerable.Empty<string>())
+						.SelectMany(x => x.Value.ModuleSourceFolder != null ? Directory.EnumerateFiles(x.Value.ModuleSourceFolder.FullName, "*.h", SearchOption.AllDirectories) : Enumerable.Empty<string>())
 						.Where(y => Directory.GetLastWriteTimeUtc(y) > UBTMakefileInfo.LastWriteTimeUtc)
 						.OrderBy(z => z).Distinct()
 					);
@@ -2554,22 +2560,22 @@ namespace UnrealBuildTool
 				{
 					foreach( var FlatModule in Target.FlatModuleCsData.Values )
 					{
-						if( FlatModule != null && !String.IsNullOrEmpty( FlatModule.ModuleSourceFolder ) )
+						if( FlatModule != null && FlatModule.ModuleSourceFolder != null )
 						{
 							// Has the directory been touched?  Don't bother checking folders that haven't been written to since the
 							// Makefile was created. 
 							//if( Directory.GetLastWriteTimeUtc( FlatModule.ModuleSourceFolder ) > UBTMakefileInfo.LastWriteTimeUtc )
 							{
-								foreach( var SourceFilePath in Directory.EnumerateFiles( FlatModule.ModuleSourceFolder, "*.cpp", SearchOption.AllDirectories ) )
+								foreach( var SourceFilePath in FileReference.Enumerate( FlatModule.ModuleSourceFolder, "*.cpp", SearchOption.AllDirectories ) )
 								{
 									// Was the file written to since the Makefile was created?  We're not able to do a fully exhaustive
 									// check for working sets when in assembler mode, so we might get it wrong occasionally.  This is
 									// why we only want to check newly-modified or created files.  Otherwise, we could potentially
 									// invalidate the Makefile every single run.
-									if( File.GetLastWriteTimeUtc( SourceFilePath ) > UBTMakefileInfo.LastWriteTimeUtc )
+									if( File.GetLastWriteTimeUtc( SourceFilePath.FullName ) > UBTMakefileInfo.LastWriteTimeUtc )
 									{
-										bool bShouldBePartOfWorkingSet = UnrealBuildTool.ShouldSourceFileBePartOfWorkingSet( SourceFilePath );
-										bool bIsAlreadyPartOfWorkingSet = LoadedUBTMakefile.SourceFileWorkingSet.Contains( FileItem.GetItemByFullPath( SourceFilePath ) );
+										bool bShouldBePartOfWorkingSet = UnrealBuildTool.ShouldSourceFileBePartOfWorkingSet( SourceFilePath.FullName );
+										bool bIsAlreadyPartOfWorkingSet = LoadedUBTMakefile.SourceFileWorkingSet.Contains( FileItem.GetItemByFileReference( SourceFilePath ) );
 										if( bShouldBePartOfWorkingSet != bIsAlreadyPartOfWorkingSet )
 										{
 											Log.TraceVerbose("{0} {1} part of source working set and now {2}; invalidating makefile ({3})", SourceFilePath, bIsAlreadyPartOfWorkingSet ? "was" : "was not", bShouldBePartOfWorkingSet ? "is" : "is not", UBTMakefileInfo.FullName);
@@ -2594,9 +2600,9 @@ namespace UnrealBuildTool
         /// </summary>
         /// <param name="Targets">List of targets.  Order is not important</param>
         /// <returns>UBTMakefile path</returns>
-        public static string GetUBTMakefilePath( List<TargetDescriptor> TargetDescs )
+        public static FileReference GetUBTMakefilePath( List<TargetDescriptor> TargetDescs )
         {
-            string UBTMakefilePath;
+            FileReference UBTMakefilePath;
 
             if( TargetDescs.Count == 1 )
             {
@@ -2605,7 +2611,7 @@ namespace UnrealBuildTool
 				bool bIsHotReload = (UEBuildConfiguration.bHotReloadFromIDE || (TargetDesc.OnlyModules != null && TargetDesc.OnlyModules.Count > 0));
 				string UBTMakefileName = bIsHotReload ? "HotReloadMakefile.ubt" : "Makefile.ubt";
 
-				UBTMakefilePath = Path.Combine(GetUBTMakefileDirectoryPathForSingleTarget(TargetDesc), UBTMakefileName);
+				UBTMakefilePath = FileReference.Combine(GetUBTMakefileDirectoryPathForSingleTarget(TargetDesc), UBTMakefileName);
             }
             else
             {
@@ -2613,14 +2619,18 @@ namespace UnrealBuildTool
                 // configurations and platforms, and save it into the base intermediate folder
                 var TargetCollectionName = MakeTargetCollectionName( TargetDescs );
 
-                string ProjectIntermediatePath = BuildConfiguration.BaseIntermediatePath;
+                DirectoryReference ProjectIntermediatePath;
                 if (UnrealBuildTool.HasUProjectFile())
                 {
-                    ProjectIntermediatePath = Path.Combine(UnrealBuildTool.GetUProjectPath(), BuildConfiguration.BaseIntermediatePath);
+                    ProjectIntermediatePath = DirectoryReference.Combine(UnrealBuildTool.GetUProjectPath(), BuildConfiguration.BaseIntermediateFolder);
                 }
+				else
+				{
+					ProjectIntermediatePath = DirectoryReference.Combine(UnrealBuildTool.EngineDirectory, BuildConfiguration.BaseIntermediateFolder);
+				}
 
                 // @todo ubtmake: The TargetCollectionName string could be really long if there is more than one target!  Hash it?
-                UBTMakefilePath = Path.Combine(ProjectIntermediatePath, TargetCollectionName + ".ubt" );
+                UBTMakefilePath = FileReference.Combine(ProjectIntermediatePath, TargetCollectionName + ".ubt" );
             }
 
             return UBTMakefilePath;
@@ -2631,14 +2641,18 @@ namespace UnrealBuildTool
 		/// </summary>
 		/// <param name="Target">The target.</param>
 		/// <returns>UBTMakefile path</returns>
-		private static string GetUBTMakefileDirectoryPathForSingleTarget(TargetDescriptor Target)
+		private static DirectoryReference GetUBTMakefileDirectoryPathForSingleTarget(TargetDescriptor Target)
 		{
 			// If there's only one target, just save the UBTMakefile in the target's build intermediate directory
 			// under a folder for that target (and platform/config combo.)
-			string PlatformIntermediatePath = BuildConfiguration.PlatformIntermediatePath;
+			DirectoryReference PlatformIntermediatePath;
 			if (UnrealBuildTool.HasUProjectFile())
 			{
-				PlatformIntermediatePath = Path.Combine(UnrealBuildTool.GetUProjectPath(), BuildConfiguration.PlatformIntermediateFolder);
+				PlatformIntermediatePath = DirectoryReference.Combine(UnrealBuildTool.GetUProjectPath(), BuildConfiguration.PlatformIntermediateFolder);
+			}
+			else
+			{
+				PlatformIntermediatePath = DirectoryReference.Combine(UnrealBuildTool.EngineDirectory, BuildConfiguration.PlatformIntermediateFolder);
 			}
 
 			// @todo ubtmake: If this is a compile triggered from the editor it will have passed along the game's target name, not the editor target name.
@@ -2651,7 +2665,7 @@ namespace UnrealBuildTool
 				TargetName += "Editor";
 			}
 
-			return Path.Combine(PlatformIntermediatePath, TargetName, Target.Configuration.ToString());
+			return DirectoryReference.Combine(PlatformIntermediatePath, TargetName, Target.Configuration.ToString());
 		}
 
         /// <summary>
@@ -2703,10 +2717,10 @@ namespace UnrealBuildTool
 		public static void SetupUBTFromUAT(string UProjectFile)
 		{
 			// Reset project file
-			SetProjectFile(String.Empty);
+			SetProjectFile(null);
 			// when running UAT, the working directory is the root UE4 dir
 			BuildConfiguration.RelativeEnginePath = "Engine";
-			SetProjectFile(UProjectFile);
+			SetProjectFile(new FileReference(UProjectFile));
 		}
 
 
@@ -2830,7 +2844,7 @@ namespace UnrealBuildTool
 							if( OriginalProducedItemFilePath != NewProducedItemFilePath )
 							{
 								// OK, the produced item's file name changed so we'll update it to point to our new file
-								var NewProducedItem = FileItem.GetItemByFullPath( NewProducedItemFilePath );
+								var NewProducedItem = FileItem.GetItemByPath( NewProducedItemFilePath );
 								Action.ProducedItems[ ItemIndex ] = NewProducedItem;
 
 								// Copy the other important settings from the original file item
diff --git a/Engine/Source/Programs/UnrealBuildTool/System/Utils.cs b/Engine/Source/Programs/UnrealBuildTool/System/Utils.cs
--- a/Engine/Source/Programs/UnrealBuildTool/System/Utils.cs
+++ b/Engine/Source/Programs/UnrealBuildTool/System/Utils.cs
@@ -891,19 +891,19 @@ namespace UnrealBuildTool
 		/// <summary>
 		/// Returns the User Settings Directory path. This matches FPlatformProcess::UserSettingsDir()
 		/// </summary>
-		public static string GetUserSettingDirectory()
+		public static DirectoryReference GetUserSettingDirectory()
 		{
 			if (BuildHostPlatform.Current.Platform == UnrealTargetPlatform.Mac)
 			{
-				return Path.Combine(Environment.GetFolderPath(Environment.SpecialFolder.Personal), "Library", "Application Support", "Epic");
+				return new DirectoryReference(Path.Combine(Environment.GetFolderPath(Environment.SpecialFolder.Personal), "Library", "Application Support", "Epic"));
 			}
 			else if (Environment.OSVersion.Platform == PlatformID.Unix)
 			{
-				return Path.Combine(Environment.GetFolderPath(Environment.SpecialFolder.ApplicationData), "Epic");
+				return new DirectoryReference(Path.Combine(Environment.GetFolderPath(Environment.SpecialFolder.ApplicationData), "Epic"));
 			}
 			else
 			{
-				return Environment.GetFolderPath(Environment.SpecialFolder.LocalApplicationData);
+				return new DirectoryReference(Environment.GetFolderPath(Environment.SpecialFolder.LocalApplicationData));
 			}
 		}
 
diff --git a/Engine/Source/Programs/UnrealBuildTool/System/VCProject.cs b/Engine/Source/Programs/UnrealBuildTool/System/VCProject.cs
--- a/Engine/Source/Programs/UnrealBuildTool/System/VCProject.cs
+++ b/Engine/Source/Programs/UnrealBuildTool/System/VCProject.cs
@@ -45,14 +45,14 @@ namespace UnrealBuildTool
 		/// Constructs a new project file object
 		/// </summary>
 		/// <param name="InitFilePath">The path to the project file on disk</param>
-		public MSBuildProjectFile(string InitFilePath)
+		public MSBuildProjectFile(FileReference InitFilePath)
 			: base(InitFilePath)
 		{
 			// Each project gets its own GUID.  This is stored in the project file and referenced in the solution file.
 
 			// First, check to see if we have an existing file on disk.  If we do, then we'll try to preserve the
 			// GUID by loading it from the existing file.
-			if (File.Exists(ProjectFilePath))
+			if (ProjectFilePath.Exists())
 			{
 				try
 				{
@@ -82,7 +82,7 @@ namespace UnrealBuildTool
 			if( ProjectFileGenerator.bGenerateProjectFiles )
 			{
 				var Doc = new XmlDocument();
-				Doc.Load( ProjectFilePath );
+				Doc.Load( ProjectFilePath.FullName );
 
 				// @todo projectfiles: Ideally we could do a better job about preserving GUIDs when only minor changes are made
 				// to the project (such as adding a single new file.) It would make diffing changes much easier!
@@ -278,7 +278,7 @@ namespace UnrealBuildTool
 		/// Constructs a new project file object
 		/// </summary>
 		/// <param name="InitFilePath">The path to the project file on disk</param>
-		public VCProjectFile(string InitFilePath)
+		public VCProjectFile(FileReference InitFilePath)
 			: base( InitFilePath )
 		{
 		}
@@ -315,7 +315,7 @@ namespace UnrealBuildTool
 		/// Implements Project interface
 		public override bool WriteProjectFile(List<UnrealTargetPlatform> InPlatforms, List<UnrealTargetConfiguration> InConfigurations)
 		{
-			string ProjectName = Path.GetFileNameWithoutExtension(ProjectFilePath);
+			string ProjectName = ProjectFilePath.GetFileNameWithoutExtension();
 
 			bool bSuccess = true;
 
@@ -541,7 +541,7 @@ namespace UnrealBuildTool
 				{
 					// We want all source file and directory paths in the project files to be relative to the project file's
 					// location on the disk.  Convert the path to be relative to the project file directory
-					var ProjectRelativeSourceFile = Utils.MakePathRelativeTo(CurFile.FilePath, Path.GetDirectoryName(ProjectFilePath));
+					var ProjectRelativeSourceFile = CurFile.Reference.MakeRelativeTo(ProjectFilePath.Directory);
 
 					// By default, files will appear relative to the project file in the solution.  This is kind of the normal Visual
 					// Studio way to do things, but because our generated project files are emitted to intermediate folders, if we always
@@ -549,9 +549,9 @@ namespace UnrealBuildTool
 					string FilterRelativeSourceDirectory = Path.GetDirectoryName(ProjectRelativeSourceFile);
 
 					// Use the specified relative base folder
-					if (CurFile.RelativeBaseFolder != null)	// NOTE: We are looking for null strings, not empty strings!
+					if (CurFile.BaseFolder != null)
 					{
-						FilterRelativeSourceDirectory = Path.GetDirectoryName(Utils.MakePathRelativeTo(CurFile.FilePath, CurFile.RelativeBaseFolder));
+						FilterRelativeSourceDirectory = Path.GetDirectoryName(CurFile.Reference.MakeRelativeTo(CurFile.BaseFolder));
 					}
 
 					LocalAliasedFiles.Add(new AliasedFile(ProjectRelativeSourceFile, FilterRelativeSourceDirectory));
@@ -660,7 +660,7 @@ namespace UnrealBuildTool
 						//    about a target type, but the project file may contain many types of targets!  Some of this logic will need to move into
 						//    the per-target configuration writing code.
 						var HackTargetType = TargetRules.TargetType.Game;
-						string HackTargetFilePath = null;
+						FileReference HackTargetFilePath = null;
 						foreach( var Combination in ProjectConfigAndTargetCombinations )
 						{
 							if( Combination.Platform == Platform &&
@@ -672,7 +672,7 @@ namespace UnrealBuildTool
 							}
 						}
 
-						if( !String.IsNullOrEmpty( HackTargetFilePath ) )
+						if( HackTargetFilePath != null )
 						{
 							OutputManifestString += ProjGenerator.GetVisualStudioOutputManifestSection(Platform, HackTargetType, HackTargetFilePath, ProjectFilePath);
 						}
@@ -703,7 +703,7 @@ namespace UnrealBuildTool
 			// Save the project file
 			if( bSuccess )
 			{
-				bSuccess = ProjectFileGenerator.WriteFileIfChanged( ProjectFilePath, VCProjectFileContent.ToString() );
+				bSuccess = ProjectFileGenerator.WriteFileIfChanged( ProjectFilePath.FullName, VCProjectFileContent.ToString() );
 			}
 
 
@@ -711,7 +711,7 @@ namespace UnrealBuildTool
 			if( bSuccess )
 			{
 				// Create a path to the project file's filters file
-				var VCFiltersFilePath = ProjectFilePath + ".filters";
+				var VCFiltersFilePath = ProjectFilePath.FullName + ".filters";
 				if( FiltersFileIsNeeded )
 				{
 					bSuccess = ProjectFileGenerator.WriteFileIfChanged( VCFiltersFilePath, VCFiltersFileContent.ToString() );
@@ -736,7 +736,7 @@ namespace UnrealBuildTool
 			if (VCUserFileContent.Length > 0)
 			{
 				// Create a path to the project file's user file
-				var VCUserFilePath = ProjectFilePath + ".user";
+				var VCUserFilePath = ProjectFilePath.FullName + ".user";
 				// Never overwrite the existing user path as it will cause them to lose their settings
 				if (File.Exists(VCUserFilePath) == false)
 				{
@@ -895,6 +895,8 @@ namespace UnrealBuildTool
 					"		<Import Project=\"$(UserRootDir)\\Microsoft.Cpp.$(Platform).user.props\" Condition=\"exists('$(UserRootDir)\\Microsoft.Cpp.$(Platform).user.props')\" Label=\"LocalAppDataPlatform\" />" + ProjectFileGenerator.NewLine +
 					"	</ImportGroup>" + ProjectFileGenerator.NewLine);
 
+				DirectoryReference ProjectDirectory = ProjectFilePath.Directory;
+
 				if (IsStubProject)
 				{
 					string ProjectRelativeUnusedDirectory = NormalizeProjectPath(Path.Combine(ProjectFileGenerator.EngineRelativePath, BuildConfiguration.BaseIntermediateFolder, "Unused"));
@@ -924,13 +926,13 @@ namespace UnrealBuildTool
 						"		<NMakeReBuildCommandLine>@echo {0} is not a supported platform for {1}. Valid platforms are {2}.</NMakeReBuildCommandLine>" + ProjectFileGenerator.NewLine +
 						"		<NMakeCleanCommandLine>@echo {0} is not a supported platform for {1}. Valid platforms are {2}.</NMakeCleanCommandLine>" + ProjectFileGenerator.NewLine +
 						"		<NMakeOutput/>" + ProjectFileGenerator.NewLine +
-						"	</PropertyGroup>" + ProjectFileGenerator.NewLine, Combination.Platform, Utils.GetFilenameWithoutAnyExtensions(Combination.ProjectTarget.TargetFilePath), String.Join(", ", SupportedPlatforms.Select(x => x.ToString())));
+						"	</PropertyGroup>" + ProjectFileGenerator.NewLine, Combination.Platform, Combination.ProjectTarget.TargetFilePath.GetFileNameWithoutAnyExtensions(), String.Join(", ", SupportedPlatforms.Select(x => x.ToString())));
 				}
 				else
 				{
 					TargetRules TargetRulesObject = Combination.ProjectTarget.TargetRules;
-					string TargetFilePath = Combination.ProjectTarget.TargetFilePath;
-					string TargetName = Utils.GetFilenameWithoutAnyExtensions(TargetFilePath);
+					FileReference TargetFilePath = Combination.ProjectTarget.TargetFilePath;
+					string TargetName = TargetFilePath.GetFileNameWithoutAnyExtensions();
 					var UBTPlatformName = IsStubProject ? StubProjectPlatformName : Platform.ToString();
 					var UBTConfigurationName = IsStubProject ? StubProjectConfigurationName : Configuration.ToString();
 
@@ -942,35 +944,34 @@ namespace UnrealBuildTool
 					bShouldCompileMonolithic |= TargetRulesObject.ShouldCompileMonolithic(Platform, Configuration);
 
 					// Get the output directory
-					string EngineRootDirectory = Path.GetFullPath(ProjectFileGenerator.EngineRelativePath);
-					string RootDirectory = EngineRootDirectory;
+					DirectoryReference RootDirectory = UnrealBuildTool.EngineDirectory;
 					if ((TargetRules.IsAGame(TargetRulesObject.Type) || TargetRulesObject.Type == TargetRules.TargetType.Server) && bShouldCompileMonolithic && !TargetRulesObject.bOutputToEngineBinaries)
 					{
-						if (UnrealBuildTool.HasUProjectFile() && Utils.IsFileUnderDirectory(TargetFilePath, UnrealBuildTool.GetUProjectPath()))
+						if (UnrealBuildTool.HasUProjectFile() && TargetFilePath.IsUnderDirectory(UnrealBuildTool.GetUProjectPath()))
 						{
-							RootDirectory = Path.GetFullPath(UnrealBuildTool.GetUProjectPath());
+							RootDirectory = UnrealBuildTool.GetUProjectPath();
 						}
 						else
 						{
-							string UnrealProjectPath = UProjectInfo.GetProjectFilePath(ProjectName);
-							if (!String.IsNullOrEmpty(UnrealProjectPath))
+							FileReference ProjectFileName;
+							if(UProjectInfo.TryGetProjectFileName(ProjectName, out ProjectFileName))
 							{
-								RootDirectory = Path.GetDirectoryName(Path.GetFullPath(UnrealProjectPath));
+								RootDirectory = ProjectFileName.Directory;
 							}
 						}
 					}
 
 					if(TargetRulesObject.Type == TargetRules.TargetType.Program && !TargetRulesObject.bOutputToEngineBinaries)
 					{
-						string UnrealProjectPath = UProjectInfo.GetProjectForTarget(TargetName);
-						if (!String.IsNullOrEmpty(UnrealProjectPath))
+						FileReference ProjectFileName;
+						if(UProjectInfo.TryGetProjectForTarget(TargetName, out ProjectFileName))
 						{
-							RootDirectory = Path.GetDirectoryName(Path.GetFullPath(UnrealProjectPath));
+							RootDirectory = ProjectFileName.Directory;
 						}
 					}
 
 					// Get the output directory
-					string OutputDirectory = Path.Combine(RootDirectory, "Binaries", UBTPlatformName);
+					DirectoryReference OutputDirectory = DirectoryReference.Combine(RootDirectory, "Binaries", UBTPlatformName);
 
 					// Get the executable name (minus any platform or config suffixes)
 					string BaseExeName = TargetName;
@@ -985,7 +986,7 @@ namespace UnrealBuildTool
 					}
 
 					// Make the output file path
-					string NMakePath = Path.Combine(OutputDirectory, BaseExeName);
+					FileReference NMakePath = FileReference.Combine(OutputDirectory, BaseExeName);
 					if (Configuration != UnrealTargetConfiguration.Development && (Configuration != UnrealTargetConfiguration.DebugGame || bShouldCompileMonolithic))
 					{
 						NMakePath += "-" + UBTPlatformName + "-" + UBTConfigurationName;
@@ -1031,7 +1032,7 @@ namespace UnrealBuildTool
 					{
 						BuildArguments += " " + UProjectPath + (UnrealBuildTool.RunningRocket() ? " -rocket" : "");
 					}
-					string BatchFilesDirectoryName = Path.Combine(ProjectFileGenerator.EngineRelativePath, "Build", "BatchFiles");
+					DirectoryReference BatchFilesDirectory = DirectoryReference.Combine(UnrealBuildTool.EngineDirectory, "Build", "BatchFiles");
 
 					// @todo UWP: For the MS toolchains, if an override was set for project generation, push that into the build strings to override the build toolchain as well
 					string BuildToolOverride = "";
@@ -1051,21 +1052,21 @@ namespace UnrealBuildTool
 
 					// NMake Build command line
 					VCProjectFileContent.Append("		<NMakeBuildCommandLine>");
-					VCProjectFileContent.Append(EscapePath(NormalizeProjectPath(Path.Combine(BatchFilesDirectoryName, "Build.bat"))) + BuildArguments.ToString());
+					VCProjectFileContent.Append(EscapePath(NormalizeProjectPath(FileReference.Combine(BatchFilesDirectory, "Build.bat"))) + BuildArguments.ToString());
 					VCProjectFileContent.Append("</NMakeBuildCommandLine>" + ProjectFileGenerator.NewLine);
 
 					// NMake ReBuild command line
 					VCProjectFileContent.Append("		<NMakeReBuildCommandLine>");
-					VCProjectFileContent.Append(EscapePath(NormalizeProjectPath(Path.Combine(BatchFilesDirectoryName, "Rebuild.bat"))) + BuildArguments.ToString());
+					VCProjectFileContent.Append(EscapePath(NormalizeProjectPath(FileReference.Combine(BatchFilesDirectory, "Rebuild.bat"))) + BuildArguments.ToString());
 					VCProjectFileContent.Append("</NMakeReBuildCommandLine>" + ProjectFileGenerator.NewLine);
 
 					// NMake Clean command line
 					VCProjectFileContent.Append("		<NMakeCleanCommandLine>");
-					VCProjectFileContent.Append(EscapePath(NormalizeProjectPath(Path.Combine(BatchFilesDirectoryName, "Clean.bat"))) + BuildArguments.ToString());
+					VCProjectFileContent.Append(EscapePath(NormalizeProjectPath(FileReference.Combine(BatchFilesDirectory, "Clean.bat"))) + BuildArguments.ToString());
 					VCProjectFileContent.Append("</NMakeCleanCommandLine>" + ProjectFileGenerator.NewLine);
 
 					VCProjectFileContent.Append("		<NMakeOutput>");
-					VCProjectFileContent.Append(NormalizeProjectPath(NMakePath));
+					VCProjectFileContent.Append(NormalizeProjectPath(NMakePath.FullName));
 					VCProjectFileContent.Append("</NMakeOutput>" + ProjectFileGenerator.NewLine);
 					VCProjectFileContent.Append("	</PropertyGroup>" + ProjectFileGenerator.NewLine);
 
@@ -1136,7 +1137,7 @@ namespace UnrealBuildTool
 		/// Constructs a new project file object
 		/// </summary>
 		/// <param name="InitFilePath">The path to the project file on disk</param>
-		public VCSharpProjectFile( string InitFilePath )
+		public VCSharpProjectFile( FileReference InitFilePath )
 			: base( InitFilePath )
 		{
 		}
@@ -1147,7 +1148,7 @@ namespace UnrealBuildTool
 		{
 			var RelativeFilePaths = new List<string>();
 			var Doc = new XmlDocument();
-			Doc.Load( ProjectFilePath );
+			Doc.Load( ProjectFilePath.FullName );
 
 			var Tags = new string[] { "Compile", "Page", "Resource" };
 			foreach( var Tag in Tags )
@@ -1180,7 +1181,7 @@ namespace UnrealBuildTool
 		 *
 		 * @param	AssemblyReference	The full path to the assembly file on disk
 		 */
-		public void AddAssemblyReference( string AssemblyReference )
+		public void AddAssemblyReference( FileReference AssemblyReference )
 		{
 			AssemblyReferences.Add( AssemblyReference );
 		}
@@ -1242,8 +1243,8 @@ namespace UnrealBuildTool
 				ProjectFileContent.Append( "<ItemGroup>" + ProjectFileGenerator.NewLine );
 				foreach( var CurReference in AssemblyReferences )
 				{
-					ProjectFileContent.Append( "\t<Reference Include=\"" + EscapeFileName(Path.GetFileNameWithoutExtension( CurReference )) + "\" >" + ProjectFileGenerator.NewLine );
-					ProjectFileContent.Append( "\t\t<HintPath>" + Utils.MakePathRelativeTo( CurReference, Path.GetDirectoryName(ProjectFilePath) ) + "</HintPath>" + ProjectFileGenerator.NewLine );
+					ProjectFileContent.Append( "\t<Reference Include=\"" + EscapeFileName(CurReference.GetFileNameWithoutExtension()) + "\" >" + ProjectFileGenerator.NewLine );
+					ProjectFileContent.Append( "\t\t<HintPath>" + CurReference.MakeRelativeTo(ProjectFilePath.Directory) + "</HintPath>" + ProjectFileGenerator.NewLine );
 					ProjectFileContent.Append( "\t</Reference>" + ProjectFileGenerator.NewLine );
 				}
 				ProjectFileContent.Append( "</ItemGroup>" + ProjectFileGenerator.NewLine );
@@ -1252,8 +1253,8 @@ namespace UnrealBuildTool
 			// Other references (note it's assumed all references here are at least of MSBuildProjectFile type.
 			foreach (var Project in DependsOnProjects)
 			{
-				var RelativePath = Utils.MakePathRelativeTo(Path.GetDirectoryName(Project.ProjectFilePath), Path.GetDirectoryName(ProjectFilePath));
-				RelativePath = Path.Combine(RelativePath, Path.GetFileName(Project.ProjectFilePath));
+				var RelativePath = Project.ProjectFilePath.Directory.MakeRelativeTo(ProjectFilePath.Directory);
+				RelativePath = Path.Combine(RelativePath, Project.ProjectFilePath.GetFileName());
 				ProjectFileContent.Append(
 					"<ItemGroup>" + ProjectFileGenerator.NewLine +
 						"<ProjectReference Include=\"" + EscapeFileName(RelativePath) + "\">" + ProjectFileGenerator.NewLine +
@@ -1269,7 +1270,7 @@ namespace UnrealBuildTool
 			// Add all files to the project.
 			foreach( var CurFile in SourceFiles )
 			{
-				var ProjectRelativeSourceFile = Utils.MakePathRelativeTo( CurFile.FilePath, Path.GetDirectoryName( ProjectFilePath ) );
+				var ProjectRelativeSourceFile = CurFile.Reference.MakeRelativeTo( ProjectFilePath.Directory );
 				ProjectFileContent.Append(
 					"		<Compile Include=\"" + EscapeFileName(ProjectRelativeSourceFile) + "\" />" + ProjectFileGenerator.NewLine);
 			}
@@ -1285,14 +1286,14 @@ namespace UnrealBuildTool
 			// Save the project file
 			if (bSuccess)
 			{
-				bSuccess = ProjectFileGenerator.WriteFileIfChanged(ProjectFilePath, ProjectFileContent.ToString());
+				bSuccess = ProjectFileGenerator.WriteFileIfChanged(ProjectFilePath.FullName, ProjectFileContent.ToString());
 			}
 
 			return bSuccess;
 		}
 
 		/// Assemblies this project is dependent on
-		protected readonly List<string> AssemblyReferences = new List<string>();
+		protected readonly List<FileReference> AssemblyReferences = new List<FileReference>();
 		/// System assemblies this project is dependent on
 		protected readonly List<string> DotNetAssemblyReferences = new List<string>() { "System", "System.Core", "System.Data", "System.Xml" };
 	}
@@ -1313,7 +1314,7 @@ namespace UnrealBuildTool
 		/// Constructs a new project file object
 		/// </summary>
 		/// <param name="InitFilePath">The path to the project file on disk</param>
-		public VSHFBProjectFile( string InitFilePath )
+		public VSHFBProjectFile( FileReference InitFilePath )
 			: base( InitFilePath )
 		{
 		}
diff --git a/Engine/Source/Programs/UnrealBuildTool/System/VCProjectFileGenerator.cs b/Engine/Source/Programs/UnrealBuildTool/System/VCProjectFileGenerator.cs
--- a/Engine/Source/Programs/UnrealBuildTool/System/VCProjectFileGenerator.cs
+++ b/Engine/Source/Programs/UnrealBuildTool/System/VCProjectFileGenerator.cs
@@ -51,45 +51,45 @@ namespace UnrealBuildTool
 
 		/// <summary>
 		/// </summary>
-		public override void CleanProjectFiles(string InMasterProjectRelativePath, string InMasterProjectName, string InIntermediateProjectFilesPath)
+		public override void CleanProjectFiles(DirectoryReference InMasterProjectDirectory, string InMasterProjectName, DirectoryReference InIntermediateProjectFilesDirectory)
 		{
-			string MasterProjectFile = Path.Combine(InMasterProjectRelativePath, InMasterProjectName);
-			string MasterProjDeleteFilename = MasterProjectFile + ".sln";
-			if (File.Exists(MasterProjDeleteFilename))
+			FileReference MasterProjectFile = FileReference.Combine(InMasterProjectDirectory, InMasterProjectName);
+			FileReference MasterProjDeleteFilename = MasterProjectFile + ".sln";
+			if (MasterProjDeleteFilename.Exists())
 			{
-				File.Delete(MasterProjDeleteFilename);
+				MasterProjDeleteFilename.Delete();
 			}
 			MasterProjDeleteFilename = MasterProjectFile + ".sdf";
-			if (File.Exists(MasterProjDeleteFilename))
+			if (MasterProjDeleteFilename.Exists())
 			{
-				File.Delete(MasterProjDeleteFilename);
+				MasterProjDeleteFilename.Delete();
 			}
 			MasterProjDeleteFilename = MasterProjectFile + ".suo";
-			if (File.Exists(MasterProjDeleteFilename))
+			if (MasterProjDeleteFilename.Exists())
 			{
-				File.Delete(MasterProjDeleteFilename);
+				MasterProjDeleteFilename.Delete();
 			}
 			MasterProjDeleteFilename = MasterProjectFile + ".v11.suo";
-			if (File.Exists(MasterProjDeleteFilename))
+			if (MasterProjDeleteFilename.Exists())
 			{
-				File.Delete(MasterProjDeleteFilename);
+				MasterProjDeleteFilename.Delete();
 			}
 			MasterProjDeleteFilename = MasterProjectFile + ".v12.suo";
-			if (File.Exists(MasterProjDeleteFilename))
+			if (MasterProjDeleteFilename.Exists())
 			{
-				File.Delete(MasterProjDeleteFilename);
+				MasterProjDeleteFilename.Delete();
 			}
 
 			// Delete the project files folder
-			if (Directory.Exists(InIntermediateProjectFilesPath))
+			if (InIntermediateProjectFilesDirectory.Exists())
 			{
 				try
 				{
-					Directory.Delete(InIntermediateProjectFilesPath, true);
+					Directory.Delete(InIntermediateProjectFilesDirectory.FullName, true);
 				}
 				catch (Exception Ex)
 				{
-					Log.TraceInformation("Error while trying to clean project files path {0}. Ignored.", InIntermediateProjectFilesPath);
+					Log.TraceInformation("Error while trying to clean project files path {0}. Ignored.", InIntermediateProjectFilesDirectory);
 					Log.TraceInformation("\t" + Ex.Message);
 				}
 			}
@@ -100,7 +100,7 @@ namespace UnrealBuildTool
 		/// </summary>
 		/// <param name="InitFilePath">Path to the project file</param>
 		/// <returns>The newly allocated project file object</returns>
-		protected override ProjectFile AllocateProjectFile( string InitFilePath )
+		protected override ProjectFile AllocateProjectFile( FileReference InitFilePath )
 		{
 			return new VCProjectFile( InitFilePath );
 		}
@@ -368,8 +368,8 @@ namespace UnrealBuildTool
 			{
 				if(Project.ProjectTargets.Count == 1)
 				{
-					string TargetFilePath = Project.ProjectTargets[0].TargetFilePath;
-					if(TargetFilePath != null && Path.GetFileNameWithoutExtension(Path.GetFileNameWithoutExtension(TargetFilePath)).Equals("ShaderCompileWorker", StringComparison.InvariantCultureIgnoreCase))
+					FileReference TargetFilePath = Project.ProjectTargets[0].TargetFilePath;
+					if(TargetFilePath != null && TargetFilePath.GetFileNameWithoutAnyExtensions().Equals("ShaderCompileWorker", StringComparison.InvariantCultureIgnoreCase))
 					{
 						ShaderCompileWorkerProject = Project;
 						break;
@@ -432,13 +432,13 @@ namespace UnrealBuildTool
 
 					// NOTE: The project name in the solution doesn't actually *have* to match the project file name on disk.  However,
 					//       we prefer it when it does match so we use the actual file name here.
-					var ProjectNameInSolution = Path.GetFileNameWithoutExtension(CurProject.ProjectFilePath);
+					var ProjectNameInSolution = CurProject.ProjectFilePath.GetFileNameWithoutExtension();
 
 					// Use the existing project's GUID that's already known to us
 					var ProjectGUID = CurProject.ProjectGUID.ToString( "B" ).ToUpperInvariant();
 
 					VCSolutionFileContent.Append(
-							"Project(\"" + ProjectTypeGUID + "\") = \"" + ProjectNameInSolution + "\", \"" + CurProject.RelativeProjectFilePath + "\", \"" + ProjectGUID + "\"" + ProjectFileGenerator.NewLine );
+							"Project(\"" + ProjectTypeGUID + "\") = \"" + ProjectNameInSolution + "\", \"" + CurProject.ProjectFilePath.MakeRelativeTo(ProjectFileGenerator.MasterProjectPath) + "\", \"" + ProjectGUID + "\"" + ProjectFileGenerator.NewLine );
 
 					// Setup dependency on UnrealBuildTool, if we need that.  This makes sure that UnrealBuildTool is
 					// freshly compiled before kicking off any build operations on this target project
@@ -839,7 +839,7 @@ namespace UnrealBuildTool
 			// Save the solution file
 			if( bSuccess )
 			{
-				var SolutionFilePath = Path.Combine( MasterProjectRelativePath, SolutionFileName );
+				var SolutionFilePath = FileReference.Combine( MasterProjectPath, SolutionFileName ).FullName;
 				bSuccess = WriteFileIfChanged( SolutionFilePath, VCSolutionFileContent.ToString() );
 			}
 
@@ -863,11 +863,11 @@ namespace UnrealBuildTool
 				}
 
 				// Check it doesn't exist before overwriting it. Since these files store the user's preferences, it'd be bad form to overwrite them.
-				string SolutionOptionsFileName = Path.Combine(MasterProjectRelativePath, Path.ChangeExtension(SolutionFileName, SolutionOptionsExtension));
+				string SolutionOptionsFileName = Path.Combine(MasterProjectPath.FullName, Path.ChangeExtension(SolutionFileName, SolutionOptionsExtension));
 				if (ProjectFileFormat == VCProjectFileFormat.VisualStudio2015)
 				{
-					SolutionOptionsFileName = Path.Combine(MasterProjectRelativePath, ".vs", Path.GetFileNameWithoutExtension(SolutionFileName), "v14", ".suo");
-					Directory.CreateDirectory(Path.Combine(MasterProjectRelativePath, ".vs", Path.GetFileNameWithoutExtension(SolutionFileName), "v14"));
+					SolutionOptionsFileName = Path.Combine(MasterProjectPath.FullName, ".vs", Path.GetFileNameWithoutExtension(SolutionFileName), "v14", ".suo");
+					Directory.CreateDirectory(Path.Combine(MasterProjectPath.FullName, ".vs", Path.GetFileNameWithoutExtension(SolutionFileName), "v14"));
 				}
 				if(!File.Exists(SolutionOptionsFileName))
 				{
@@ -890,7 +890,7 @@ namespace UnrealBuildTool
 					VCSolutionExplorerState ExplorerState = new VCSolutionExplorerState();
 					foreach(ProjectFile ProjectFile in AllProjectFiles)
 					{
-						string ProjectName = Path.GetFileNameWithoutExtension(ProjectFile.ProjectFilePath);
+						string ProjectName = ProjectFile.ProjectFilePath.GetFileNameWithoutExtension();
 						if(ProjectFile == DefaultProject)
 						{
 							ExplorerState.OpenProjects.Add(new Tuple<string, string[]>(ProjectName, new string[]{ ProjectName }));
diff --git a/Engine/Source/Programs/UnrealBuildTool/System/XcodeProject.cs b/Engine/Source/Programs/UnrealBuildTool/System/XcodeProject.cs
--- a/Engine/Source/Programs/UnrealBuildTool/System/XcodeProject.cs
+++ b/Engine/Source/Programs/UnrealBuildTool/System/XcodeProject.cs
@@ -19,7 +19,7 @@ namespace UnrealBuildTool
 		/// <summary>
 		/// Constructor
 		/// </summary>
-		public XcodeSourceFile(string InitFilePath, string InitRelativeBaseFolder)
+		public XcodeSourceFile(FileReference InitFilePath, DirectoryReference InitRelativeBaseFolder)
 			: base(InitFilePath, InitRelativeBaseFolder)
 		{
 			FileGUID = XcodeProjectFileGenerator.MakeXcodeGuid();
@@ -91,7 +91,7 @@ namespace UnrealBuildTool
 
 			foreach (XcodeSourceFile File in Files)
 			{
-				Contents.Append(string.Format("\t\t\t\t{0} /* {1} */,{2}", File.FileRefGUID, Path.GetFileName(File.FilePath), ProjectFileGenerator.NewLine));
+				Contents.Append(string.Format("\t\t\t\t{0} /* {1} */,{2}", File.FileRefGUID, File.Reference.GetFileName(), ProjectFileGenerator.NewLine));
 			}
 
 			if (!bFilesOnly)
@@ -126,7 +126,7 @@ namespace UnrealBuildTool
 		/// Constructs a new project file object
 		/// </summary>
 		/// <param name="InitFilePath">The path to the project file on disk</param>
-		public XcodeProjectFile(string InitFilePath)
+		public XcodeProjectFile(FileReference InitFilePath)
 			: base( InitFilePath )
 		{
 		}
@@ -176,23 +176,23 @@ namespace UnrealBuildTool
 
 		private bool IsPartOfUE4XcodeHelperTarget(XcodeSourceFile SourceFile)
 		{
-			string FileExtension = Path.GetExtension(SourceFile.FilePath);
+			string FileExtension = SourceFile.Reference.GetExtension();
 
 			if (IsSourceCode(FileExtension))// || GetFileType(FileExtension) == "sourcecode.c.h") @todo: It seemed that headers need to be added to project for live issues detection to work in them
 			{
 				foreach (string PlatformName in Enum.GetNames(typeof(UnrealTargetPlatform)))
 				{
 					string AltName = PlatformName == "Win32" || PlatformName == "Win64" ? "windows" : PlatformName.ToLower();
-					if ((SourceFile.FilePath.ToLower().Contains("/" + PlatformName.ToLower() + "/") || SourceFile.FilePath.ToLower().Contains("/" + AltName + "/"))
-						&& PlatformName != "Mac" && !SourceFile.FilePath.Contains("MetalRHI"))
+					if ((SourceFile.Reference.FullName.ToLower().Contains("/" + PlatformName.ToLower() + "/") || SourceFile.Reference.FullName.ToLower().Contains("/" + AltName + "/"))
+						&& PlatformName != "Mac" && !SourceFile.Reference.FullName.Contains("MetalRHI"))
 					{
 						// UE4XcodeHelper is Mac only target, so skip other platforms files
 						return false;
 					}
-					else if (SourceFile.FilePath.EndsWith("SimplygonMeshReduction.cpp") || SourceFile.FilePath.EndsWith("MeshBoneReduction.cpp") || SourceFile.FilePath.EndsWith("Android.cpp")
-						|| SourceFile.FilePath.EndsWith("Amazon.cpp") || SourceFile.FilePath.EndsWith("FacebookModule.cpp") || SourceFile.FilePath.EndsWith("SDL_angle.c")
-						|| SourceFile.FilePath.Contains("VisualStudioSourceCodeAccess") || SourceFile.FilePath.Contains("AndroidDevice") || SourceFile.FilePath.Contains("IOSDevice")
-						|| SourceFile.FilePath.Contains("WindowsDevice") || SourceFile.FilePath.Contains("WindowsMoviePlayer") || SourceFile.FilePath.EndsWith("IOSTapJoy.cpp"))
+					else if (SourceFile.Reference.FullName.EndsWith("SimplygonMeshReduction.cpp") || SourceFile.Reference.FullName.EndsWith("MeshBoneReduction.cpp") || SourceFile.Reference.FullName.EndsWith("Android.cpp")
+						|| SourceFile.Reference.FullName.EndsWith("Amazon.cpp") || SourceFile.Reference.FullName.EndsWith("FacebookModule.cpp") || SourceFile.Reference.FullName.EndsWith("SDL_angle.c")
+						|| SourceFile.Reference.FullName.Contains("VisualStudioSourceCodeAccess") || SourceFile.Reference.FullName.Contains("AndroidDevice") || SourceFile.Reference.FullName.Contains("IOSDevice")
+						|| SourceFile.Reference.FullName.Contains("WindowsDevice") || SourceFile.Reference.FullName.Contains("WindowsMoviePlayer") || SourceFile.Reference.FullName.EndsWith("IOSTapJoy.cpp"))
 					{
 						// @todo: We need a way to filter out files that use SDKs we don't have
 						return false;
@@ -209,13 +209,13 @@ namespace UnrealBuildTool
 		 * Returns a project navigator group to which the file should belong based on its path.
 		 * Creates a group tree if it doesn't exist yet.
 		 */
-		public XcodeFileGroup FindGroupByFullPath(ref Dictionary<string, XcodeFileGroup> Groups, string FullPath)
+		public XcodeFileGroup FindGroupByFullPath(ref Dictionary<string, XcodeFileGroup> Groups, DirectoryReference FullPath)
 		{
-			string RelativePath = (FullPath == XcodeProjectFileGenerator.MasterProjectRelativePath) ? "" : Utils.MakePathRelativeTo(FullPath, XcodeProjectFileGenerator.MasterProjectRelativePath);
+			string RelativePath = (FullPath == XcodeProjectFileGenerator.MasterProjectPath) ? "" : FullPath.MakeRelativeTo(XcodeProjectFileGenerator.MasterProjectPath);
 			if (RelativePath.StartsWith(".."))
 			{
-				string UE4Dir = Path.GetFullPath(Directory.GetCurrentDirectory() + "../../..");
-				RelativePath = Utils.MakePathRelativeTo(FullPath, UE4Dir);
+				DirectoryReference UE4Dir = new DirectoryReference(Directory.GetCurrentDirectory() + "../../..");
+				RelativePath = FullPath.MakeRelativeTo(UE4Dir);
 			}
 
 			string[] Parts = RelativePath.Split(Path.DirectorySeparatorChar);
@@ -260,9 +260,9 @@ namespace UnrealBuildTool
 		/// <param name="InitFilePath">Path to the source file on disk</param>
 		/// <param name="InitProjectSubFolder">Optional sub-folder to put the file in.  If empty, this will be determined automatically from the file's path relative to the project file</param>
 		/// <returns>The newly allocated source file object</returns>
-		public override SourceFile AllocSourceFile(string InitFilePath, string InitProjectSubFolder)
+		public override SourceFile AllocSourceFile(FileReference InitFilePath, DirectoryReference InitProjectSubFolder)
 		{
-			if (Path.GetFileName(InitFilePath) == ".DS_Store")
+			if (InitFilePath.GetFileName() == ".DS_Store")
 			{
 				return null;
 			}
@@ -282,9 +282,9 @@ namespace UnrealBuildTool
 			foreach (var CurSourceFile in SourceFiles)
 			{
 				XcodeSourceFile SourceFile = CurSourceFile as XcodeSourceFile;
-				string FileName = Path.GetFileName(SourceFile.FilePath);
+				string FileName = SourceFile.Reference.GetFileName();
 				string FileExtension = Path.GetExtension(FileName);
-				string FilePath = Utils.MakePathRelativeTo(SourceFile.FilePath, XcodeProjectFileGenerator.MasterProjectRelativePath);
+				string FilePath = SourceFile.Reference.MakeRelativeTo(XcodeProjectFileGenerator.MasterProjectPath);
 				string FilePathMac = Utils.CleanDirectorySeparators(FilePath, '/');
 
 				// Xcode doesn't parse the project file correctly if the build files don't use the same file reference ID,
@@ -319,7 +319,7 @@ namespace UnrealBuildTool
 					BuildPhaseSection.Append("\t\t\t\t" + SourceFile.FileGUID + " /* " + FileName + " in Sources */," + ProjectFileGenerator.NewLine);
 				}
 
-				string GroupPath = Path.GetFullPath(Path.GetDirectoryName(SourceFile.FilePath));
+				DirectoryReference GroupPath = SourceFile.Reference.Directory;
 				XcodeFileGroup Group = FindGroupByFullPath(ref Groups, GroupPath);
 				if (Group != null)
 				{
diff --git a/Engine/Source/Programs/UnrealBuildTool/System/XcodeProjectFileGenerator.cs b/Engine/Source/Programs/UnrealBuildTool/System/XcodeProjectFileGenerator.cs
--- a/Engine/Source/Programs/UnrealBuildTool/System/XcodeProjectFileGenerator.cs
+++ b/Engine/Source/Programs/UnrealBuildTool/System/XcodeProjectFileGenerator.cs
@@ -4,6 +4,7 @@ using System;
 using System.Collections.Generic;
 using System.Text;
 using System.IO;
+using System.Linq;
 
 namespace UnrealBuildTool
 {
@@ -56,7 +57,7 @@ namespace UnrealBuildTool
 		/// <param name="InTargetPlatform">Name of the target that may be built.</param>
 		/// <param name="InDependencies">Name of the target that may be built.</param>
 		/// <param name="bHasPlist">Name of the target that may be built.</param>
-		public XcodeProjectTarget(string InDisplayName, string InTargetName, XcodeTargetType InType, string InTargetFilePath, string InProductName = "", UnrealTargetPlatform InTargetPlatform = UnrealTargetPlatform.Mac, bool bInIsMacOnly = false, List<XcodeTargetDependency> InDependencies = null, bool bHasPlist = false, List<XcodeFrameworkRef> InFrameworks = null)
+		public XcodeProjectTarget(string InDisplayName, string InTargetName, XcodeTargetType InType, FileReference InTargetFilePath, string InProductName = "", UnrealTargetPlatform InTargetPlatform = UnrealTargetPlatform.Mac, bool bInIsMacOnly = false, List<XcodeTargetDependency> InDependencies = null, bool bHasPlist = false, List<XcodeFrameworkRef> InFrameworks = null)
 		{
 			DisplayName = InDisplayName;
 			TargetName = InTargetName;
@@ -92,7 +93,7 @@ namespace UnrealBuildTool
 		public UnrealTargetPlatform TargetPlatform;	//Mac, IOS
 		public bool bIsMacOnly;
 		public string ProductName;
-		public string TargetFilePath;
+		public FileReference TargetFilePath;
 
 		public string Guid;
 		public string BuildConfigGuild;
@@ -217,7 +218,7 @@ namespace UnrealBuildTool
 			}
 		}
 
-		public override void CleanProjectFiles(string InMasterProjectRelativePath, string InMasterProjectName, string InIntermediateProjectFilesPath)
+		public override void CleanProjectFiles(DirectoryReference InMasterProjectDirectory, string InMasterProjectName, DirectoryReference InIntermediateProjectFilesDirectory)
 		{
 			//@todo Mac. Implement this function...
 		}
@@ -227,7 +228,7 @@ namespace UnrealBuildTool
 		/// </summary>
 		/// <param name="InitFilePath">Path to the project file</param>
 		/// <returns>The newly allocated project file object</returns>
-		protected override ProjectFile AllocateProjectFile(string InitFilePath)
+		protected override ProjectFile AllocateProjectFile(FileReference InitFilePath)
 		{
 			return new XcodeProjectFile(InitFilePath);
 		}
@@ -350,11 +351,11 @@ namespace UnrealBuildTool
 			if (Target.Type == XcodeTargetType.Legacy)
 			{
 				string UProjectPath = "";
-				if (UnrealBuildTool.HasUProjectFile() && !string.IsNullOrEmpty(Target.TargetFilePath) && Utils.IsFileUnderDirectory(Target.TargetFilePath, UnrealBuildTool.GetUProjectPath()))
+				if (UnrealBuildTool.HasUProjectFile() && Target.TargetFilePath != null && Target.TargetFilePath.IsUnderDirectory(UnrealBuildTool.GetUProjectPath()))
 				{
-					if (MasterProjectRelativePath == UnrealBuildTool.GetUProjectPath())
+					if (MasterProjectPath == UnrealBuildTool.GetUProjectPath())
 					{
-						UProjectPath = " " + "\\\"$(PROJECT_DIR)/" + Path.GetFileName(UnrealBuildTool.GetUProjectFile()) + "\\\"";
+						UProjectPath = " " + "\\\"$(PROJECT_DIR)/" + Path.GetFileName(UnrealBuildTool.GetUProjectFile().FullName) + "\\\"";
 					}
 					else
 					{
@@ -722,7 +723,7 @@ namespace UnrealBuildTool
 		/// <param name="PreprocessorDefinitions">Optional string with preprocessor definitions section</param>
 		private void AppendBuildConfigs(ref StringBuilder Contents, XcodeProjectTarget Target, string HeaderSearchPaths = "", string PreprocessorDefinitions = "")
 		{
-			List<string> GameFolders = UEBuildTarget.DiscoverAllGameFolders();
+			List<DirectoryReference> GameFolders = UEBuildTarget.DiscoverAllGameFolders();
 
 			bool IsAGame = false;
 			string GamePath = null;
@@ -732,17 +733,17 @@ namespace UnrealBuildTool
 			string EngineRelative = "";
 
 			IUEToolChain Toolchain = UEToolChain.GetPlatformToolChain(CPPTargetPlatform.IOS);
-			foreach (string GameFolder in GameFolders)
+			foreach (string GameFolder in GameFolders.Select(x => x.FullName))
 			{
 				if (File.Exists(Path.Combine(GameFolder, Target.TargetName+".uproject")))
 				{
 					IsAGame = true;
-					GamePath = Toolchain.ConvertPath(Path.GetFullPath(GameFolder));
+					GamePath = Toolchain.ConvertPath(GameFolder);
 					break;
 				}
 				else if (!string.IsNullOrEmpty(GameNameFromClientName) && GameFolder.EndsWith(GameNameFromClientName))
 				{
-					GamePath = Toolchain.ConvertPath(Path.GetFullPath(GameFolder));
+					GamePath = Toolchain.ConvertPath(GameFolder);
 					break;
 				}
 			}
@@ -880,12 +881,11 @@ namespace UnrealBuildTool
 		}
 
 		private void AddPreGeneratedProject(ref string PBXBuildFileSection, ref string PBXFileReferenceSection, ref string PBXSourcesBuildPhaseSection,
-			ref Dictionary<string, XcodeFileGroup> Groups, string ProjectPath)
+			ref Dictionary<string, XcodeFileGroup> Groups, FileReference ProjectPath)
 		{
-			var ProjectFileName = Utils.MakePathRelativeTo(ProjectPath, MasterProjectRelativePath);
-			var XcodeProject = new XcodeProjectFile(ProjectFileName);
-			string ProjectDirectory = Path.GetDirectoryName(ProjectPath);
-			XcodeProject.AddFilesToProject(SourceFileSearch.FindFiles(new List<string>() { ProjectDirectory }, ExcludeNoRedistFiles: bExcludeNoRedistFiles, SubdirectoryNamesToExclude:new List<string>() { "obj" }, SearchSubdirectories:true, IncludePrivateSourceCode:true), null);
+			var XcodeProject = new XcodeProjectFile(ProjectPath);
+			DirectoryReference ProjectDirectory = ProjectPath.Directory;
+			XcodeProject.AddFilesToProject(SourceFileSearch.FindFiles(new List<DirectoryReference>() { ProjectDirectory }, ExcludeNoRedistFiles: bExcludeNoRedistFiles, SubdirectoryNamesToExclude:new List<string>() { "obj" }, SearchSubdirectories:true, IncludePrivateSourceCode:true), null);
 			XcodeProject.GenerateSectionsContents(ref PBXBuildFileSection, ref PBXFileReferenceSection, ref PBXSourcesBuildPhaseSection, ref Groups);
 		}
 
@@ -906,7 +906,7 @@ namespace UnrealBuildTool
 				{
 					if (TargetFile.TargetFilePath != null)
 					{
-						string TargetFileName = Path.GetFileNameWithoutExtension(TargetFile.TargetFilePath);
+						string TargetFileName = TargetFile.TargetFilePath.GetFileNameWithoutExtension();
 						string TargetFileMinusTarget = TargetFileName.Substring(0, TargetFileName.LastIndexOf(".Target"));
 						foreach (XcodeProjectTarget Target in ProjectTargets)
 						{
@@ -944,7 +944,7 @@ namespace UnrealBuildTool
 					foreach (var CurSourceFile in XcodeProject.SourceFiles)
 					{
 						XcodeSourceFile SourceFile = CurSourceFile as XcodeSourceFile;
-						string GroupPath = Path.GetFullPath(Path.GetDirectoryName(SourceFile.FilePath));
+						DirectoryReference GroupPath = SourceFile.Reference.Directory;
 						XcodeFileGroup Group = XcodeProject.FindGroupByFullPath(ref Groups, GroupPath);
 						if (Group != null)
 						{
@@ -959,12 +959,12 @@ namespace UnrealBuildTool
 
 				foreach (var CurPath in XcodeProject.IntelliSenseIncludeSearchPaths)
 				{
-					AddIncludeDirectory(ref IncludeDirectories, CurPath, Path.GetDirectoryName(XcodeProject.ProjectFilePath));
+					AddIncludeDirectory(ref IncludeDirectories, CurPath, Path.GetDirectoryName(XcodeProject.ProjectFilePath.FullName));
 				}
 
 				foreach (var CurPath in XcodeProject.IntelliSenseSystemIncludeSearchPaths)
 				{
-					AddIncludeDirectory(ref SystemIncludeDirectories, CurPath, Path.GetDirectoryName(XcodeProject.ProjectFilePath));
+					AddIncludeDirectory(ref SystemIncludeDirectories, CurPath, Path.GetDirectoryName(XcodeProject.ProjectFilePath.FullName));
 				}
 
 				foreach (var CurDefinition in XcodeProject.IntelliSensePreprocessorDefinitions)
@@ -988,7 +988,7 @@ namespace UnrealBuildTool
 			if (!bGeneratingRocketProjectFiles)
 			{
 				// Add UnrealBuildTool to the master project
-				string ProjectPath = System.IO.Path.Combine(System.IO.Path.Combine(EngineRelativePath, "Source"), "Programs", "UnrealBuildTool", "UnrealBuildTool.csproj");
+				FileReference ProjectPath = FileReference.Combine(UnrealBuildTool.EngineSourceDirectory, "Programs", "UnrealBuildTool", "UnrealBuildTool.csproj");
 				AddPreGeneratedProject(ref PBXBuildFileSection, ref PBXFileReferenceSection, ref PBXSourcesBuildPhaseSection, ref Groups, ProjectPath);
 			}
 		}
@@ -1052,7 +1052,7 @@ namespace UnrealBuildTool
 		/// Adds the include directory to the list, after converting it to relative to UE4 root
 		private void AddIncludeDirectory(ref List<string> IncludeDirectories, string IncludeDir, string ProjectDir)
 		{
-			string FullProjectPath = Path.GetFullPath(ProjectFileGenerator.MasterProjectRelativePath);
+			string FullProjectPath = ProjectFileGenerator.MasterProjectPath.FullName;
 			string FullPath = "";
 			if (IncludeDir.StartsWith("/") && !IncludeDir.StartsWith(FullProjectPath))
 			{
@@ -1086,22 +1086,23 @@ namespace UnrealBuildTool
 				// add UE4Game as most likely this is a content-only project
 				ProjectTarget TargetProject = new global::UnrealBuildTool.ProjectTarget();
 				var Target = new TargetInfo(UnrealTargetPlatform.Mac, UnrealTargetConfiguration.Development);
-				TargetProject.TargetFilePath = "";
+				TargetProject.TargetFilePath = null;
 				RulesAssembly EngineRulesAssembly = RulesCompiler.CreateEngineRulesAssembly();
-				TargetProject.TargetRules = EngineRulesAssembly.CreateTargetRules("UE4Game", Target, false, out TargetProject.TargetFilePath);
+				FileReference TargetFilePath;
+				TargetProject.TargetRules = EngineRulesAssembly.CreateTargetRules("UE4Game", Target, false, out TargetFilePath);
+				TargetProject.TargetFilePath = TargetFilePath;
 				TargetProjects.Add(TargetProject);
 			}
 			foreach (var TargetProject in TargetProjects)
 			{
-					string TargetPath = TargetProject.TargetFilePath;
-					string TargetName = Utils.GetFilenameWithoutAnyExtensions (Path.GetFileName (TargetPath));
+					FileReference TargetPath = TargetProject.TargetFilePath;
+					string TargetName = TargetPath.GetFileNameWithoutAnyExtensions();
 					bool WantProjectFileForTarget = true;
 					bool IsEngineTarget = false;
 					if (bGeneratingGameProjectFiles || bGeneratingRocketProjectFiles)
 					{
 						// Check to see if this is an Engine target.  That is, the target is located under the "Engine" folder
-						string TargetFileRelativeToEngineDirectory = Utils.MakePathRelativeTo(TargetPath, Path.Combine(EngineRelativePath), AlwaysTreatSourceAsDirectory: false);
-						if (!TargetFileRelativeToEngineDirectory.StartsWith("..") && !Path.IsPathRooted(TargetFileRelativeToEngineDirectory))
+						if (TargetPath.IsUnderDirectory(UnrealBuildTool.EngineDirectory))
 						{
 							// This is an engine target
 							IsEngineTarget = true;
@@ -1128,7 +1129,7 @@ namespace UnrealBuildTool
 
 					if (WantProjectFileForTarget)
 					{
-						string TargetFilePath = TargetProject.TargetFilePath;
+						FileReference TargetFilePath = TargetProject.TargetFilePath;
 						var Target = new TargetInfo(UnrealTargetPlatform.Mac, UnrealTargetConfiguration.Development);
 						var TargetRulesObject = TargetProject.TargetRules; // RulesCompiler.CreateTargetRules(TargetName, Target, false, out TargetFilePath);
 						List<UnrealTargetPlatform> SupportedPlatforms = new List<UnrealTargetPlatform>();
@@ -1223,16 +1224,16 @@ namespace UnrealBuildTool
 			{
 				TargetBinariesFolder = Path.GetDirectoryName(Directory.GetCurrentDirectory());
 
-                if (!ExeBaseName.EndsWith("Editor") && !string.IsNullOrWhiteSpace(Target.TargetFilePath) && !Utils.IsFileUnderDirectory(Target.TargetFilePath, TargetBinariesFolder))
+                if (!ExeBaseName.EndsWith("Editor") && Target.TargetFilePath != null && !Target.TargetFilePath.IsUnderDirectory(new DirectoryReference(TargetBinariesFolder)))
                 {
 					string DepotRoot = TargetBinariesFolder.TrimEnd(Path.DirectorySeparatorChar);
 					DepotRoot = DepotRoot.Substring(0, DepotRoot.LastIndexOf(Path.DirectorySeparatorChar));
 					DepotRoot = DepotRoot.TrimEnd(Path.DirectorySeparatorChar);
 
-                    if (Utils.IsFileUnderDirectory(Target.TargetFilePath, DepotRoot))
+                    if (Target.TargetFilePath.IsUnderDirectory(new DirectoryReference(DepotRoot)))
                     {
-                        string PreviousPath = Target.TargetFilePath;
-                        string CurrentPath = Target.TargetFilePath;
+                        string PreviousPath = Target.TargetFilePath.FullName;
+                        string CurrentPath = Target.TargetFilePath.FullName;
                         while (!string.IsNullOrWhiteSpace(CurrentPath) && CurrentPath != DepotRoot)
                         {
                             PreviousPath = CurrentPath;
@@ -1246,7 +1247,7 @@ namespace UnrealBuildTool
                     }
                 }
 
-				if (ExeBaseName == Target.TargetName && !Target.TargetFilePath.StartsWith(TargetBinariesFolder) && !ExeBaseName.EndsWith("Editor") && ExeExtension == ".app")
+				if (ExeBaseName == Target.TargetName && (Target.TargetFilePath == null || !Target.TargetFilePath.FullName.StartsWith(TargetBinariesFolder)) && !ExeBaseName.EndsWith("Editor") && ExeExtension == ".app")
 				{
 					TargetBinariesFolder = TargetBinariesFolder.Replace("/Engine", "/") + ExeBaseName + "/Binaries/";
 				}
@@ -1424,7 +1425,7 @@ namespace UnrealBuildTool
 					continue;
 				}
 
-				if (UnrealBuildTool.HasUProjectFile() && !string.IsNullOrEmpty(Target.TargetFilePath) && Utils.IsFileUnderDirectory(Target.TargetFilePath, UnrealBuildTool.GetUProjectPath()))
+				if (UnrealBuildTool.HasUProjectFile() && Target.TargetFilePath != null && Target.TargetFilePath.IsUnderDirectory(UnrealBuildTool.GetUProjectPath()))
 				{
 					if (Target.TargetName.EndsWith("Editor"))
 					{
@@ -1455,8 +1456,8 @@ namespace UnrealBuildTool
 		private void WriteProject(StringBuilder XcodeProjectFileContent, List<XcodeProjectTarget> ProjectTargets)
 		{
 			// @todo: Windows and Mac should store the project in the same folder
-			string PathBranch = (BuildHostPlatform.Current.Platform != UnrealTargetPlatform.Mac) ? (!string.IsNullOrEmpty(UnrealBuildTool.GetUProjectPath()) ? "/Intermediate/IOS" : "/Engine/Intermediate/IOS") : "";
-			var XcodeProjectPath = MasterProjectRelativePath + PathBranch + "/" + MasterProjectName + ".xcodeproj";
+			string PathBranch = (BuildHostPlatform.Current.Platform != UnrealTargetPlatform.Mac) ? ((UnrealBuildTool.GetUProjectPath() != null) ? "/Intermediate/IOS" : "/Engine/Intermediate/IOS") : "";
+			var XcodeProjectPath = MasterProjectPath.FullName + PathBranch + "/" + MasterProjectName + ".xcodeproj";
 			if (bGeneratingGameProjectFiles && GameProjectName == "UE4Game")
 			{
 				XcodeProjectPath = UnrealBuildTool.GetUProjectPath() + "/" + MasterProjectName + ".xcodeproj";
@@ -1491,8 +1492,8 @@ namespace UnrealBuildTool
 
 		private void WriteWorkspace()
 		{
-			string PathBranch = (BuildHostPlatform.Current.Platform != UnrealTargetPlatform.Mac) ? (!string.IsNullOrEmpty(UnrealBuildTool.GetUProjectPath()) ? "/Intermediate/IOS" : "/Engine/Intermediate/IOS") : "";
-			var XcodeProjectPath = MasterProjectRelativePath + PathBranch + "/" + MasterProjectName + ".xcodeproj";
+			string PathBranch = (BuildHostPlatform.Current.Platform != UnrealTargetPlatform.Mac) ? ((UnrealBuildTool.GetUProjectPath() != null) ? "/Intermediate/IOS" : "/Engine/Intermediate/IOS") : "";
+			var XcodeProjectPath = MasterProjectPath + PathBranch + "/" + MasterProjectName + ".xcodeproj";
 			if (bGeneratingGameProjectFiles && GameProjectName == "UE4Game")
 			{
 				XcodeProjectPath = UnrealBuildTool.GetUProjectPath() + "/" + MasterProjectName + ".xcodeproj";
@@ -1533,13 +1534,13 @@ namespace UnrealBuildTool
 			// attempt to determine targets for the project
 			List<XcodeProjectTarget> ProjectTargets = new List<XcodeProjectTarget> ();
 			// add mandatory ones
-			XcodeProjectTarget UE4ProjectTarget = new XcodeProjectTarget ("UE4", "UE4", XcodeTargetType.Project, "");
-			XcodeProjectTarget UE4XcodeHelperTarget = new XcodeProjectTarget ("UE4XcodeHelper", "UE4XcodeHelper", XcodeTargetType.XcodeHelper, "", "libUE4XcodeHelper.a");
+			XcodeProjectTarget UE4ProjectTarget = new XcodeProjectTarget ("UE4", "UE4", XcodeTargetType.Project, null);
+			XcodeProjectTarget UE4XcodeHelperTarget = new XcodeProjectTarget ("UE4XcodeHelper", "UE4XcodeHelper", XcodeTargetType.XcodeHelper, null, "libUE4XcodeHelper.a");
 			ProjectTargets.AddRange(new XcodeProjectTarget[] { UE4ProjectTarget, UE4XcodeHelperTarget });
 
 			if (ProjectFilePlatform.HasFlag(XcodeProjectFilePlatform.iOS))
 			{
-				XcodeProjectTarget UE4CmdLineRunTarget = new XcodeProjectTarget("UE4CmdLineRun", "UE4CmdLineRun", XcodeTargetType.XCTest, "", "UE4CmdLineRun.xctest", UnrealTargetPlatform.IOS);
+				XcodeProjectTarget UE4CmdLineRunTarget = new XcodeProjectTarget("UE4CmdLineRun", "UE4CmdLineRun", XcodeTargetType.XCTest, null, "UE4CmdLineRun.xctest", UnrealTargetPlatform.IOS);
 				ProjectTargets.Add(UE4CmdLineRunTarget);
 				// This GUID will be referenced by each app's test action.
 				UE4CmdLineGuid = UE4CmdLineRunTarget.Guid;
diff --git a/Engine/Source/Programs/UnrealBuildTool/ToolChain/RemoteToolChain.cs b/Engine/Source/Programs/UnrealBuildTool/ToolChain/RemoteToolChain.cs
--- a/Engine/Source/Programs/UnrealBuildTool/ToolChain/RemoteToolChain.cs
+++ b/Engine/Source/Programs/UnrealBuildTool/ToolChain/RemoteToolChain.cs
@@ -160,7 +160,7 @@ namespace UnrealBuildTool
 						throw new BuildException("Configuration setting was using ${PROJECT_ROOT}, but there was no project specified");
 					}
 
-					string Temp = Result.Replace("${PROJECT_ROOT}", Path.GetFullPath(UnrealBuildTool.GetUProjectPath()));
+					string Temp = Result.Replace("${PROJECT_ROOT}", Path.GetFullPath(UnrealBuildTool.GetUProjectPath().FullName));
 					
 					// get the best version
 					Result = LookForSpecialFile(Temp);
@@ -387,16 +387,18 @@ namespace UnrealBuildTool
 					if (!bFoundOverrideSSHPrivateKey)
 					{
 						// all the places to look for a key
-						string[] Locations = new string[] {
-							Path.Combine(Environment.GetFolderPath(Environment.SpecialFolder.ApplicationData), "Unreal Engine", "UnrealBuildTool"),
-							Path.Combine(Environment.GetFolderPath(Environment.SpecialFolder.Personal), "Unreal Engine", "UnrealBuildTool"),
-							Path.Combine(UnrealBuildTool.GetUProjectPath(), "Build", "NotForLicensees"),
-							Path.Combine(UnrealBuildTool.GetUProjectPath(), "Build", "NoRedist"),
-							Path.Combine(UnrealBuildTool.GetUProjectPath(), "Build"),
-							Path.Combine(BuildConfiguration.RelativeEnginePath, "Build", "NotForLicensees"),
-							Path.Combine(BuildConfiguration.RelativeEnginePath, "Build", "NoRedist"),
-							Path.Combine(BuildConfiguration.RelativeEnginePath, "Build"),
-						};
+						List<string> Locations = new List<string>();
+						Locations.Add(Path.Combine(Environment.GetFolderPath(Environment.SpecialFolder.ApplicationData), "Unreal Engine", "UnrealBuildTool"));
+						Locations.Add(Path.Combine(Environment.GetFolderPath(Environment.SpecialFolder.Personal), "Unreal Engine", "UnrealBuildTool"));
+						if(UnrealBuildTool.HasUProjectFile())
+						{
+							Locations.Add(Path.Combine(UnrealBuildTool.GetUProjectPath().FullName, "Build", "NotForLicensees"));
+							Locations.Add(Path.Combine(UnrealBuildTool.GetUProjectPath().FullName, "Build", "NoRedist"));
+							Locations.Add(Path.Combine(UnrealBuildTool.GetUProjectPath().FullName, "Build"));
+						}
+						Locations.Add(Path.Combine(BuildConfiguration.RelativeEnginePath, "Build", "NotForLicensees"));
+						Locations.Add(Path.Combine(BuildConfiguration.RelativeEnginePath, "Build", "NoRedist"));
+						Locations.Add(Path.Combine(BuildConfiguration.RelativeEnginePath, "Build"));
 
 						// look for a key file
 						foreach (string Location in Locations)
@@ -472,7 +474,7 @@ namespace UnrealBuildTool
 		}
 
         /** Converts the passed in path from UBT host to compiler native format. */
-        public override String ConvertPath(String OriginalPath)
+        public override string ConvertPath(string OriginalPath)
         {
             if (BuildHostPlatform.Current.Platform != UnrealTargetPlatform.Mac)
             {
@@ -507,7 +509,7 @@ namespace UnrealBuildTool
 			}
 			else
 			{
-				return ".";
+				return UnrealBuildTool.EngineSourceDirectory.FullName;;
 			}
 		}
 
