Use FileReference and DirectoryReference classes instead of strings for paths. Massively reduces the amount of paranoid calls to GetFullPath(), CleanDirectorySeparators() et al., and improves performa

diff --git a/Engine/Source/Programs/UnrealBuildTool/ToolChain/UEToolChain.cs b/Engine/Source/Programs/UnrealBuildTool/ToolChain/UEToolChain.cs
--- a/Engine/Source/Programs/UnrealBuildTool/ToolChain/UEToolChain.cs
+++ b/Engine/Source/Programs/UnrealBuildTool/ToolChain/UEToolChain.cs
@@ -19,7 +19,7 @@ namespace UnrealBuildTool
 		
 		FileItem[] LinkAllFiles(LinkEnvironment LinkEnvironment, bool bBuildImportLibraryOnly);
 		
-		void CompileCSharpProject(CSharpEnvironment CompileEnvironment, string ProjectFileName, string DestinationFile);
+		void CompileCSharpProject(CSharpEnvironment CompileEnvironment, FileReference ProjectFileName, FileReference DestinationFile);
 		
 		/** Converts the passed in path from UBT host to compiler native format. */
 		String ConvertPath(String OriginalPath);
@@ -40,7 +40,7 @@ namespace UnrealBuildTool
 
 		void AddFilesToReceipt(TargetReceipt Receipt, UEBuildBinary Binary);
 
-		bool ShouldAddDebugFileToReceipt(string OutputFile, BuildProductType OutputType);
+		bool ShouldAddDebugFileToReceipt(FileReference OutputFile, BuildProductType OutputType);
 
 		void SetupBundleDependencies(List<UEBuildBinary> Binaries, string GameName);
 
@@ -100,7 +100,7 @@ namespace UnrealBuildTool
 		}
 
 
-		public virtual void CompileCSharpProject(CSharpEnvironment CompileEnvironment, string ProjectFileName, string DestinationFile)
+		public virtual void CompileCSharpProject(CSharpEnvironment CompileEnvironment, FileReference ProjectFileName, FileReference DestinationFile)
 		{
 		}
 
@@ -110,26 +110,23 @@ namespace UnrealBuildTool
 		/// <param name="LinkEnvironment"></param>
 		/// <param name="OutputFile"></param>
 		/// <returns></returns>
-		public static string GetResponseFileName( LinkEnvironment LinkEnvironment, FileItem OutputFile )
+		public static FileReference GetResponseFileName( LinkEnvironment LinkEnvironment, FileItem OutputFile )
 		{
 			// Construct a relative path for the intermediate response file
-			string ResponseFileName = Path.Combine( LinkEnvironment.Config.IntermediateDirectory, Path.GetFileName( OutputFile.AbsolutePath ) + ".response" );
+			FileReference ResponseFileName = FileReference.Combine( LinkEnvironment.Config.IntermediateDirectory, Path.GetFileName( OutputFile.AbsolutePath ) + ".response" );
 			if (UnrealBuildTool.HasUProjectFile())
 			{
 				// If this is the uproject being built, redirect the intermediate
-				if (Utils.IsFileUnderDirectory( OutputFile.AbsolutePath, UnrealBuildTool.GetUProjectPath() ))
+				if (OutputFile.Reference.IsUnderDirectory( UnrealBuildTool.GetUProjectPath() ))
 				{
-					ResponseFileName = Path.Combine(
+					ResponseFileName = FileReference.Combine(
 						UnrealBuildTool.GetUProjectPath(), 
 						BuildConfiguration.PlatformIntermediateFolder,
-						Path.GetFileNameWithoutExtension(UnrealBuildTool.GetUProjectFile()),
+						UnrealBuildTool.GetUProjectFile().GetFileNameWithoutExtension(),
 						LinkEnvironment.Config.Target.Configuration.ToString(),
 						Path.GetFileName(OutputFile.AbsolutePath) + ".response");
 				}
 			}
-			// Convert the relative path to an absolute path
-			ResponseFileName = Path.GetFullPath( ResponseFileName );
-
 			return ResponseFileName;
 		}
 
@@ -183,7 +180,7 @@ namespace UnrealBuildTool
 		/// </summary>
 		/// <param name="OutputFile">Build product to add</param>
 		/// <param name="DebugExtension">Extension for the matching debug file (may be null).</param>
-		public virtual bool ShouldAddDebugFileToReceipt(string OutputFile, BuildProductType OutputType)
+		public virtual bool ShouldAddDebugFileToReceipt(FileReference OutputFile, BuildProductType OutputType)
 		{
 			return true;
 		}
diff --git a/Engine/Source/Programs/UnrealBuildTool/UWP/UWPDeploy.cs b/Engine/Source/Programs/UnrealBuildTool/UWP/UWPDeploy.cs
--- a/Engine/Source/Programs/UnrealBuildTool/UWP/UWPDeploy.cs
+++ b/Engine/Source/Programs/UnrealBuildTool/UWP/UWPDeploy.cs
@@ -220,14 +220,14 @@ namespace UnrealBuildTool
 			Log.TraceInformation("Prepping {0} for deployment to {1}", InAppName, InTarget.Platform.ToString());
 			System.DateTime PrepDeployStartTime = DateTime.UtcNow;
 
-			string TargetFilename = InTarget.RulesAssembly.GetTargetFileName(InAppName);
+			string TargetFilename = InTarget.RulesAssembly.GetTargetFileName(InAppName).FullName;
 			string ProjectSourceFolder = new FileInfo(TargetFilename).DirectoryName + "/";
 
 			string RelativeTargetDirectory;
 			string EngineSourceRelativeBinaryPath;
 			UWPProjectGenerator.GetTargetUWPPaths(InTarget.RulesAssembly, InAppName, InTarget.Rules, out EngineSourceRelativeBinaryPath, out RelativeTargetDirectory);
 
-			PrepForUATPackageOrDeploy(InAppName, InTarget.ProjectDirectory, InTarget.OutputPath, BuildConfiguration.RelativeEnginePath, false, "");
+			PrepForUATPackageOrDeploy(InAppName, InTarget.ProjectDirectory.FullName, InTarget.OutputPath.FullName, BuildConfiguration.RelativeEnginePath, false, "");
 
 			// TODO - richiem - restore this if we find that it's needed.
 			//Log.TraceInformation("...copying the CELL dll...");
diff --git a/Engine/Source/Programs/UnrealBuildTool/UWP/UWPProjectGenerator.cs b/Engine/Source/Programs/UnrealBuildTool/UWP/UWPProjectGenerator.cs
--- a/Engine/Source/Programs/UnrealBuildTool/UWP/UWPProjectGenerator.cs
+++ b/Engine/Source/Programs/UnrealBuildTool/UWP/UWPProjectGenerator.cs
@@ -35,7 +35,7 @@ namespace UnrealBuildTool
 			OutEngineSourceRelativeBinaryPath = "";
 			OutRelativeTargetPath = "";
 
-			string TargetFilename = InTargetRulesAssembly.GetTargetFileName(InTargetName);
+			string TargetFilename = InTargetRulesAssembly.GetTargetFileName(InTargetName).FullName;
 
 			string ProjectSourceFolder = new FileInfo(TargetFilename).DirectoryName;
 
@@ -398,7 +398,7 @@ namespace UnrealBuildTool
 		}
 
 
-		public override void GenerateGameProperties(UnrealTargetConfiguration Configuration, StringBuilder VCProjectFileContent, TargetRules.TargetType TargetType, string RootDirectory, string TargetFilePath)
+		public override void GenerateGameProperties(UnrealTargetConfiguration Configuration, StringBuilder VCProjectFileContent, TargetRules.TargetType TargetType, DirectoryReference RootDirectory, FileReference TargetFilePath)
 		{
 			// @todo UWP: This used to be "WINUAP=1".  Need to verify that 'UWP' is the correct define that we want here.
 			VCProjectFileContent.Append("		<NMakePreprocessorDefinitions>$(NMakePreprocessorDefinitions);PLATFORM_UWP=1;UWP=1;</NMakePreprocessorDefinitions>" + ProjectFileGenerator.NewLine);
diff --git a/Engine/Source/Programs/UnrealBuildTool/UWP/UWPToolChain.cs b/Engine/Source/Programs/UnrealBuildTool/UWP/UWPToolChain.cs
--- a/Engine/Source/Programs/UnrealBuildTool/UWP/UWPToolChain.cs
+++ b/Engine/Source/Programs/UnrealBuildTool/UWP/UWPToolChain.cs
@@ -568,8 +568,8 @@ namespace UnrealBuildTool
 				if (CompileEnvironment.Config.PrecompiledHeaderAction == PrecompiledHeaderAction.Create)
 				{
 					// Generate a CPP File that just includes the precompiled header.
-					string PCHCPPFilename = "PCH." + ModuleName + "." + Path.GetFileName(CompileEnvironment.Config.PrecompiledHeaderIncludeFilename) + ".cpp";
-					string PCHCPPPath = Path.Combine(CompileEnvironment.Config.OutputDirectory, PCHCPPFilename);
+					string PCHCPPFilename = "PCH." + ModuleName + "." + CompileEnvironment.Config.PrecompiledHeaderIncludeFilename.GetFileName() + ".cpp";
+					FileReference PCHCPPPath = FileReference.Combine(CompileEnvironment.Config.OutputDirectory, PCHCPPFilename);
 					FileItem PCHCPPFile = FileItem.CreateIntermediateTextFile(
 						PCHCPPPath,
 						string.Format("#include \"{0}\"\r\n", CompileEnvironment.Config.PrecompiledHeaderIncludeFilename)
@@ -582,10 +582,10 @@ namespace UnrealBuildTool
 
 					var PrecompiledFileExtension = UEBuildPlatform.BuildPlatformDictionary[UnrealTargetPlatform.UWP].GetBinaryExtension(UEBuildBinaryType.PrecompiledHeader);
 					// Add the precompiled header file to the produced items list.
-					FileItem PrecompiledHeaderFile = FileItem.GetItemByPath(
-						Path.Combine(
+					FileItem PrecompiledHeaderFile = FileItem.GetItemByFileReference(
+						FileReference.Combine(
 							CompileEnvironment.Config.OutputDirectory,
-						Path.GetFileName(SourceFile.AbsolutePath) + PrecompiledFileExtension
+							Path.GetFileName(SourceFile.AbsolutePath) + PrecompiledFileExtension
 							)
 						);
 					CompileAction.ProducedItems.Add(PrecompiledHeaderFile);
@@ -603,7 +603,7 @@ namespace UnrealBuildTool
 					{
 						// NOTE: The symbol name we use here is arbitrary, and all that matters is that it is
 						// unique per PCH module used in our library
-						string FakeUniquePCHSymbolName = Path.GetFileNameWithoutExtension(CompileEnvironment.Config.PrecompiledHeaderIncludeFilename);
+						string FakeUniquePCHSymbolName = CompileEnvironment.Config.PrecompiledHeaderIncludeFilename.GetFileNameWithoutExtension();
 						FileArguments.AppendFormat(" /Yl{0}", FakeUniquePCHSymbolName);
 					}
 
@@ -643,8 +643,8 @@ namespace UnrealBuildTool
 				{
 					var ObjectFileExtension = UEBuildPlatform.BuildPlatformDictionary[UnrealTargetPlatform.UWP].GetBinaryExtension(UEBuildBinaryType.Object);
 					// Add the object file to the produced item list.
-					FileItem ObjectFile = FileItem.GetItemByPath(
-						Path.Combine(
+					FileItem ObjectFile = FileItem.GetItemByFileReference(
+						FileReference.Combine(
 							CompileEnvironment.Config.OutputDirectory,
 							Path.GetFileName(SourceFile.AbsolutePath) + ObjectFileExtension
 							)
@@ -668,12 +668,12 @@ namespace UnrealBuildTool
 					// All files using the same PCH are required to share the same PDB that was used when compiling the PCH
 					if (CompileEnvironment.Config.PrecompiledHeaderAction == PrecompiledHeaderAction.Include)
 					{
-						PDBFileName = "PCH." + ModuleName + "." + Path.GetFileName(CompileEnvironment.Config.PrecompiledHeaderIncludeFilename);
+						PDBFileName = "PCH." + ModuleName + "." + CompileEnvironment.Config.PrecompiledHeaderIncludeFilename.GetFileName();
 					}
 					// Files creating a PCH use a PDB per file.
 					else if (CompileEnvironment.Config.PrecompiledHeaderAction == PrecompiledHeaderAction.Create)
 					{
-						PDBFileName = "PCH." + ModuleName + "." + Path.GetFileName(CompileEnvironment.Config.PrecompiledHeaderIncludeFilename);
+						PDBFileName = "PCH." + ModuleName + "." + CompileEnvironment.Config.PrecompiledHeaderIncludeFilename.GetFileName();
 						bActionProducesPDB = true;
 					}
 					// Ungrouped C++ files use a PDB per file.
@@ -689,8 +689,8 @@ namespace UnrealBuildTool
 					}
 
 					// Specify the PDB file that the compiler should write to.
-					FileItem PDBFile = FileItem.GetItemByPath(
-							Path.Combine(
+					FileItem PDBFile = FileItem.GetItemByFileReference(
+							FileReference.Combine(
 								CompileEnvironment.Config.OutputDirectory,
 								PDBFileName + ".pdb"
 								)
@@ -718,7 +718,7 @@ namespace UnrealBuildTool
 					AppendCLArguments_CPP(CompileEnvironment, FileArguments);
 				}
 
-				CompileAction.WorkingDirectory = Path.GetFullPath(".");
+				CompileAction.WorkingDirectory = UnrealBuildTool.EngineSourceDirectory.FullName;
 				CompileAction.CommandPath      = EnvVars.CompilerPath;
 				CompileAction.CommandArguments = Arguments.ToString() + FileArguments.ToString() + CompileEnvironment.Config.AdditionalArguments;
 
@@ -763,7 +763,7 @@ namespace UnrealBuildTool
 			{
 				Action CompileAction = new Action(ActionType.Compile);
 				CompileAction.CommandDescription = "Resource";
-				CompileAction.WorkingDirectory = Path.GetFullPath(".");
+				CompileAction.WorkingDirectory = UnrealBuildTool.EngineSourceDirectory.FullName;
 				CompileAction.CommandPath = EnvVars.ResourceCompilerPath;
 				CompileAction.StatusDescription = Path.GetFileName(RCFile.AbsolutePath);
 
@@ -796,8 +796,8 @@ namespace UnrealBuildTool
 				}
 
 				// Add the RES file to the produced item list.
-				FileItem CompiledResourceFile = FileItem.GetItemByPath(
-					Path.Combine(
+				FileItem CompiledResourceFile = FileItem.GetItemByFileReference(
+					FileReference.Combine(
 						Environment.Config.OutputDirectory,
 						Path.GetFileName(RCFile.AbsolutePath) + ".res"
 						)
@@ -822,7 +822,7 @@ namespace UnrealBuildTool
 
 			if (LinkEnvironment.Config.bIsBuildingDotNetAssembly)
 			{
-				return FileItem.GetItemByPath(LinkEnvironment.Config.OutputFilePath);
+				return FileItem.GetItemByFileReference(LinkEnvironment.Config.OutputFilePath);
 			}
 
 			bool bIsBuildingLibrary = LinkEnvironment.Config.bIsBuildingLibrary || bBuildImportLibraryOnly;
@@ -848,7 +848,7 @@ namespace UnrealBuildTool
 				Arguments.Append(" /DEF");
 
 				// Ensure that the import library references the correct filename for the linked binary.
-				Arguments.AppendFormat(" /NAME:\"{0}\"", Path.GetFileName(LinkEnvironment.Config.OutputFilePath));
+				Arguments.AppendFormat(" /NAME:\"{0}\"", LinkEnvironment.Config.OutputFilePath.GetFileName());
 			}
 
 
@@ -883,8 +883,8 @@ namespace UnrealBuildTool
 			// file is not needed for our builds, but there is no way to prevent MSVC from generating it when
 			// linking targets that have exports.  We don't want this to clobber our LIB file and invalidate the
 			// existing timstamp, so instead we simply emit it with a different name
-			string ImportLibraryFilePath = Path.Combine( LinkEnvironment.Config.IntermediateDirectory, 
-														 Path.GetFileNameWithoutExtension(LinkEnvironment.Config.OutputFilePath) + ".lib" );
+			FileReference ImportLibraryFilePath = FileReference.Combine( LinkEnvironment.Config.IntermediateDirectory, 
+														 LinkEnvironment.Config.OutputFilePath.GetFileNameWithoutExtension() + ".lib" );
 
 			if (LinkEnvironment.Config.bIsCrossReferenced && !bBuildImportLibraryOnly)
 			{
@@ -894,11 +894,11 @@ namespace UnrealBuildTool
 			FileItem OutputFile;
 			if (bBuildImportLibraryOnly)
 			{
-				OutputFile = FileItem.GetItemByPath(ImportLibraryFilePath);
+				OutputFile = FileItem.GetItemByFileReference(ImportLibraryFilePath);
 			}
 			else
 			{
-				OutputFile = FileItem.GetItemByPath(LinkEnvironment.Config.OutputFilePath);
+				OutputFile = FileItem.GetItemByFileReference(LinkEnvironment.Config.OutputFilePath);
 				OutputFile.bNeedsHotReloadNumbersDLLCleanUp = LinkEnvironment.Config.bIsBuildingDLL;
 			}
 
@@ -942,7 +942,7 @@ namespace UnrealBuildTool
 			}
 
 			// Create a response file for the linker
-			string ResponseFileName = GetResponseFileName( LinkEnvironment, OutputFile );
+			FileReference ResponseFileName = GetResponseFileName( LinkEnvironment, OutputFile );
 
 			// Never create response files when we are only generating IntelliSense data
 			if( !ProjectFileGenerator.bGenerateProjectFiles )
@@ -957,8 +957,8 @@ namespace UnrealBuildTool
 			if (bBuildImportLibraryOnly || (LinkEnvironment.Config.bHasExports && !bIsBuildingLibrary))
 			{
 				// An export file is written to the output directory implicitly; add it to the produced items list.
-				string ExportFilePath = Path.ChangeExtension(ImportLibraryFilePath, ".exp");
-				FileItem ExportFile = FileItem.GetItemByPath(ExportFilePath);
+				FileReference ExportFilePath = ImportLibraryFilePath.ChangeExtension(".exp");
+				FileItem ExportFile = FileItem.GetItemByFileReference(ExportFilePath);
 				ProducedItems.Add(ExportFile);
 			}
 
@@ -969,7 +969,7 @@ namespace UnrealBuildTool
 					(LinkEnvironment.Config.Target.Configuration != CPPTargetConfiguration.Shipping))
 				{
 					// Write the import library to the output directory for nFringe support.
-					FileItem ImportLibraryFile = FileItem.GetItemByPath(ImportLibraryFilePath);
+					FileItem ImportLibraryFile = FileItem.GetItemByFileReference(ImportLibraryFilePath);
 					Arguments.AppendFormat(" /IMPLIB:\"{0}\"", ImportLibraryFilePath);
 					ProducedItems.Add(ImportLibraryFile);
 				}
@@ -977,8 +977,8 @@ namespace UnrealBuildTool
 				if (LinkEnvironment.Config.bCreateDebugInfo)
 				{
 					// Write the PDB file to the output directory.			
-					string PDBFilePath = Path.Combine(LinkEnvironment.Config.OutputDirectory, Path.GetFileNameWithoutExtension(OutputFile.AbsolutePath) + ".pdb");
-					FileItem PDBFile = FileItem.GetItemByPath(PDBFilePath);
+					FileReference PDBFilePath = FileReference.Combine(LinkEnvironment.Config.OutputDirectory, Path.GetFileNameWithoutExtension(OutputFile.AbsolutePath) + ".pdb");
+					FileItem PDBFile = FileItem.GetItemByFileReference(PDBFilePath);
 					Arguments.AppendFormat(" /PDB:\"{0}\"", PDBFilePath);
 					ProducedItems.Add(PDBFile);
 
@@ -987,7 +987,7 @@ namespace UnrealBuildTool
 					if (true)
 					{
 						string MAPFilePath = Path.Combine(LinkEnvironment.Config.OutputDirectory, Path.GetFileNameWithoutExtension(OutputFile.AbsolutePath) + ".map");
-						FileItem MAPFile = FileItem.GetItemByPath(MAPFilePath);
+						FileItem MAPFile = FileItem.GetItemByFileReference(MAPFilePath);
 						LinkAction.CommandArguments += string.Format(" /MAP:\"{0}\"", MAPFilePath);
 						LinkAction.ProducedItems.Add(MAPFile);
 					}
@@ -1001,7 +1001,7 @@ namespace UnrealBuildTool
 			// Create an action that invokes the linker.
 			Action LinkAction = new Action(ActionType.Link);
 			LinkAction.CommandDescription = "Link";
-			LinkAction.WorkingDirectory   = Path.GetFullPath(".");
+			LinkAction.WorkingDirectory   = UnrealBuildTool.EngineSourceDirectory.FullName;
 			LinkAction.CommandPath        = bIsBuildingLibrary ? EnvVars.LibraryLinkerPath : EnvVars.LinkerPath;
 			LinkAction.CommandArguments   = Arguments.ToString();
 			LinkAction.ProducedItems    .AddRange(ProducedItems);
diff --git a/Engine/Source/Programs/UnrealBuildTool/Utilities/JsonObject.cs b/Engine/Source/Programs/UnrealBuildTool/Utilities/JsonObject.cs
--- a/Engine/Source/Programs/UnrealBuildTool/Utilities/JsonObject.cs
+++ b/Engine/Source/Programs/UnrealBuildTool/Utilities/JsonObject.cs
@@ -8,7 +8,7 @@ namespace UnrealBuildTool
 {
 	public class JsonParseException : Exception
 	{
-		public JsonParseException(string Format, params string[] Args) : base(String.Format(Format, Args))
+		public JsonParseException(string Format, params object[] Args) : base(String.Format(Format, Args))
 		{
 		}
 	}
diff --git a/Engine/Source/Programs/UnrealBuildTool/WinRT/WinRTToolChain.cs b/Engine/Source/Programs/UnrealBuildTool/WinRT/WinRTToolChain.cs
--- a/Engine/Source/Programs/UnrealBuildTool/WinRT/WinRTToolChain.cs
+++ b/Engine/Source/Programs/UnrealBuildTool/WinRT/WinRTToolChain.cs
@@ -318,7 +318,7 @@ namespace UnrealBuildTool
 				Result += " /APPCONTAINER";
 
 				// location of metadata
-				Result += string.Format(" /WINMDFILE:\"{0}\"", Path.ChangeExtension(LinkEnvironment.Config.OutputFilePath, "winmd"));
+				Result += string.Format(" /WINMDFILE:\"{0}\"", LinkEnvironment.Config.OutputFilePath.ChangeExtension("winmd"));
 			}
 
 			if (LinkEnvironment.Config.Target.Platform == CPPTargetPlatform.WinRT_ARM)
@@ -515,8 +515,8 @@ namespace UnrealBuildTool
 				if (CompileEnvironment.Config.PrecompiledHeaderAction == PrecompiledHeaderAction.Create)
 				{
 					// Generate a CPP File that just includes the precompiled header.
-					string PCHCPPFilename = "PCH." + Path.GetFileName(CompileEnvironment.Config.PrecompiledHeaderIncludeFilename) + ".cpp";
-					string PCHCPPPath = Path.Combine(CompileEnvironment.Config.OutputDirectory, PCHCPPFilename);
+					string PCHCPPFilename = "PCH." + CompileEnvironment.Config.PrecompiledHeaderIncludeFilename.GetFileName() + ".cpp";
+					FileReference PCHCPPPath = FileReference.Combine(CompileEnvironment.Config.OutputDirectory, PCHCPPFilename);
 					FileItem PCHCPPFile = FileItem.CreateIntermediateTextFile(
 						PCHCPPPath,
 						string.Format("#include \"{0}\"\r\n", CompileEnvironment.Config.PrecompiledHeaderIncludeFilename)
@@ -530,8 +530,8 @@ namespace UnrealBuildTool
 					var PCHExtension = UEBuildPlatform.BuildPlatformDictionary[UnrealTargetPlatform.WinRT].GetBinaryExtension(UEBuildBinaryType.PrecompiledHeader);
 
 					// Add the precompiled header file to the produced items list.
-					FileItem PrecompiledHeaderFile = FileItem.GetItemByPath(
-						Path.Combine(
+					FileItem PrecompiledHeaderFile = FileItem.GetItemByFileReference(
+						FileReference.Combine(
 							CompileEnvironment.Config.OutputDirectory,
 							Path.GetFileName(SourceFile.AbsolutePath) + PCHExtension
 							)
@@ -552,7 +552,7 @@ namespace UnrealBuildTool
 					{
 						// NOTE: The symbol name we use here is arbitrary, and all that matters is that it is
 						// unique per PCH module used in our library
-						string FakeUniquePCHSymbolName = Path.GetFileNameWithoutExtension(CompileEnvironment.Config.PrecompiledHeaderIncludeFilename);
+						string FakeUniquePCHSymbolName = CompileEnvironment.Config.PrecompiledHeaderIncludeFilename.GetFileNameWithoutExtension();
 						FileArguments += string.Format(" /Yl{0}", FakeUniquePCHSymbolName);
 					}
 					CompileAction.StatusDescription = PCHCPPFilename;
@@ -587,8 +587,8 @@ namespace UnrealBuildTool
 				var ObjectFileExtension = UEBuildPlatform.BuildPlatformDictionary[UnrealTargetPlatform.WinRT].GetBinaryExtension(UEBuildBinaryType.Object);
 
 				// Add the object file to the produced item list.
-				FileItem ObjectFile = FileItem.GetItemByPath(
-					Path.Combine(
+				FileItem ObjectFile = FileItem.GetItemByFileReference(
+					FileReference.Combine(
 						CompileEnvironment.Config.OutputDirectory,
 						Path.GetFileName(SourceFile.AbsolutePath) + ObjectFileExtension
 						)
@@ -608,7 +608,7 @@ namespace UnrealBuildTool
 					// All files using the same PCH are required to share a PDB.
 					if (CompileEnvironment.Config.PrecompiledHeaderAction == PrecompiledHeaderAction.Include)
 					{
-						PDBFileName = Path.GetFileName(CompileEnvironment.Config.PrecompiledHeaderIncludeFilename);
+						PDBFileName = CompileEnvironment.Config.PrecompiledHeaderIncludeFilename.GetFileName();
 					}
 					// Files creating a PCH or ungrouped C++ files use a PDB per file.
 					else if (CompileEnvironment.Config.PrecompiledHeaderAction == PrecompiledHeaderAction.Create || !bIsPlainCFile)
@@ -623,8 +623,8 @@ namespace UnrealBuildTool
 					}
 
 					// Specify the PDB file that the compiler should write to.
-					FileItem PDBFile = FileItem.GetItemByPath(
-							Path.Combine(
+					FileItem PDBFile = FileItem.GetItemByFileReference(
+							FileReference.Combine(
 								CompileEnvironment.Config.OutputDirectory,
 								PDBFileName + ".pdb"
 								)
@@ -652,7 +652,7 @@ namespace UnrealBuildTool
 					FileArguments += GetCLArguments_CPP(CompileEnvironment);
 				}
 
-				CompileAction.WorkingDirectory = Path.GetFullPath(".");
+				CompileAction.WorkingDirectory = UnrealBuildTool.EngineSourceDirectory.FullName;
 				CompileAction.CommandPath = GetVCToolPath(CompileEnvironment.Config.Target.Platform, CompileEnvironment.Config.Target.Configuration, "cl");
 				CompileAction.CommandArguments = Arguments + FileArguments + CompileEnvironment.Config.AdditionalArguments;
 				CompileAction.StatusDescription = string.Format("{0}", Path.GetFileName(SourceFile.AbsolutePath));
@@ -673,7 +673,7 @@ namespace UnrealBuildTool
 		{
 			if (LinkEnvironment.Config.bIsBuildingDotNetAssembly)
 			{
-				return FileItem.GetItemByPath(LinkEnvironment.Config.OutputFilePath);
+				return FileItem.GetItemByFileReference(LinkEnvironment.Config.OutputFilePath);
 			}
 
 			bool bIsBuildingLibrary = LinkEnvironment.Config.bIsBuildingLibrary || bBuildImportLibraryOnly;
@@ -681,7 +681,7 @@ namespace UnrealBuildTool
 
 			// Create an action that invokes the linker.
 			Action LinkAction = new Action(ActionType.Link);
-			LinkAction.WorkingDirectory = Path.GetFullPath(".");
+			LinkAction.WorkingDirectory = UnrealBuildTool.EngineSourceDirectory.FullName;
 			LinkAction.CommandPath = GetVCToolPath(
 				LinkEnvironment.Config.Target.Platform,
 				LinkEnvironment.Config.Target.Configuration,
@@ -703,7 +703,7 @@ namespace UnrealBuildTool
 				LinkAction.CommandArguments += " /DEF";
 
 				// Ensure that the import library references the correct filename for the linked binary.
-				LinkAction.CommandArguments += string.Format(" /NAME:\"{0}\"", Path.GetFileName(LinkEnvironment.Config.OutputFilePath));
+				LinkAction.CommandArguments += string.Format(" /NAME:\"{0}\"", LinkEnvironment.Config.OutputFilePath.GetFileName());
 			}
 
 			if (!LinkEnvironment.Config.bIsBuildingLibrary || (LinkEnvironment.Config.bIsBuildingLibrary && bIncludeDependentLibrariesInLibrary))
@@ -726,8 +726,8 @@ namespace UnrealBuildTool
 			// file is not needed for our builds, but there is no way to prevent MSVC from generating it when
 			// linking targets that have exports.  We don't want this to clobber our LIB file and invalidate the
 			// existing timstamp, so instead we simply emit it with a different name
-			string ImportLibraryFilePath = Path.Combine(LinkEnvironment.Config.IntermediateDirectory,
-														 Path.GetFileNameWithoutExtension(LinkEnvironment.Config.OutputFilePath) + ".lib");
+			FileReference ImportLibraryFilePath = FileReference.Combine(LinkEnvironment.Config.IntermediateDirectory,
+														 LinkEnvironment.Config.OutputFilePath.GetFileNameWithoutExtension() + ".lib");
 
 			if (LinkEnvironment.Config.bIsCrossReferenced && !bBuildImportLibraryOnly)
 			{
@@ -737,11 +737,11 @@ namespace UnrealBuildTool
 			FileItem OutputFile;
 			if (bBuildImportLibraryOnly)
 			{
-				OutputFile = FileItem.GetItemByPath(ImportLibraryFilePath);
+				OutputFile = FileItem.GetItemByFileReference(ImportLibraryFilePath);
 			}
 			else
 			{
-				OutputFile = FileItem.GetItemByPath(LinkEnvironment.Config.OutputFilePath);
+				OutputFile = FileItem.GetItemByFileReference(LinkEnvironment.Config.OutputFilePath);
 				OutputFile.bNeedsHotReloadNumbersDLLCleanUp = LinkEnvironment.Config.bIsBuildingDLL;
 			}
 			LinkAction.ProducedItems.Add(OutputFile);
@@ -782,7 +782,7 @@ namespace UnrealBuildTool
 			}
 
 			// Create a response file for the linker
-			string ResponseFileName = GetResponseFileName(LinkEnvironment, OutputFile);
+			FileReference ResponseFileName = GetResponseFileName(LinkEnvironment, OutputFile);
 
 			// Never create response files when we are only generating IntelliSense data
 			if (!ProjectFileGenerator.bGenerateProjectFiles)
@@ -797,8 +797,8 @@ namespace UnrealBuildTool
 			if (bBuildImportLibraryOnly || (LinkEnvironment.Config.bHasExports && !bIsBuildingLibrary))
 			{
 				// An export file is written to the output directory implicitly; add it to the produced items list.
-				string ExportFilePath = Path.ChangeExtension(ImportLibraryFilePath, ".exp");
-				FileItem ExportFile = FileItem.GetItemByPath(ExportFilePath);
+				FileReference ExportFilePath = ImportLibraryFilePath.ChangeExtension(".exp");
+				FileItem ExportFile = FileItem.GetItemByFileReference(ExportFilePath);
 				LinkAction.ProducedItems.Add(ExportFile);
 			}
 
@@ -810,7 +810,7 @@ namespace UnrealBuildTool
 					&& (!((LinkEnvironment.Config.Target.Configuration == CPPTargetConfiguration.Shipping) && (LinkEnvironment.bShouldCompileMonolithic != false))))
 				{
 					// Write the import library to the output directory for nFringe support.
-					FileItem ImportLibraryFile = FileItem.GetItemByPath(ImportLibraryFilePath);
+					FileItem ImportLibraryFile = FileItem.GetItemByFileReference(ImportLibraryFilePath);
 					LinkAction.CommandArguments += string.Format(" /IMPLIB:\"{0}\"", ImportLibraryFilePath);
 					LinkAction.ProducedItems.Add(ImportLibraryFile);
 				}
@@ -819,8 +819,8 @@ namespace UnrealBuildTool
 				{
 					// Write the PDB file to the output directory.			
 					{
-						string PDBFilePath = Path.Combine(LinkEnvironment.Config.OutputDirectory, Path.GetFileNameWithoutExtension(OutputFile.AbsolutePath) + ".pdb");
-						FileItem PDBFile = FileItem.GetItemByPath(PDBFilePath);
+						FileReference PDBFilePath = FileReference.Combine(LinkEnvironment.Config.OutputDirectory, Path.GetFileNameWithoutExtension(OutputFile.AbsolutePath) + ".pdb");
+						FileItem PDBFile = FileItem.GetItemByFileReference(PDBFilePath);
 						LinkAction.CommandArguments += string.Format(" /PDB:\"{0}\"", PDBFilePath);
 						LinkAction.ProducedItems.Add(PDBFile);
 					}
@@ -829,8 +829,8 @@ namespace UnrealBuildTool
 #if false					
 					if (true)
 					{
-						string MAPFilePath = Path.Combine(LinkEnvironment.Config.OutputDirectory, Path.GetFileNameWithoutExtension(OutputFile.AbsolutePath) + ".map");
-						FileItem MAPFile = FileItem.GetItemByPath(MAPFilePath);
+						FileReference MAPFilePath = FileReference.Combine(LinkEnvironment.Config.OutputDirectory, Path.GetFileNameWithoutExtension(OutputFile.AbsolutePath) + ".map");
+						FileItem MAPFile = FileItem.GetItemByFileReference(MAPFilePath);
 						LinkAction.CommandArguments += string.Format(" /MAP:\"{0}\"", MAPFilePath);
 						LinkAction.ProducedItems.Add(MAPFile);
 					}
diff --git a/Engine/Source/Programs/UnrealBuildTool/Windows/UEBuildWindows.cs b/Engine/Source/Programs/UnrealBuildTool/Windows/UEBuildWindows.cs
--- a/Engine/Source/Programs/UnrealBuildTool/Windows/UEBuildWindows.cs
+++ b/Engine/Source/Programs/UnrealBuildTool/Windows/UEBuildWindows.cs
@@ -400,7 +400,7 @@ namespace UnrealBuildTool
             return false;
         }
 
-		public override bool HasDefaultBuildConfig(UnrealTargetPlatform Platform, string ProjectPath)
+		public override bool HasDefaultBuildConfig(UnrealTargetPlatform Platform, DirectoryReference ProjectPath)
 		{
 			if (Platform == UnrealTargetPlatform.Win32)
 			{
diff --git a/Engine/Source/Programs/UnrealBuildTool/Windows/VCToolChain.cs b/Engine/Source/Programs/UnrealBuildTool/Windows/VCToolChain.cs
--- a/Engine/Source/Programs/UnrealBuildTool/Windows/VCToolChain.cs
+++ b/Engine/Source/Programs/UnrealBuildTool/Windows/VCToolChain.cs
@@ -1114,8 +1114,8 @@ namespace UnrealBuildTool
 				if (CompileEnvironment.Config.PrecompiledHeaderAction == PrecompiledHeaderAction.Create)
 				{
 					// Generate a CPP File that just includes the precompiled header.
-					string PCHCPPFilename = "PCH." + Path.GetFileName(CompileEnvironment.Config.PrecompiledHeaderIncludeFilename) + ".cpp";
-					string PCHCPPPath = Path.Combine(CompileEnvironment.Config.OutputDirectory, PCHCPPFilename);
+					string PCHCPPFilename = "PCH." + CompileEnvironment.Config.PrecompiledHeaderIncludeFilename.GetFileName() + ".cpp";
+					FileReference PCHCPPPath = FileReference.Combine(CompileEnvironment.Config.OutputDirectory, PCHCPPFilename);
 					FileItem PCHCPPFile = FileItem.CreateIntermediateTextFile(
 						PCHCPPPath,
 						string.Format("#include \"{0}\"\r\n", CompileEnvironment.Config.PrecompiledHeaderIncludeFilename)
@@ -1128,10 +1128,10 @@ namespace UnrealBuildTool
 
 					var PrecompiledFileExtension = UEBuildPlatform.BuildPlatformDictionary[UnrealTargetPlatform.Win64].GetBinaryExtension(UEBuildBinaryType.PrecompiledHeader);
 					// Add the precompiled header file to the produced items list.
-					FileItem PrecompiledHeaderFile = FileItem.GetItemByPath(
-						Path.Combine(
+					FileItem PrecompiledHeaderFile = FileItem.GetItemByFileReference(
+						FileReference.Combine(
 							CompileEnvironment.Config.OutputDirectory,
-						Path.GetFileName(SourceFile.AbsolutePath) + PrecompiledFileExtension
+							Path.GetFileName(SourceFile.AbsolutePath) + PrecompiledFileExtension
 							)
 						);
 					CompileAction.ProducedItems.Add(PrecompiledHeaderFile);
@@ -1151,7 +1151,7 @@ namespace UnrealBuildTool
 						{
 							// NOTE: The symbol name we use here is arbitrary, and all that matters is that it is
 							// unique per PCH module used in our library
-							string FakeUniquePCHSymbolName = Path.GetFileNameWithoutExtension(CompileEnvironment.Config.PrecompiledHeaderIncludeFilename);
+							string FakeUniquePCHSymbolName = CompileEnvironment.Config.PrecompiledHeaderIncludeFilename.GetFileNameWithoutExtension();
 							FileArguments.AppendFormat(" /Yl{0}", FakeUniquePCHSymbolName);
 						}
 					}
@@ -1222,8 +1222,8 @@ namespace UnrealBuildTool
 				{
 					var ObjectFileExtension = UEBuildPlatform.BuildPlatformDictionary[UnrealTargetPlatform.Win64].GetBinaryExtension(UEBuildBinaryType.Object);
 					// Add the object file to the produced item list.
-					FileItem ObjectFile = FileItem.GetItemByPath(
-						Path.Combine(
+					FileItem ObjectFile = FileItem.GetItemByFileReference(
+						FileReference.Combine(
 							CompileEnvironment.Config.OutputDirectory,
 							Path.GetFileName(SourceFile.AbsolutePath) + ObjectFileExtension
 							)
@@ -1258,12 +1258,12 @@ namespace UnrealBuildTool
 					// All files using the same PCH are required to share the same PDB that was used when compiling the PCH
 					if (CompileEnvironment.Config.PrecompiledHeaderAction == PrecompiledHeaderAction.Include)
 					{
-						PDBFileName = "PCH." + Path.GetFileName(CompileEnvironment.Config.PrecompiledHeaderIncludeFilename);
+						PDBFileName = "PCH." + CompileEnvironment.Config.PrecompiledHeaderIncludeFilename.GetFileName();
 					}
 					// Files creating a PCH use a PDB per file.
 					else if (CompileEnvironment.Config.PrecompiledHeaderAction == PrecompiledHeaderAction.Create)
 					{
-						PDBFileName = "PCH." + Path.GetFileName(CompileEnvironment.Config.PrecompiledHeaderIncludeFilename);
+						PDBFileName = "PCH." + CompileEnvironment.Config.PrecompiledHeaderIncludeFilename.GetFileName();
 						bActionProducesPDB = true;
 					}
 					// Ungrouped C++ files use a PDB per file.
@@ -1280,8 +1280,8 @@ namespace UnrealBuildTool
 
 					{ 
 						// Specify the PDB file that the compiler should write to.
-						FileItem PDBFile = FileItem.GetItemByPath(
-								Path.Combine(
+						FileItem PDBFile = FileItem.GetItemByFileReference(
+								FileReference.Combine(
 									CompileEnvironment.Config.OutputDirectory,
 									PDBFileName + ".pdb"
 									)
@@ -1310,7 +1310,7 @@ namespace UnrealBuildTool
 					AppendCLArguments_CPP( CompileEnvironment, FileArguments );
 				}
 
-				CompileAction.WorkingDirectory = Path.GetFullPath(".");
+				CompileAction.WorkingDirectory = UnrealBuildTool.EngineSourceDirectory.FullName;
 				if (BuildConfiguration.bRunUnrealCodeAnalyzer)
 				{
 					CompileAction.CommandPath = System.IO.Path.Combine(CompileAction.WorkingDirectory, @"..", @"Binaries", @"Win32", @"UnrealCodeAnalyzer.exe");
@@ -1325,7 +1325,7 @@ namespace UnrealBuildTool
 				{
 					string UCAMode = BuildConfiguration.bUCACheckUObjectThreadSafety ? @"-CheckThreadSafety " : @"-CreateIncludeFiles ";
 					var ObjectFileExtension = BuildConfiguration.bUCACheckUObjectThreadSafety ? @".tsc" : @".includes";
-					FileItem ObjectFile = FileItem.GetItemByPath(Path.Combine(CompileEnvironment.Config.OutputDirectory, Path.GetFileName(SourceFile.AbsolutePath) + ObjectFileExtension));
+					FileItem ObjectFile = FileItem.GetItemByFileReference(FileReference.Combine(CompileEnvironment.Config.OutputDirectory, Path.GetFileName(SourceFile.AbsolutePath) + ObjectFileExtension));
 					var ClangPath = System.IO.Path.Combine(CompileAction.WorkingDirectory, @"ThirdParty", @"llvm", @"3.5.0", @"bin", @"vs2013", @"x86", @"release", @"clang++.exe");
 					UnrealCodeAnalyzerArguments = UCAMode + SourceFile.AbsolutePath + @" -OutputFile=""" + ObjectFile.AbsolutePath + @""" -- " + ClangPath + @" --driver-mode=cl ";
 				}
@@ -1387,7 +1387,7 @@ namespace UnrealBuildTool
 			{
 				Action CompileAction = new Action(ActionType.Compile);
 				CompileAction.CommandDescription = "Resource";
-				CompileAction.WorkingDirectory   = Path.GetFullPath(".");
+				CompileAction.WorkingDirectory   = UnrealBuildTool.EngineSourceDirectory.FullName;
 				CompileAction.CommandPath        = EnvVars.ResourceCompilerPath;
 				CompileAction.StatusDescription  = Path.GetFileName(RCFile.AbsolutePath);
 
@@ -1440,8 +1440,8 @@ namespace UnrealBuildTool
 				}
 
 				// Add the RES file to the produced item list.
-				FileItem CompiledResourceFile = FileItem.GetItemByPath(
-					Path.Combine(
+				FileItem CompiledResourceFile = FileItem.GetItemByFileReference(
+					FileReference.Combine(
 						Environment.Config.OutputDirectory,
 						Path.GetFileName(RCFile.AbsolutePath) + ".res"
 						)
@@ -1481,7 +1481,7 @@ namespace UnrealBuildTool
 
 			if (LinkEnvironment.Config.bIsBuildingDotNetAssembly)
 			{
-				return FileItem.GetItemByPath(LinkEnvironment.Config.OutputFilePath);
+				return FileItem.GetItemByFileReference(LinkEnvironment.Config.OutputFilePath);
 			}
 
 			bool bIsBuildingLibrary = LinkEnvironment.Config.bIsBuildingLibrary || bBuildImportLibraryOnly;
@@ -1510,7 +1510,7 @@ namespace UnrealBuildTool
 				Arguments.Append(" /DEF");
 
 				// Ensure that the import library references the correct filename for the linked binary.
-				Arguments.AppendFormat(" /NAME:\"{0}\"", Path.GetFileName(LinkEnvironment.Config.OutputFilePath));
+				Arguments.AppendFormat(" /NAME:\"{0}\"", LinkEnvironment.Config.OutputFilePath.GetFileName());
 			}
 
 
@@ -1545,8 +1545,8 @@ namespace UnrealBuildTool
 			// file is not needed for our builds, but there is no way to prevent MSVC from generating it when
 			// linking targets that have exports.  We don't want this to clobber our LIB file and invalidate the
 			// existing timstamp, so instead we simply emit it with a different name
-			string ImportLibraryFilePath = Path.Combine( LinkEnvironment.Config.IntermediateDirectory, 
-														 Path.GetFileNameWithoutExtension(LinkEnvironment.Config.OutputFilePath) + ".lib" );
+			FileReference ImportLibraryFilePath = FileReference.Combine( LinkEnvironment.Config.IntermediateDirectory, 
+														 LinkEnvironment.Config.OutputFilePath.GetFileNameWithoutExtension() + ".lib" );
 
 			if (LinkEnvironment.Config.bIsCrossReferenced && !bBuildImportLibraryOnly)
 			{
@@ -1556,11 +1556,11 @@ namespace UnrealBuildTool
 			FileItem OutputFile;
 			if (bBuildImportLibraryOnly)
 			{
-				OutputFile = FileItem.GetItemByPath(ImportLibraryFilePath);
+				OutputFile = FileItem.GetItemByFileReference(ImportLibraryFilePath);
 			}
 			else
 			{
-				OutputFile = FileItem.GetItemByPath(LinkEnvironment.Config.OutputFilePath);
+				OutputFile = FileItem.GetItemByFileReference(LinkEnvironment.Config.OutputFilePath);
 				OutputFile.bNeedsHotReloadNumbersDLLCleanUp = LinkEnvironment.Config.bIsBuildingDLL;
 			}
 
@@ -1604,7 +1604,7 @@ namespace UnrealBuildTool
 			}
 
 			// Create a response file for the linker
-			string ResponseFileName = GetResponseFileName( LinkEnvironment, OutputFile );
+			FileReference ResponseFileName = GetResponseFileName( LinkEnvironment, OutputFile );
 
 			// Never create response files when we are only generating IntelliSense data
 			if( !ProjectFileGenerator.bGenerateProjectFiles )
@@ -1619,8 +1619,8 @@ namespace UnrealBuildTool
 			if (bBuildImportLibraryOnly || (LinkEnvironment.Config.bHasExports && !bIsBuildingLibrary))
 			{
 				// An export file is written to the output directory implicitly; add it to the produced items list.
-				string ExportFilePath = Path.ChangeExtension(ImportLibraryFilePath, ".exp");
-				FileItem ExportFile = FileItem.GetItemByPath(ExportFilePath);
+				FileReference ExportFilePath = ImportLibraryFilePath.ChangeExtension(".exp");
+				FileItem ExportFile = FileItem.GetItemByFileReference(ExportFilePath);
 				ProducedItems.Add(ExportFile);
 			}
 
@@ -1632,7 +1632,7 @@ namespace UnrealBuildTool
 					&& (!((LinkEnvironment.Config.Target.Configuration == CPPTargetConfiguration.Shipping) && (LinkEnvironment.bShouldCompileMonolithic != false))))
 				{
 					// Write the import library to the output directory for nFringe support.
-					FileItem ImportLibraryFile = FileItem.GetItemByPath(ImportLibraryFilePath);
+					FileItem ImportLibraryFile = FileItem.GetItemByFileReference(ImportLibraryFilePath);
 					Arguments.AppendFormat(" /IMPLIB:\"{0}\"", ImportLibraryFilePath);
 					ProducedItems.Add(ImportLibraryFile);
 				}
@@ -1641,8 +1641,8 @@ namespace UnrealBuildTool
 				{
 					// Write the PDB file to the output directory.			
 					{
-						string PDBFilePath = Path.Combine(LinkEnvironment.Config.OutputDirectory, Path.GetFileNameWithoutExtension(OutputFile.AbsolutePath) + ".pdb");
-						FileItem PDBFile = FileItem.GetItemByPath(PDBFilePath);
+						FileReference PDBFilePath = FileReference.Combine(LinkEnvironment.Config.OutputDirectory, Path.GetFileNameWithoutExtension(OutputFile.AbsolutePath) + ".pdb");
+						FileItem PDBFile = FileItem.GetItemByFileReference(PDBFilePath);
 						Arguments.AppendFormat(" /PDB:\"{0}\"", PDBFilePath);
 						ProducedItems.Add(PDBFile);
 					}
@@ -1651,8 +1651,8 @@ namespace UnrealBuildTool
 #if false					
 					if (true)
 					{
-						string MAPFilePath = Path.Combine(LinkEnvironment.Config.OutputDirectory, Path.GetFileNameWithoutExtension(OutputFile.AbsolutePath) + ".map");
-						FileItem MAPFile = FileItem.GetItemByPath(MAPFilePath);
+						FileReference MAPFilePath = FileReference.Combine(LinkEnvironment.Config.OutputDirectory, Path.GetFileNameWithoutExtension(OutputFile.AbsolutePath) + ".map");
+						FileItem MAPFile = FileItem.GetItemByFileReference(MAPFilePath);
 						Arguments.AppendFormat(" /MAP:\"{0}\"", MAPFilePath);
 						LinkAction.ProducedItems.Add(MAPFile);
 					}
@@ -1667,7 +1667,7 @@ namespace UnrealBuildTool
 			// Create an action that invokes the linker.
 			Action LinkAction = new Action(ActionType.Link);
 			LinkAction.CommandDescription = "Link";
-			LinkAction.WorkingDirectory   = Path.GetFullPath(".");
+			LinkAction.WorkingDirectory   = UnrealBuildTool.EngineSourceDirectory.FullName;
 			LinkAction.CommandPath        = bIsBuildingLibrary ? EnvVars.LibraryLinkerPath : EnvVars.LinkerPath;
 			LinkAction.CommandArguments   = Arguments.ToString();
 			LinkAction.ProducedItems    .AddRange(ProducedItems);
@@ -1697,7 +1697,7 @@ namespace UnrealBuildTool
 			return OutputFile;
 		}
 
-		public override void CompileCSharpProject(CSharpEnvironment CompileEnvironment, string ProjectFileName, string DestinationFile)
+		public override void CompileCSharpProject(CSharpEnvironment CompileEnvironment, FileReference ProjectFileName, FileReference DestinationFile)
 		{
 			// Initialize environment variables required for spawned tools.
 			var EnvVars = VCEnvironment.SetEnvironment(CompileEnvironment.EnvironmentTargetPlatform);
@@ -1705,7 +1705,7 @@ namespace UnrealBuildTool
 			var BuildProjectAction = new Action(ActionType.BuildProject);
 
 			// Specify the source file (prerequisite) for the action
-			var ProjectFileItem = FileItem.GetExistingItemByPath(ProjectFileName);
+			var ProjectFileItem = FileItem.GetExistingItemByFileReference(ProjectFileName);
 			if (ProjectFileItem == null)
 			{
 				throw new BuildException("Expected C# project file {0} to exist.", ProjectFileName);
@@ -1713,13 +1713,13 @@ namespace UnrealBuildTool
 			
 			// Add the project and the files contained to the prerequisites.
 			BuildProjectAction.PrerequisiteItems.Add(ProjectFileItem);
-			var ProjectFile = new VCSharpProjectFile( Utils.MakePathRelativeTo( ProjectFileName, ProjectFileGenerator.MasterProjectRelativePath ) );
+			var ProjectFile = new VCSharpProjectFile( ProjectFileName );
 			var ProjectPreReqs = ProjectFile.GetCSharpDependencies();
-			var ProjectFolder = Path.GetDirectoryName(ProjectFileName);
+			var ProjectFolder = ProjectFileName.Directory;
 			foreach( string ProjectPreReqRelativePath in ProjectPreReqs )
 			{
-				string ProjectPreReqAbsolutePath = Path.Combine( ProjectFolder, ProjectPreReqRelativePath );
-				var ProjectPreReqFileItem = FileItem.GetExistingItemByPath(ProjectPreReqAbsolutePath);
+				FileReference ProjectPreReqAbsolutePath = FileReference.Combine( ProjectFolder, ProjectPreReqRelativePath );
+				var ProjectPreReqFileItem = FileItem.GetExistingItemByFileReference(ProjectPreReqAbsolutePath);
 				if( ProjectPreReqFileItem == null )
 				{
 					throw new BuildException("Expected C# dependency {0} to exist.", ProjectPreReqAbsolutePath);
@@ -1731,8 +1731,8 @@ namespace UnrealBuildTool
 			BuildProjectAction.bCanExecuteRemotely = false;
 
 			// Setup execution via MSBuild.
-			BuildProjectAction.WorkingDirectory  = Path.GetFullPath(".");
-			BuildProjectAction.StatusDescription = Path.GetFileName(ProjectFileName);
+			BuildProjectAction.WorkingDirectory  = UnrealBuildTool.EngineSourceDirectory.FullName;
+			BuildProjectAction.StatusDescription = ProjectFileName.GetFileName();
 			BuildProjectAction.CommandPath       = EnvVars.MSBuildPath;
 			if (CompileEnvironment.TargetConfiguration == CSharpTargetConfiguration.Debug)
 			{
@@ -1750,9 +1750,9 @@ namespace UnrealBuildTool
 			BuildProjectAction.CommandArguments += String.Format(" \"{0}\"", ProjectFileItem.AbsolutePath);
 
 			// Specify the output files.
-			string PDBFilePath = Path.Combine( Path.GetDirectoryName(DestinationFile), Path.GetFileNameWithoutExtension(DestinationFile) + ".pdb" );
-			FileItem PDBFile = FileItem.GetItemByPath( PDBFilePath );
-			BuildProjectAction.ProducedItems.Add( FileItem.GetItemByPath(DestinationFile) );		
+			FileReference PDBFilePath = FileReference.Combine( DestinationFile.Directory, DestinationFile.GetFileNameWithoutExtension() + ".pdb" );
+			FileItem PDBFile = FileItem.GetItemByFileReference( PDBFilePath );
+			BuildProjectAction.ProducedItems.Add( FileItem.GetItemByFileReference(DestinationFile) );		
 			BuildProjectAction.ProducedItems.Add( PDBFile );
 		}
 
@@ -1781,7 +1781,7 @@ namespace UnrealBuildTool
 		{
 			if (Binary.Config.Type == UEBuildBinaryType.DynamicLinkLibrary)
 			{
-				Receipt.AddBuildProduct(Path.Combine(Binary.Config.IntermediateDirectory, Path.GetFileNameWithoutExtension(Binary.Config.OutputFilePath) + ".lib"), BuildProductType.ImportLibrary);
+				Receipt.AddBuildProduct(FileReference.Combine(Binary.Config.IntermediateDirectory, Binary.Config.OutputFilePath.GetFileNameWithoutExtension() + ".lib"), BuildProductType.ImportLibrary);
 			}
 		}
 
