Rename FileSystemName.cs to FileSystemReference.cs.

diff --git a/Engine/Source/Programs/UnrealBuildTool/System/FileSystemNames.cs b/Engine/Source/Programs/UnrealBuildTool/System/FileSystemNames.cs
deleted file mode 100644
diff --git a/Engine/Source/Programs/UnrealBuildTool/System/FileSystemReference.cs b/Engine/Source/Programs/UnrealBuildTool/System/FileSystemReference.cs
new file mode 100644
--- /dev/null
+++ b/Engine/Source/Programs/UnrealBuildTool/System/FileSystemReference.cs
@@ -0,0 +1,630 @@
+﻿using System;
+using System.Collections.Generic;
+using System.IO;
+using System.Linq;
+using System.Text;
+using System.Threading.Tasks;
+
+namespace UnrealBuildTool
+{
+	/// <summary>
+	/// Base class for file system objects (files or directories).
+	/// </summary>
+	[Serializable]
+	public abstract class FileSystemReference
+	{
+		/// <summary>
+		/// The path to this object. Stored as an absolute path, with O/S preferred separator characters, and no trailing slash for directories.
+		/// </summary>
+		public readonly string FullName;
+
+		/// <summary>
+		/// The canonical full name for this object.
+		/// </summary>
+		public readonly string CanonicalName;
+
+		/// <summary>
+		/// Constructs a filesystem object for the given path.
+		/// </summary>
+		public FileSystemReference(string InPath)
+		{
+			FullName = Path.GetFullPath(InPath);
+			CanonicalName = FullName.ToLowerInvariant();
+		}
+
+		/// <summary>
+		/// Direct constructor for a path
+		/// </summary>
+		protected FileSystemReference(string InFullName, string InCanonicalName)
+		{
+			FullName = InFullName;
+			CanonicalName = InCanonicalName;
+		}
+
+		/// <summary>
+		/// Create a full path by concatenating multiple strings
+		/// </summary>
+		/// <returns></returns>
+		static protected string CombineStrings(DirectoryReference BaseDirectory, params string[] Fragments)
+		{
+			// Get the initial string to append to, and strip any root directory suffix from it
+			StringBuilder NewFullName = new StringBuilder(BaseDirectory.FullName);
+			if(NewFullName.Length > 0 && NewFullName[NewFullName.Length - 1] == Path.DirectorySeparatorChar)
+			{
+				NewFullName.Remove(NewFullName.Length - 1, 1);
+			}
+
+			// Scan through the fragments to append, appending them to a string and updating the base length as we go
+			foreach(string Fragment in Fragments)
+			{
+				// Check if this fragment is an absolute path
+				if((Fragment.Length >= 2 && Fragment[1] == ':') || (Fragment.Length >= 1 && (Fragment[0] == Path.DirectorySeparatorChar || Fragment[0] == Path.AltDirectorySeparatorChar)))
+				{
+					// It is. Reset the new name to the full version of this path.
+					NewFullName.Clear();
+					NewFullName.Append(Path.GetFullPath(Fragment));
+				}
+				else
+				{
+					// Append all the parts of this fragment to the end of the existing path.
+					int StartIdx = 0;
+					while(StartIdx < Fragment.Length)
+					{
+						// Find the end of this fragment. We may have been passed multiple paths in the same string.
+						int EndIdx = StartIdx;
+						while(EndIdx < Fragment.Length && Fragment[EndIdx] != Path.DirectorySeparatorChar && Fragment[EndIdx] != Path.AltDirectorySeparatorChar)
+						{
+							EndIdx++;
+						}
+
+						// Ignore any empty sections, like leading or trailing slashes, and '.' directory references.
+						int Length = EndIdx - StartIdx;
+						if(Length == 0)
+						{
+							// Multiple directory separators in a row; illegal.
+							throw new ArgumentException("Path fragment '{0}' contains invalid directory separators.");
+						}
+						else if(Length == 2 && Fragment[StartIdx] == '.' && Fragment[StartIdx + 1] == '.')
+						{
+							// Remove the last directory name
+							for(int SeparatorIdx = NewFullName.Length - 1; SeparatorIdx >= 0; SeparatorIdx--)
+							{
+								if(NewFullName[SeparatorIdx] == Path.DirectorySeparatorChar)
+								{
+									NewFullName.Remove(SeparatorIdx, NewFullName.Length - SeparatorIdx);
+									break;
+								}
+							}
+						}
+						else if(Length != 1 || Fragment[StartIdx] != '.')
+						{
+							// Append this fragment
+							NewFullName.Append(Path.DirectorySeparatorChar);
+							NewFullName.Append(Fragment, StartIdx, Length);
+						}
+
+						// Move to the next part
+						StartIdx = EndIdx + 1;
+					}
+				}
+			}
+
+			// Append the directory separator
+			if(NewFullName.Length == 0 || (NewFullName.Length == 2 && NewFullName[1] == ':'))
+			{
+				NewFullName.Append(Path.DirectorySeparatorChar);
+			}
+			
+			// Set the new path variables
+			return NewFullName.ToString();
+		}
+
+		/// <summary>
+		/// Checks whether this name has the given extension.
+		/// </summary>
+		/// <param name="Extension">The extension to check</param>
+		/// <returns>True if this name has the given extension, false otherwise</returns>
+		public bool HasExtension(string Extension)
+		{
+			if(Extension.Length > 0 && Extension[0] != '.')
+			{
+				return HasExtension("." + Extension);
+			}
+			else
+			{ 
+				return CanonicalName.EndsWith(Extension.ToLowerInvariant());
+			}
+		}
+
+		/// <summary>
+		/// Determines if the given object is at or under the given directory
+		/// </summary>
+		/// <param name="Directory"></param>
+		/// <returns></returns>
+		public bool IsUnderDirectory(DirectoryReference Other)
+		{
+			return CanonicalName.StartsWith(Other.CanonicalName) && (CanonicalName.Length == Other.CanonicalName.Length || CanonicalName[Other.CanonicalName.Length] == Path.DirectorySeparatorChar);
+		}
+
+		/// <summary>
+		/// Creates a relative path from the given base directory
+		/// </summary>
+		/// <param name="Directory">The directory to create a relative path from</param>
+		/// <returns>A relative path from the given directory</returns>
+		public string MakeRelativeTo(DirectoryReference Directory)
+		{
+			// Find how much of the path is common between the two paths. This length does not include a trailing directory separator character.
+			int CommonDirectoryLength = -1;
+			for (int Idx = 0;;Idx++)
+			{
+				if(Idx == CanonicalName.Length)
+				{
+					// The two paths are identical. Just return the "." character.
+					if(Idx == Directory.CanonicalName.Length)
+					{
+						return ".";
+					}
+
+					// Check if we're finishing on a complete directory name
+					if(Directory.CanonicalName[Idx] == Path.DirectorySeparatorChar)
+					{
+						CommonDirectoryLength = Idx;
+					}
+					break;
+				}
+				else if(Idx == Directory.CanonicalName.Length)
+				{
+					// Check whether the end of the directory name coincides with a boundary for the current name.
+					if(CanonicalName[Idx] == Path.DirectorySeparatorChar)
+					{
+						CommonDirectoryLength = Idx;
+					}
+					break;
+				}
+				else
+				{
+					// Check the two paths match, and bail if they don't. Increase the common directory length if we've reached a separator.
+					if (CanonicalName[Idx] != Directory.CanonicalName[Idx])
+					{
+						break;
+					}
+					if (CanonicalName[Idx] == Path.DirectorySeparatorChar)
+					{
+						CommonDirectoryLength = Idx;
+					}
+				}
+			}
+
+			// If there's no relative path, just return the absolute path
+			if(CommonDirectoryLength == -1)
+			{
+				return FullName;
+			}
+
+			// Append all the '..' separators to get back to the common directory, then the rest of the string to reach the target item
+			StringBuilder Result = new StringBuilder();
+			for(int Idx = CommonDirectoryLength + 1; Idx < Directory.CanonicalName.Length; Idx++)
+			{
+				// Move up a directory
+				Result.Append("..");
+				Result.Append(Path.DirectorySeparatorChar);
+
+				// Scan to the next directory separator
+				while(Idx < Directory.CanonicalName.Length && Directory.CanonicalName[Idx] != Path.DirectorySeparatorChar)
+				{
+					Idx++;
+				}
+			}
+			if(CommonDirectoryLength + 1 < FullName.Length)
+			{
+				Result.Append(FullName, CommonDirectoryLength + 1, FullName.Length - CommonDirectoryLength - 1);
+			}
+			return Result.ToString();
+		}
+
+		/// <summary>
+		/// Returns a string representation of this filesystem object
+		/// </summary>
+		/// <returns>Full path to the object</returns>
+		public override string ToString()
+		{
+			return FullName;
+		}
+	}
+
+	/// <summary>
+	/// Representation of an absolute directory path. Allows fast hashing and comparisons.
+	/// </summary>
+	[Serializable]
+	public class DirectoryReference : FileSystemReference
+	{
+		/// <summary>
+		/// Default constructor.
+		/// </summary>
+		/// <param name="InPath">Path to this directory.</param>
+		public DirectoryReference(string InPath)
+			: base(InPath)
+		{
+		}
+
+		/// <summary>
+		/// Constructor for creating a directory object directly from two strings.
+		/// </summary>
+		/// <param name="InFullName"></param>
+		/// <param name="InCanonicalName"></param>
+		protected DirectoryReference(string InFullName, string InCanonicalName)
+			: base(InFullName, InCanonicalName)
+		{
+		}
+
+		/// <summary>
+		/// Gets the top level directory name
+		/// </summary>
+		/// <returns>The name of the directory</returns>
+		public string GetDirectoryName()
+		{
+			return Path.GetFileName(FullName);
+		}
+
+		/// <summary>
+		/// Gets the directory containing this object
+		/// </summary>
+		/// <returns>A new directory object representing the directory containing this object</returns>
+		public DirectoryReference ParentDirectory
+		{
+			get
+			{
+				if (IsRootDirectory())
+				{
+					return null;
+				}
+
+				int ParentLength = CanonicalName.LastIndexOf(Path.DirectorySeparatorChar);
+				if (ParentLength == 2 && CanonicalName[1] == ':')
+				{
+					ParentLength++;
+				}
+
+				return new DirectoryReference(FullName.Substring(0, ParentLength), CanonicalName.Substring(0, ParentLength));
+			}
+		}
+
+		/// <summary>
+		/// Gets the parent directory for a file
+		/// </summary>
+		/// <param name="File">The file to get directory for</param>
+		/// <returns>The full directory name containing the given file</returns>
+		public static DirectoryReference GetParentDirectory(FileReference File)
+		{
+			int ParentLength = File.CanonicalName.LastIndexOf(Path.DirectorySeparatorChar);
+			return new DirectoryReference(File.FullName.Substring(0, ParentLength), File.CanonicalName.Substring(0, ParentLength));
+		}
+
+		/// <summary>
+		/// Creates the directory
+		/// </summary>
+		public void CreateDirectory()
+		{
+			Directory.CreateDirectory(FullName);
+		}
+
+		/// <summary>
+		/// Checks whether the directory exists
+		/// </summary>
+		/// <returns>True if this directory exists</returns>
+		public bool Exists()
+		{
+			return Directory.Exists(FullName);
+		}
+
+		/// <summary>
+		/// Determines whether this path represents a root directory in the filesystem
+		/// </summary>
+		/// <returns>True if this path is a root directory, false otherwise</returns>
+		public bool IsRootDirectory()
+		{
+			return CanonicalName[CanonicalName.Length - 1] == Path.DirectorySeparatorChar;
+		}
+
+		/// <summary>
+		/// Combine several fragments with a base directory, to form a new directory name
+		/// </summary>
+		/// <param name="BaseDirectory">The base directory</param>
+		/// <param name="Fragments">Fragments to combine with the base directory</param>
+		/// <returns>The new directory name</returns>
+		public static DirectoryReference Combine(DirectoryReference BaseDirectory, params string[] Fragments)
+		{
+			string FullName = FileSystemReference.CombineStrings(BaseDirectory, Fragments);
+			return new DirectoryReference(FullName, FullName.ToLowerInvariant());
+		}
+
+		/// <summary>
+		/// Compares two filesystem object names for equality. Uses the canonical name representation, not the display name representation.
+		/// </summary>
+		/// <param name="A">First object to compare.</param>
+		/// <param name="B">Second object to compare.</param>
+		/// <returns>True if the names represent the same object, false otherwise</returns>
+		public static bool operator ==(DirectoryReference A, DirectoryReference B)
+		{
+			if((object)A == null)
+			{
+				return (object)B == null;
+			}
+			else
+			{
+				return (object)B != null && A.CanonicalName == B.CanonicalName;
+			}
+		}
+
+		/// <summary>
+		/// Compares two filesystem object names for inequality. Uses the canonical name representation, not the display name representation.
+		/// </summary>
+		/// <param name="A">First object to compare.</param>
+		/// <param name="B">Second object to compare.</param>
+		/// <returns>False if the names represent the same object, true otherwise</returns>
+		public static bool operator !=(DirectoryReference A, DirectoryReference B)
+		{
+			return !(A == B);
+		}
+
+		/// <summary>
+		/// Compares against another object for equality.
+		/// </summary>
+		/// <param name="A">First name to compare.</param>
+		/// <param name="B">Second name to compare.</param>
+		/// <returns>True if the names represent the same object, false otherwise</returns>
+		public override bool Equals(object Obj)
+		{
+			return (Obj is DirectoryReference) && ((DirectoryReference)Obj) == this;
+		}
+
+		/// <summary>
+		/// Returns a hash code for this object
+		/// </summary>
+		/// <returns></returns>
+		public override int GetHashCode()
+		{
+			return CanonicalName.GetHashCode();
+		}
+
+		/// <summary>
+		/// Helper function to create a remote directory reference. Unlike normal DirectoryReference objects, these aren't converted to a full path in the local filesystem.
+		/// </summary>
+		/// <param name="AbsolutePath">The absolute path in the remote file system</param>
+		/// <returns>New directory reference</returns>
+		public static DirectoryReference MakeRemote(string AbsolutePath)
+		{
+			return new DirectoryReference(AbsolutePath, AbsolutePath.ToLowerInvariant());
+		}
+	}
+
+	/// <summary>
+	/// Representation of an absolute file path. Allows fast hashing and comparisons.
+	/// </summary>
+	[Serializable]
+	public class FileReference : FileSystemReference
+	{
+		/// <summary>
+		/// Default constructor.
+		/// </summary>
+		/// <param name="InPath">Path to this file</param>
+		public FileReference(string InPath)
+			: base(InPath)
+		{
+		}
+
+		/// <summary>
+		/// Default constructor.
+		/// </summary>
+		/// <param name="InPath">Path to this file</param>
+		protected FileReference(string InFullName, string InCanonicalName)
+			: base(InFullName, InCanonicalName)
+		{
+		}
+
+		/// <summary>
+		/// Gets the file name without path information
+		/// </summary>
+		/// <returns>A string containing the file name</returns>
+		public string GetFileName()
+		{
+			return Path.GetFileName(FullName);
+		}
+
+		/// <summary>
+		/// Gets the file name without path information or an extension
+		/// </summary>
+		/// <returns>A string containing the file name without an extension</returns>
+		public string GetFileNameWithoutExtension()
+		{
+			return Path.GetFileNameWithoutExtension(FullName);
+		}
+
+		/// <summary>
+		/// Gets the file name without path or any extensions
+		/// </summary>
+		/// <returns>A string containing the file name without an extension</returns>
+		public string GetFileNameWithoutAnyExtensions()
+		{
+			int StartIdx = FullName.LastIndexOf(Path.DirectorySeparatorChar) + 1;
+
+			int EndIdx = FullName.IndexOf('.', StartIdx);
+			if(EndIdx < StartIdx)
+			{
+				return FullName.Substring(StartIdx);
+			}
+			else
+			{
+				return FullName.Substring(StartIdx, EndIdx - StartIdx);
+			}
+		}
+
+		/// <summary>
+		/// Gets the extension for this filename
+		/// </summary>
+		/// <returns>A string containing the extension of this filename</returns>
+		public string GetExtension()
+		{
+			return Path.GetExtension(FullName);
+		}
+
+		/// <summary>
+		/// Change the file's extension to something else
+		/// </summary>
+		/// <param name="Extension">The new extension</param>
+		/// <returns>A FileReference with the same path and name, but with the new extension</returns>
+		public FileReference ChangeExtension(string Extension)
+		{
+			string NewFullName = Path.ChangeExtension(FullName, Extension);
+			return new FileReference(NewFullName, NewFullName.ToLowerInvariant());
+		}
+
+		/// <summary>
+		/// Gets the directory containing this file
+		/// </summary>
+		/// <returns>A new directory object representing the directory containing this object</returns>
+		public DirectoryReference Directory
+		{
+			get { return DirectoryReference.GetParentDirectory(this); }
+		}
+
+		/// <summary>
+		/// Determines whether the given filename exists
+		/// </summary>
+		/// <returns>True if it exists, false otherwise</returns>
+		public bool Exists()
+		{
+			return File.Exists(FullName);
+		}
+
+		/// <summary>
+		/// Deletes this file
+		/// </summary>
+		public void Delete()
+		{
+			File.Delete(FullName);
+		}
+
+		/// <summary>
+		/// Combine several fragments with a base directory, to form a new filename
+		/// </summary>
+		/// <param name="BaseDirectory">The base directory</param>
+		/// <param name="Fragments">Fragments to combine with the base directory</param>
+		/// <returns>The new file name</returns>
+		public static FileReference Combine(DirectoryReference BaseDirectory, params string[] Fragments)
+		{
+			string FullName = FileSystemReference.CombineStrings(BaseDirectory, Fragments);
+			return new FileReference(FullName, FullName.ToLowerInvariant());
+		}
+
+		/// <summary>
+		/// Append a string to the end of a filename
+		/// </summary>
+		/// <param name="A">The base file reference</param>
+		/// <param name="B">Suffix to be appended</param>
+		/// <returns>The new file reference</returns>
+		public static FileReference operator+(FileReference A, string B)
+		{
+			return new FileReference(A.FullName + B, A.CanonicalName + B.ToLowerInvariant());
+		}
+
+		/// <summary>
+		/// Compares two filesystem object names for equality. Uses the canonical name representation, not the display name representation.
+		/// </summary>
+		/// <param name="A">First object to compare.</param>
+		/// <param name="B">Second object to compare.</param>
+		/// <returns>True if the names represent the same object, false otherwise</returns>
+		public static bool operator ==(FileReference A, FileReference B)
+		{
+			if((object)A == null)
+			{
+				return (object)B == null;
+			}
+			else
+			{
+				return (object)B != null && A.CanonicalName == B.CanonicalName;
+			}
+		}
+
+		/// <summary>
+		/// Compares two filesystem object names for inequality. Uses the canonical name representation, not the display name representation.
+		/// </summary>
+		/// <param name="A">First object to compare.</param>
+		/// <param name="B">Second object to compare.</param>
+		/// <returns>False if the names represent the same object, true otherwise</returns>
+		public static bool operator !=(FileReference A, FileReference B)
+		{
+			return !(A == B);
+		}
+
+		/// <summary>
+		/// Compares against another object for equality.
+		/// </summary>
+		/// <param name="A">First name to compare.</param>
+		/// <param name="B">Second name to compare.</param>
+		/// <returns>True if the names represent the same object, false otherwise</returns>
+		public override bool Equals(object Obj)
+		{
+			return (Obj is FileReference) && ((FileReference)Obj) == this;
+		}
+
+		/// <summary>
+		/// Returns a hash code for this object
+		/// </summary>
+		/// <returns></returns>
+		public override int GetHashCode()
+		{
+			return CanonicalName.GetHashCode();
+		}
+
+		/// <summary>
+		/// Enumerate files from a given directory
+		/// </summary>
+		/// <returns>Sequence of file references</returns>
+		public static IEnumerable<FileReference> Enumerate(DirectoryReference BaseDirectory)
+		{
+			foreach(string FileName in System.IO.Directory.EnumerateFiles(BaseDirectory.FullName))
+			{
+				yield return new FileReference(FileName, FileName.ToLowerInvariant());
+			}
+		}
+
+		/// <summary>
+		/// Enumerate files from a given directory
+		/// </summary>
+		/// <param name="Pattern">Pattern to match</param>
+		/// <returns>Sequence of file references</returns>
+		public static IEnumerable<FileReference> Enumerate(DirectoryReference BaseDirectory, string Pattern)
+		{
+			foreach(string FileName in System.IO.Directory.EnumerateFiles(BaseDirectory.FullName, Pattern))
+			{
+				yield return new FileReference(FileName, FileName.ToLowerInvariant());
+			}
+		}
+
+		/// <summary>
+		/// Enumerate files from a given directory
+		/// </summary>
+		/// <param name="Pattern">Pattern to match</param>
+		/// <param name="Search">Options for the search</param>
+		/// <returns>Sequence of file references</returns>
+		public static IEnumerable<FileReference> Enumerate(DirectoryReference BaseDirectory, string Pattern, SearchOption Option)
+		{
+			foreach(string FileName in System.IO.Directory.EnumerateFiles(BaseDirectory.FullName, Pattern, Option))
+			{
+				yield return new FileReference(FileName, FileName.ToLowerInvariant());
+			}
+		}
+
+		/// <summary>
+		/// Helper function to create a remote file reference. Unlike normal FileReference objects, these aren't converted to a full path in the local filesystem, but are
+		/// left as they are passed in.
+		/// </summary>
+		/// <param name="AbsolutePath">The absolute path in the remote file system</param>
+		/// <returns>New file reference</returns>
+		public static FileReference MakeRemote(string AbsolutePath)
+		{
+			return new FileReference(AbsolutePath, AbsolutePath.ToLowerInvariant());
+		}
+	}
+}
diff --git a/Engine/Source/Programs/UnrealBuildTool/UnrealBuildTool.csproj b/Engine/Source/Programs/UnrealBuildTool/UnrealBuildTool.csproj
--- a/Engine/Source/Programs/UnrealBuildTool/UnrealBuildTool.csproj
+++ b/Engine/Source/Programs/UnrealBuildTool/UnrealBuildTool.csproj
@@ -124,7 +124,7 @@
     <Compile Include="Mac\UEBuildMac.cs" />
     <Compile Include="Mac\MacProjectGenerator.cs" />
     <Compile Include="Mac\UEDeployMac.cs" />
-    <Compile Include="System\FileSystemNames.cs" />
+    <Compile Include="System\FileSystemReference.cs" />
     <Compile Include="System\TargetReceipt.cs" />
     <Compile Include="System\VersionManifest.cs" />
     <Compile Include="ToolChain\AppleToolChain.cs" />
