Allow UBT to wait on being able to run if the -waitmutex parameter is specified, and add it to the default command line from generated project files. Restores ability to build the entire solution from Visual Studio and have all UBT runs serialized.

diff --git a/Engine/Source/Programs/UnrealBuildTool/System/UnrealBuildTool.cs b/Engine/Source/Programs/UnrealBuildTool/System/UnrealBuildTool.cs
--- a/Engine/Source/Programs/UnrealBuildTool/System/UnrealBuildTool.cs
+++ b/Engine/Source/Programs/UnrealBuildTool/System/UnrealBuildTool.cs
@@ -930,6 +930,14 @@ namespace UnrealBuildTool
                     bUseMutex = false;
                 }
             }
+			bool bWaitMutex = false;
+            {
+                int WaitMutexArgumentIndex;
+                if (Utils.ParseCommandLineFlag(Arguments, "-WaitMutex", out WaitMutexArgumentIndex))
+                {
+                    bWaitMutex = true;
+                }
+            }
 
             bool bAutoSDKOnly = false;
             {
@@ -974,7 +982,17 @@ namespace UnrealBuildTool
                     SingleInstanceMutex = new Mutex(true, MutexName, out bCreatedMutex);
                     if (!bCreatedMutex)
                     {
-                        if (bAutoSDKOnly || bValidatePlatforms)
+						if(bWaitMutex)
+						{
+							try
+							{
+								SingleInstanceMutex.WaitOne();
+							}
+							catch(AbandonedMutexException)
+							{
+							}
+						}
+                        else if (bAutoSDKOnly || bValidatePlatforms)
                         {
                             throw new BuildException("A conflicting instance of UnrealBuildTool is already running. Either -autosdkonly or -validateplatform was passed, which allows only a single instance to be run globally. Therefore, the conflicting instance may be in another location or the current location: {0}. A process manager may be used to determine the conflicting process and what tool may have launched it.", Assembly.GetEntryAssembly().GetOriginalLocation());
                         }
diff --git a/Engine/Source/Programs/UnrealBuildTool/System/VCProject.cs b/Engine/Source/Programs/UnrealBuildTool/System/VCProject.cs
--- a/Engine/Source/Programs/UnrealBuildTool/System/VCProject.cs
+++ b/Engine/Source/Programs/UnrealBuildTool/System/VCProject.cs
@@ -1032,6 +1032,10 @@ namespace UnrealBuildTool
 					{
 						BuildArguments += " " + UProjectPath + (UnrealBuildTool.RunningRocket() ? " -rocket" : "");
 					}
+
+					// Always wait for the mutex between UBT invocations, so that building the whole solution doesn't fail.
+					BuildArguments += " -waitmutex";
+
 					DirectoryReference BatchFilesDirectory = DirectoryReference.Combine(UnrealBuildTool.EngineDirectory, "Build", "BatchFiles");
 
 					// @todo UWP: For the MS toolchains, if an override was set for project generation, push that into the build strings to override the build toolchain as well
