fix for building UE4Game for IOS

diff --git a/Engine/Source/Programs/UnrealBuildTool/IOS/UEDeployIOS.cs b/Engine/Source/Programs/UnrealBuildTool/IOS/UEDeployIOS.cs
--- a/Engine/Source/Programs/UnrealBuildTool/IOS/UEDeployIOS.cs
+++ b/Engine/Source/Programs/UnrealBuildTool/IOS/UEDeployIOS.cs
@@ -739,7 +739,7 @@ namespace UnrealBuildTool.IOS
 					}
 				}
 
-				GeneratePList(ProjectDirectory, bIsUE4Game, GameName, Path.GetFileNameWithoutExtension(UnrealBuildTool.GetUProjectFile().FullName), "../../Engine", "");
+				GeneratePList(ProjectDirectory, bIsUE4Game, GameName, UnrealBuildTool.GetUProjectFile() == null ? "" : Path.GetFileNameWithoutExtension(UnrealBuildTool.GetUProjectFile().FullName), "../../Engine", "");
 			}
 			return true;
 		}
