UE-20594: Fix no suffix bug and it's error during hot-reload.

diff --git a/Engine/Source/Programs/UnrealBuildTool/System/UnrealBuildTool.cs b/Engine/Source/Programs/UnrealBuildTool/System/UnrealBuildTool.cs
--- a/Engine/Source/Programs/UnrealBuildTool/System/UnrealBuildTool.cs
+++ b/Engine/Source/Programs/UnrealBuildTool/System/UnrealBuildTool.cs
@@ -1753,7 +1753,13 @@ namespace UnrealBuildTool
                         UBTMakefile = LoadUBTMakefile( UBTMakefilePath, out ReasonNotLoaded );
 
 						// Invalid makefile if only modules have changed
-						if (UBTMakefile != null && !TargetDescs.SelectMany(x => x.OnlyModules).Select(x => x.OnlyModuleName.ToLower()).SequenceEqual(UBTMakefile.Targets.SelectMany(x => x.OnlyModules).Select(x => x.OnlyModuleName.ToLower())))
+						if (UBTMakefile != null && !TargetDescs.SelectMany(x => x.OnlyModules)
+								.Select(x => new Tuple<string, bool>(x.OnlyModuleName.ToLower(), string.IsNullOrWhiteSpace(x.OnlyModuleSuffix)))
+								.SequenceEqual(
+									UBTMakefile.Targets.SelectMany(x => x.OnlyModules)
+									.Select(x => new Tuple<string, bool>(x.OnlyModuleName.ToLower(), string.IsNullOrWhiteSpace(x.OnlyModuleSuffix)))
+								)
+							)
 						{
 							UBTMakefile     = null;
 							ReasonNotLoaded = "modules to compile have changed";
@@ -2783,6 +2789,12 @@ namespace UnrealBuildTool
 						// Remove "-####-Platform-Configuration" suffix in Debug configuration
 						var IndexOfSecondLastHyphen = OriginalFileNameWithoutExtension.LastIndexOf('-', IndexOfLastHyphen - 1, IndexOfLastHyphen);
 						var IndexOfThirdLastHyphen = OriginalFileNameWithoutExtension.LastIndexOf('-', IndexOfSecondLastHyphen - 1, IndexOfSecondLastHyphen);
+
+						if (!int.TryParse(OriginalFileNameWithoutExtension.Substring(IndexOfThirdLastHyphen + 1, IndexOfSecondLastHyphen - IndexOfThirdLastHyphen - 1), out NumberSuffix))
+						{
+							throw new BuildException("Expected produced item file name '{0}' to contain a numbered suffix when hot reloading second time.", OriginalFileNameWithoutExtension);
+						}
+
 						OriginalFileNameWithoutNumberSuffix = OriginalFileNameWithoutExtension.Substring(0, IndexOfThirdLastHyphen);
 						PlatformConfigSuffix = OriginalFileNameWithoutExtension.Substring(IndexOfSecondLastHyphen);
 					}
