Fix path to the dependency cache file.

diff --git a/Engine/Source/Programs/UnrealBuildTool/System/DependencyCache.cs b/Engine/Source/Programs/UnrealBuildTool/System/DependencyCache.cs
--- a/Engine/Source/Programs/UnrealBuildTool/System/DependencyCache.cs
+++ b/Engine/Source/Programs/UnrealBuildTool/System/DependencyCache.cs
@@ -373,7 +373,7 @@ namespace UnrealBuildTool
 			}
 			else
 			{
-				PlatformIntermediatePath = DirectoryReference.Combine(UnrealBuildTool.EngineDirectory, BuildConfiguration.PlatformIntermediatePath);
+				PlatformIntermediatePath = DirectoryReference.Combine(UnrealBuildTool.EngineDirectory, BuildConfiguration.PlatformIntermediateFolder);
 			}
 			return FileReference.Combine(PlatformIntermediatePath, Target.GetTargetName(), "DependencyCache.bin");
 		}
