Optimize an inner loop that runs over all files that get added to a project. Path.GetDirectoryName() is pretty slow, and if we're running using full paths we don't need to validate it again.

diff --git a/Engine/Source/Programs/UnrealBuildTool/System/Project.cs b/Engine/Source/Programs/UnrealBuildTool/System/Project.cs
--- a/Engine/Source/Programs/UnrealBuildTool/System/Project.cs
+++ b/Engine/Source/Programs/UnrealBuildTool/System/Project.cs
@@ -62,7 +62,7 @@ namespace UnrealBuildTool
 	/// <summary>
 	/// Class that stores info about aliased file.
 	/// </summary>
-	public class AliasedFile
+	public struct AliasedFile
 	{
 		public AliasedFile(string FileSystemPath, string ProjectPath)
 		{
@@ -71,10 +71,10 @@ namespace UnrealBuildTool
 		}
 
 		// File system path.
-		public string FileSystemPath { get; private set; }
+		public readonly string FileSystemPath;
 
 		// Project path.
-		public string ProjectPath { get; private set; }
+		public readonly string ProjectPath;
 	}
 
 	public abstract class ProjectFile : IntelliSenseGatherer
diff --git a/Engine/Source/Programs/UnrealBuildTool/System/VCProject.cs b/Engine/Source/Programs/UnrealBuildTool/System/VCProject.cs
--- a/Engine/Source/Programs/UnrealBuildTool/System/VCProject.cs
+++ b/Engine/Source/Programs/UnrealBuildTool/System/VCProject.cs
@@ -546,12 +546,25 @@ namespace UnrealBuildTool
 					// By default, files will appear relative to the project file in the solution.  This is kind of the normal Visual
 					// Studio way to do things, but because our generated project files are emitted to intermediate folders, if we always
 					// did this it would yield really ugly paths int he solution explorer
-					string FilterRelativeSourceDirectory = Path.GetDirectoryName(ProjectRelativeSourceFile);
+					string FilterRelativeSourceDirectory;
+					if(CurFile.BaseFolder == null)
+					{
+						FilterRelativeSourceDirectory = ProjectRelativeSourceFile;
+					}
+					else
+					{
+						FilterRelativeSourceDirectory = CurFile.Reference.MakeRelativeTo(CurFile.BaseFolder);
+					}
 
-					// Use the specified relative base folder
-					if (CurFile.BaseFolder != null)
+					// Manually remove the filename for the filter. We run through this code path a lot, so just do it manually.
+					int LastSeparatorIdx = FilterRelativeSourceDirectory.LastIndexOf(Path.DirectorySeparatorChar);
+					if(LastSeparatorIdx == -1)
+					{
+						FilterRelativeSourceDirectory = "";
+					}
+					else
 					{
-						FilterRelativeSourceDirectory = Path.GetDirectoryName(CurFile.Reference.MakeRelativeTo(CurFile.BaseFolder));
+						FilterRelativeSourceDirectory = FilterRelativeSourceDirectory.Substring(0, LastSeparatorIdx);
 					}
 
 					LocalAliasedFiles.Add(new AliasedFile(ProjectRelativeSourceFile, FilterRelativeSourceDirectory));
