UBT: fix CodeLite project generator.

diff --git a/Engine/Source/Programs/UnrealBuildTool/System/CodeLiteProject.cs b/Engine/Source/Programs/UnrealBuildTool/System/CodeLiteProject.cs
--- a/Engine/Source/Programs/UnrealBuildTool/System/CodeLiteProject.cs
+++ b/Engine/Source/Programs/UnrealBuildTool/System/CodeLiteProject.cs
@@ -64,7 +64,12 @@ namespace UnrealBuildTool
 				Directory.CreateDirectory(FilePath);
 			}
 
-			string GameProjectFile = UnrealBuildTool.GetUProjectFile().FullName;
+			string GameProjectFile = "";
+			FileReference GameProject = UnrealBuildTool.GetUProjectFile();
+			if (GameProject != null)
+			{
+				GameProjectFile = GameProject.FullName;
+			}
 
 			//
 			// Write all targets which will be separate projects.
@@ -228,7 +233,6 @@ namespace UnrealBuildTool
 				XElement CodeLiteGlobalSettings = new XElement("GlobalSettings");
 				CodeLiteSettings.Add(CodeLiteSettings);
 
-
 				foreach (var CurConf in InConfigurations)
 				{
 					XElement CodeLiteConfiguration = new XElement("Configuration");
@@ -344,7 +348,7 @@ namespace UnrealBuildTool
 						} 
 						else if (ProjectTargetType == TargetRules.TargetType.Editor) 
 						{
-							if (ProjectName != "UE4Editor")
+							if (ProjectName != "UE4Editor" && GameProjectFile != "")
 							{
 								string commandArguments = "\"" + GameProjectFile + "\"" + " -game";
 								XAttribute CommandArguments = new XAttribute("CommandArguments", commandArguments);
