Remove some esoteric options from the project file generator (for only showing public headers, and hiding NoRedist folders) to improve performance of directory scanning.

diff --git a/Engine/Source/Programs/UnrealBuildTool/Configuration/UEBuildTarget.cs b/Engine/Source/Programs/UnrealBuildTool/Configuration/UEBuildTarget.cs
--- a/Engine/Source/Programs/UnrealBuildTool/Configuration/UEBuildTarget.cs
+++ b/Engine/Source/Programs/UnrealBuildTool/Configuration/UEBuildTarget.cs
@@ -3441,9 +3441,7 @@ namespace UnrealBuildTool
 						else
 						{
 							// Don't have a project file for this module with the source file names cached already, so find the source files ourselves
-							SourceFilePaths = SourceFileSearch.FindModuleSourceFiles(
-								ModuleRulesFile: ModuleFileName,
-								ExcludeNoRedistFiles: false);
+							SourceFilePaths = SourceFileSearch.FindModuleSourceFiles(ModuleRulesFile: ModuleFileName);
 						}
 						FoundSourceFiles = GetCPlusPlusFilesToBuild(SourceFilePaths, ModuleDirectory, Platform);
 					}
diff --git a/Engine/Source/Programs/UnrealBuildTool/System/CMakefileGenerator.cs b/Engine/Source/Programs/UnrealBuildTool/System/CMakefileGenerator.cs
--- a/Engine/Source/Programs/UnrealBuildTool/System/CMakefileGenerator.cs
+++ b/Engine/Source/Programs/UnrealBuildTool/System/CMakefileGenerator.cs
@@ -165,7 +165,7 @@ namespace UnrealBuildTool
 			var AllModuleFiles = DiscoverModules(FindGameProjects());
 			foreach (FileReference CurModuleFile in AllModuleFiles)
 			{
-				var FoundFiles = SourceFileSearch.FindModuleSourceFiles(CurModuleFile, ExcludeNoRedistFiles: bExcludeNoRedistFiles);
+				var FoundFiles = SourceFileSearch.FindModuleSourceFiles(CurModuleFile);
 				foreach (FileReference CurSourceFile in FoundFiles)
 				{
 
diff --git a/Engine/Source/Programs/UnrealBuildTool/System/ProjectFileGenerator.cs b/Engine/Source/Programs/UnrealBuildTool/System/ProjectFileGenerator.cs
--- a/Engine/Source/Programs/UnrealBuildTool/System/ProjectFileGenerator.cs
+++ b/Engine/Source/Programs/UnrealBuildTool/System/ProjectFileGenerator.cs
@@ -207,9 +207,6 @@ namespace UnrealBuildTool
 		/// Global static new line string used by ProjectFileGenerator to generate project files.
 		public static readonly string NewLine = Environment.NewLine;
 
-		/// If true, any module or source file under a NoRedist folder will not be included in the generated projects
-		protected bool bExcludeNoRedistFiles = false;
-
 		/// If true, we'll parse subdirectories of third-party projects to locate source and header files to include in the
 		/// generated projects.  This can make the generated projects quite a bit bigger, but makes it easier to open files
 		/// directly from the IDE.
@@ -329,11 +326,6 @@ namespace UnrealBuildTool
 			{
 				Log.TraceInformation("Discovering modules, targets and source code for project...");
 
-				// NOTE: Realistically, the distro that the Rocket user is generating projects FROM won't have NoRedist files in it.  But when 
-				//       testing from a developer branch, this is useful to get authentic projects. This only really matters when 
-				//       bIncludeEngineModulesInRocketProjects=true (defaults to false.)				
-				bExcludeNoRedistFiles = true;
-
 				MasterProjectPath = UnrealBuildTool.GetUProjectPath();
 				IntermediateProjectFilesPath = DirectoryReference.Combine( MasterProjectPath, "Intermediate", "ProjectFiles" );
 
@@ -861,7 +853,7 @@ namespace UnrealBuildTool
 						var GameProjectFile = GameFolderAndProjectFile.Value;
 						var DirectoriesToSearch = new List<DirectoryReference>();
 						DirectoriesToSearch.Add( GameConfigDirectory );
-						GameProjectFile.AddFilesToProject( SourceFileSearch.FindFiles( DirectoriesToSearch, ExcludeNoRedistFiles:bExcludeNoRedistFiles ), GameProjectDirectory );
+						GameProjectFile.AddFilesToProject( SourceFileSearch.FindFiles( DirectoriesToSearch ), GameProjectDirectory );
 					}
 				}
 			}
@@ -876,7 +868,7 @@ namespace UnrealBuildTool
 			{
 				var DirectoriesToSearch = new List<DirectoryReference>();
 				DirectoriesToSearch.Add( EngineLocalizationDirectory );
-				EngineProject.AddFilesToProject( SourceFileSearch.FindFiles( DirectoriesToSearch, ExcludeNoRedistFiles:bExcludeNoRedistFiles), UnrealBuildTool.EngineDirectory );
+				EngineProject.AddFilesToProject( SourceFileSearch.FindFiles( DirectoriesToSearch ), UnrealBuildTool.EngineDirectory );
 			}
 		}
 
@@ -889,7 +881,7 @@ namespace UnrealBuildTool
             {
 				var DirectoriesToSearch = new List<DirectoryReference>();
 				DirectoriesToSearch.Add( EngineTemplateDirectory );
-				EngineProject.AddFilesToProject( SourceFileSearch.FindFiles( DirectoriesToSearch, ExcludeNoRedistFiles:bExcludeNoRedistFiles), UnrealBuildTool.EngineDirectory );
+				EngineProject.AddFilesToProject( SourceFileSearch.FindFiles( DirectoriesToSearch ), UnrealBuildTool.EngineDirectory );
 			}
         }
 
@@ -903,7 +895,7 @@ namespace UnrealBuildTool
 				var DirectoriesToSearch = new List<DirectoryReference>();
 				DirectoriesToSearch.Add( EngineConfigDirectory );
 
-				EngineProject.AddFilesToProject( SourceFileSearch.FindFiles( DirectoriesToSearch, ExcludeNoRedistFiles:bExcludeNoRedistFiles), UnrealBuildTool.EngineDirectory );
+				EngineProject.AddFilesToProject( SourceFileSearch.FindFiles( DirectoriesToSearch ), UnrealBuildTool.EngineDirectory );
 			}
 		}
 
@@ -916,7 +908,7 @@ namespace UnrealBuildTool
 				var DirectoriesToSearch = new List<DirectoryReference>();
 				DirectoriesToSearch.Add(UHTConfigDirectory);
 
-				EngineProject.AddFilesToProject(SourceFileSearch.FindFiles(DirectoriesToSearch, ExcludeNoRedistFiles: bExcludeNoRedistFiles), UnrealBuildTool.EngineDirectory);
+				EngineProject.AddFilesToProject(SourceFileSearch.FindFiles(DirectoriesToSearch), UnrealBuildTool.EngineDirectory);
 			}
 		}
 
@@ -932,18 +924,7 @@ namespace UnrealBuildTool
 			List<FileReference> FoundModuleFiles = RulesCompiler.FindAllRulesSourceFiles( RulesCompiler.RulesFileType.Module, GameFolders: AllGameProjects.Select(x => x.Folder).ToList(), ForeignPlugins:null, AdditionalSearchPaths:null );
 			foreach( var BuildFileName in FoundModuleFiles )
 			{
-				bool IncludeThisModule = true;
-
-				// Skip NoRedist files if necessary
-				if( bExcludeNoRedistFiles )
-				{
-					IncludeThisModule = !IsNoRedistModule(BuildFileName);
-				}
-
-				if( IncludeThisModule )
-				{
-					AllModuleFiles.Add( BuildFileName );
-				}
+				AllModuleFiles.Add( BuildFileName );
 			}
 			return AllModuleFiles;
 		}
@@ -1034,12 +1015,6 @@ namespace UnrealBuildTool
 					}
 				}
 
-				// Skip NoRedist files if necessary
-				if( bExcludeNoRedistFiles )
-				{
-					IncludeThisTarget = !IsNoRedistModule(CurTargetFile);
-				}
-
 				if( IncludeThisTarget )
 				{
 					AllTargetFiles.Add( CurTargetFile );
@@ -1385,7 +1360,7 @@ namespace UnrealBuildTool
 					{
 						var DirectoriesToSearch = new List<DirectoryReference>();
 						DirectoriesToSearch.Add( ProgramConfigDirectory );
-						ProgramProjectFile.AddFilesToProject( SourceFileSearch.FindFiles( DirectoriesToSearch, ExcludeNoRedistFiles:bExcludeNoRedistFiles), ProgramDirectory );
+						ProgramProjectFile.AddFilesToProject( SourceFileSearch.FindFiles( DirectoriesToSearch ), ProgramDirectory );
 					}
 				}
 			}
@@ -1559,19 +1534,11 @@ namespace UnrealBuildTool
 				bool IsEngineModule = CurModuleFile.IsUnderDirectory(UnrealBuildTool.EngineDirectory);
 				bool IsThirdPartyModule = CurModuleFile.IsUnderDirectory(EngineSourceThirdPartyDirectory);
 
-				bool IncludePrivateSourceCode = true;
 				if( bGeneratingGameProjectFiles || bGeneratingRocketProjectFiles )
 				{
 					if( IsEngineModule )
 					{
-						if( IncludeEngineSource )
-						{
-				            // Always include all engine source files unless we were asked to only find public header files via command-line option.  In older releases
-				            // of Unreal Engine, we only included public header files in binary distributions, but this is no longer the case.  You might get
-				            // faster solution load times if you use '-OnlyPublic' though.
-				            IncludePrivateSourceCode &= !UnrealBuildTool.CommandLineContains( "-OnlyPublic" );
-						}
-						else
+						if( !IncludeEngineSource )
 						{
 							// We were asked to exclude engine modules from the generated projects
 							WantProjectFileForModule = false;
@@ -1629,7 +1596,7 @@ namespace UnrealBuildTool
 
 
 					// Find all of the source files (and other files) and add them to the project
-					var FoundFiles = SourceFileSearch.FindModuleSourceFiles( CurModuleFile, ExcludeNoRedistFiles: bExcludeNoRedistFiles, SearchSubdirectories:SearchSubdirectories, IncludePrivateSourceCode:IncludePrivateSourceCode );
+					var FoundFiles = SourceFileSearch.FindModuleSourceFiles( CurModuleFile, SearchSubdirectories:SearchSubdirectories );
 					ProjectFile.AddFilesToProject( FoundFiles, BaseFolder );
 
 					// Check if there's a plugin directory here
@@ -1648,7 +1615,7 @@ namespace UnrealBuildTool
 								DirectoryReference PluginResourcesFolder = DirectoryReference.Combine(PluginFileName.Directory, "Resources");
 								if(PluginResourcesFolder.Exists())
 								{
-									ProjectFile.AddFilesToProject(SourceFileSearch.FindFiles(DirectoriesToSearch: new List<DirectoryReference>{ PluginResourcesFolder }, ExcludeNoRedistFiles: bExcludeNoRedistFiles ), BaseFolder );
+									ProjectFile.AddFilesToProject(SourceFileSearch.FindFiles(DirectoriesToSearch: new List<DirectoryReference>{ PluginResourcesFolder } ), BaseFolder );
 								}
 							}
 						}
@@ -1877,7 +1844,6 @@ namespace UnrealBuildTool
 			}
 			EngineProject.AddFilesToProject( SourceFileSearch.FindFiles(
 				DirectoriesToSearch: DirectoriesToSearch,
-				ExcludeNoRedistFiles: bExcludeNoRedistFiles,
 				SubdirectoryNamesToExclude: SubdirectoryNamesToExclude ), UnrealBuildTool.EngineDirectory );
 		}
 
@@ -1894,7 +1860,6 @@ namespace UnrealBuildTool
 
 			EngineProject.AddFilesToProject( SourceFileSearch.FindFiles(
 				DirectoriesToSearch: DirectoriesToSearch,
-				ExcludeNoRedistFiles: bExcludeNoRedistFiles,
 				SubdirectoryNamesToExclude: SubdirectoryNamesToExclude ), UnrealBuildTool.EngineDirectory );
 		}
 
@@ -1928,7 +1893,6 @@ namespace UnrealBuildTool
 
 				var DocumentationFiles = SourceFileSearch.FindFiles(
 					DirectoriesToSearch: DirectoriesToSearch,
-					ExcludeNoRedistFiles: bExcludeNoRedistFiles,
 					SubdirectoryNamesToExclude: SubdirectoryNamesToExclude );
 
 				// Filter out non-English documentation files if we were configured to do so
diff --git a/Engine/Source/Programs/UnrealBuildTool/System/QMakefileGenerator.cs b/Engine/Source/Programs/UnrealBuildTool/System/QMakefileGenerator.cs
--- a/Engine/Source/Programs/UnrealBuildTool/System/QMakefileGenerator.cs
+++ b/Engine/Source/Programs/UnrealBuildTool/System/QMakefileGenerator.cs
@@ -265,7 +265,7 @@ namespace UnrealBuildTool
 			var AllModuleFiles = DiscoverModules(FindGameProjects());
 			foreach (FileReference CurModuleFile in AllModuleFiles)
 			{
-				var FoundFiles = SourceFileSearch.FindModuleSourceFiles(CurModuleFile, ExcludeNoRedistFiles: bExcludeNoRedistFiles);
+				var FoundFiles = SourceFileSearch.FindModuleSourceFiles(CurModuleFile);
 				foreach (FileReference CurSourceFile in FoundFiles)
 				{
 					string SourceFileRelativeToRoot = CurSourceFile.MakeRelativeTo(UnrealBuildTool.EngineDirectory);
diff --git a/Engine/Source/Programs/UnrealBuildTool/System/SourceFileSearch.cs b/Engine/Source/Programs/UnrealBuildTool/System/SourceFileSearch.cs
--- a/Engine/Source/Programs/UnrealBuildTool/System/SourceFileSearch.cs
+++ b/Engine/Source/Programs/UnrealBuildTool/System/SourceFileSearch.cs
@@ -15,7 +15,7 @@ namespace UnrealBuildTool
 
 
 		/// Finds mouse source files
-		public static List<FileReference> FindModuleSourceFiles( FileReference ModuleRulesFile, bool ExcludeNoRedistFiles, bool SearchSubdirectories = true, bool IncludePrivateSourceCode = true )
+		public static List<FileReference> FindModuleSourceFiles( FileReference ModuleRulesFile, bool SearchSubdirectories = true )
 		{
 			// The module's "base directory" is simply the directory where its xxx.Build.cs file is stored.  We'll always
 			// harvest source files for this module in this base directory directory and all of its sub-directories.
@@ -24,7 +24,7 @@ namespace UnrealBuildTool
 			// Find all of the source files (and other files) and add them to the project
 			var DirectoriesToSearch = new List<DirectoryReference>();
 			DirectoriesToSearch.Add( ModuleBaseDirectory );
-			var FoundFiles = SourceFileSearch.FindFiles( DirectoriesToSearch, ExcludeNoRedistFiles: ExcludeNoRedistFiles, SubdirectoryNamesToExclude:null, SearchSubdirectories:SearchSubdirectories, IncludePrivateSourceCode:IncludePrivateSourceCode );
+			var FoundFiles = SourceFileSearch.FindFiles( DirectoriesToSearch, SubdirectoryNamesToExclude:null, SearchSubdirectories:SearchSubdirectories );
 
 			return FoundFiles;
 		}
@@ -34,11 +34,10 @@ namespace UnrealBuildTool
 		/// Fills in a project file with source files discovered in the specified list of paths
 		/// </summary>
 		/// <param name="SourceFilePaths">List of paths to look for source files</param>
-		/// <param name="ExcludeNoRedistFiles">True if files under NoRedist subdirectories should be excluded</param>
 		/// <param name="SubdirectoryNamesToExclude">Directory base names to ignore when searching subdirectories.  Can be null.</param>
 		/// <param name="SearchSubdirectories">True to include subdirectories, otherwise we only search the list of base directories</param>
 		/// <param name="IncludePrivateSourceCode">True if source files in the 'Private' directory should be included</param>
-		public static List<FileReference> FindFiles( List<DirectoryReference> DirectoriesToSearch, bool ExcludeNoRedistFiles, List<string> SubdirectoryNamesToExclude = null, bool SearchSubdirectories = true, bool IncludePrivateSourceCode = true )
+		public static List<FileReference> FindFiles( List<DirectoryReference> DirectoriesToSearch, List<string> SubdirectoryNamesToExclude = null, bool SearchSubdirectories = true )
 		{
 			// Certain file types should never be added to project files
 			var ExcludedFileTypes = new string[] 
@@ -79,36 +78,6 @@ namespace UnrealBuildTool
 
 						bool IncludeThisFile = true;
 
-						if( IncludeThisFile )
-						{
-							if( !IncludePrivateSourceCode )
-							{
-								// We consider to be Private source every file under the module's Source directory that is not either:
-								//		- Under Public sub-folder
-								//		- Ends in a C# build script file extension
-								bool IsBuildScript = RelativeSourceFilePath.EndsWith( ".Build.cs", StringComparison.InvariantCultureIgnoreCase ) || RelativeSourceFilePath.EndsWith( ".Target.cs", StringComparison.InvariantCultureIgnoreCase );
-								if( !IsBuildScript &&
-									!RelativeSourceFilePath.StartsWith( "Public" + Path.DirectorySeparatorChar, StringComparison.InvariantCultureIgnoreCase ) &&
-									!( RelativeSourceFilePath.IndexOf( Path.DirectorySeparatorChar + "Public" + Path.DirectorySeparatorChar, StringComparison.InvariantCultureIgnoreCase ) != -1 ) )
-								{
-									IncludeThisFile = false;
-								}
-							}
-						}
-
-						if( IncludeThisFile )
-						{
-							// Skip NoRedist files if necessary
-							if( ExcludeNoRedistFiles )
-							{
-								if( RelativeSourceFilePath.StartsWith( "NoRedist" + Path.DirectorySeparatorChar, StringComparison.InvariantCultureIgnoreCase ) ||
-									RelativeSourceFilePath.IndexOf( Path.DirectorySeparatorChar + "NoRedist" + Path.DirectorySeparatorChar, StringComparison.InvariantCultureIgnoreCase ) != -1 )
-								{
-									IncludeThisFile = false;
-								}
-							}
-						}
-
 						// Mask out blacklisted folders
 						if( IncludeThisFile && SubdirectoryNamesToExclude != null )
 						{
diff --git a/Engine/Source/Programs/UnrealBuildTool/System/XcodeProjectFileGenerator.cs b/Engine/Source/Programs/UnrealBuildTool/System/XcodeProjectFileGenerator.cs
--- a/Engine/Source/Programs/UnrealBuildTool/System/XcodeProjectFileGenerator.cs
+++ b/Engine/Source/Programs/UnrealBuildTool/System/XcodeProjectFileGenerator.cs
@@ -885,7 +885,7 @@ namespace UnrealBuildTool
 		{
 			var XcodeProject = new XcodeProjectFile(ProjectPath);
 			DirectoryReference ProjectDirectory = ProjectPath.Directory;
-			XcodeProject.AddFilesToProject(SourceFileSearch.FindFiles(new List<DirectoryReference>() { ProjectDirectory }, ExcludeNoRedistFiles: bExcludeNoRedistFiles, SubdirectoryNamesToExclude:new List<string>() { "obj" }, SearchSubdirectories:true, IncludePrivateSourceCode:true), null);
+			XcodeProject.AddFilesToProject(SourceFileSearch.FindFiles(new List<DirectoryReference>() { ProjectDirectory }, SubdirectoryNamesToExclude:new List<string>() { "obj" }, SearchSubdirectories:true), null);
 			XcodeProject.GenerateSectionsContents(ref PBXBuildFileSection, ref PBXFileReferenceSection, ref PBXSourcesBuildPhaseSection, ref Groups);
 		}
 
