Optimize searching for source files under a given directory, by constructing DirectoryReference and FileReference objects directly from the full paths returned by the system.

diff --git a/Engine/Source/Programs/UnrealBuildTool/System/DirectoryLookupCache.cs b/Engine/Source/Programs/UnrealBuildTool/System/DirectoryLookupCache.cs
--- a/Engine/Source/Programs/UnrealBuildTool/System/DirectoryLookupCache.cs
+++ b/Engine/Source/Programs/UnrealBuildTool/System/DirectoryLookupCache.cs
@@ -12,7 +12,7 @@ namespace UnrealBuildTool
 		{
 			if( BaseDirectory.Exists() )
 			{
-				Files = new HashSet<FileReference>(FileReference.Enumerate(BaseDirectory));
+				Files = new HashSet<FileReference>(BaseDirectory.EnumerateFileReferences());
 			}
 			else
 			{
diff --git a/Engine/Source/Programs/UnrealBuildTool/System/FileSystemReference.cs b/Engine/Source/Programs/UnrealBuildTool/System/FileSystemReference.cs
--- a/Engine/Source/Programs/UnrealBuildTool/System/FileSystemReference.cs
+++ b/Engine/Source/Programs/UnrealBuildTool/System/FileSystemReference.cs
@@ -32,6 +32,15 @@ namespace UnrealBuildTool
 			CanonicalName = FullName.ToLowerInvariant();
 		}
 
+		/// <summary>
+		/// Constructs a reference from the given FileSystemInfo.
+		/// </summary>
+		public FileSystemReference(FileSystemInfo InInfo)
+		{
+			FullName = InInfo.FullName;
+			CanonicalName = FullName.ToLowerInvariant();
+		}
+
 		/// <summary>
 		/// Direct constructor for a path
 		/// </summary>
@@ -247,6 +256,15 @@ namespace UnrealBuildTool
 		{
 		}
 
+		/// <summary>
+		/// Construct a DirectoryReference from a DirectoryInfo object.
+		/// </summary>
+		/// <param name="InInfo">Path to this file</param>
+		public DirectoryReference(DirectoryInfo InInfo)
+			: base(InInfo)
+		{
+		}
+
 		/// <summary>
 		/// Constructor for creating a directory object directly from two strings.
 		/// </summary>
@@ -317,6 +335,78 @@ namespace UnrealBuildTool
 			return Directory.Exists(FullName);
 		}
 
+		/// <summary>
+		/// Enumerate files from a given directory
+		/// </summary>
+		/// <returns>Sequence of file references</returns>
+		public IEnumerable<FileReference> EnumerateFileReferences()
+		{
+			foreach(string FileName in Directory.EnumerateFiles(FullName))
+			{
+				yield return FileReference.MakeFromNormalizedFullPath(FileName);
+			}
+		}
+
+		/// <summary>
+		/// Enumerate files from a given directory
+		/// </summary>
+		/// <returns>Sequence of file references</returns>
+		public IEnumerable<FileReference> EnumerateFileReferences(string Pattern)
+		{
+			foreach(string FileName in Directory.EnumerateFiles(FullName, Pattern))
+			{
+				yield return FileReference.MakeFromNormalizedFullPath(FileName);
+			}
+		}
+
+		/// <summary>
+		/// Enumerate files from a given directory
+		/// </summary>
+		/// <returns>Sequence of file references</returns>
+		public IEnumerable<FileReference> EnumerateFileReferences(string Pattern, SearchOption Option)
+		{
+			foreach(string FileName in Directory.EnumerateFiles(FullName, Pattern, Option))
+			{
+				yield return FileReference.MakeFromNormalizedFullPath(FileName);
+			}
+		}
+
+		/// <summary>
+		/// Enumerate subdirectories in a given directory
+		/// </summary>
+		/// <returns>Sequence of directory references</returns>
+		public IEnumerable<DirectoryReference> EnumerateDirectoryReferences()
+		{
+			foreach(string DirectoryName in Directory.EnumerateDirectories(FullName))
+			{
+				yield return DirectoryReference.MakeFromNormalizedFullPath(DirectoryName);
+			}
+		}
+
+		/// <summary>
+		/// Enumerate subdirectories in a given directory
+		/// </summary>
+		/// <returns>Sequence of directory references</returns>
+		public IEnumerable<DirectoryReference> EnumerateDirectoryReferences(string Pattern)
+		{
+			foreach(string DirectoryName in Directory.EnumerateDirectories(FullName, Pattern))
+			{
+				yield return DirectoryReference.MakeFromNormalizedFullPath(DirectoryName);
+			}
+		}
+
+		/// <summary>
+		/// Enumerate subdirectories in a given directory
+		/// </summary>
+		/// <returns>Sequence of directory references</returns>
+		public IEnumerable<DirectoryReference> EnumerateDirectoryReferences(string Pattern, SearchOption Option)
+		{
+			foreach(string DirectoryName in Directory.EnumerateDirectories(FullName, Pattern, Option))
+			{
+				yield return DirectoryReference.MakeFromNormalizedFullPath(DirectoryName);
+			}
+		}
+
 		/// <summary>
 		/// Determines whether this path represents a root directory in the filesystem
 		/// </summary>
@@ -396,6 +486,16 @@ namespace UnrealBuildTool
 		{
 			return new DirectoryReference(AbsolutePath, AbsolutePath.ToLowerInvariant());
 		}
+
+		/// <summary>
+		/// Helper function to create a directory reference from a raw platform path. The path provided *MUST* be exactly the same as that returned by Path.GetFullPath(). 
+		/// </summary>
+		/// <param name="AbsolutePath">The absolute path in the file system</param>
+		/// <returns>New file reference</returns>
+		public static DirectoryReference MakeFromNormalizedFullPath(string AbsolutePath)
+		{
+			return new DirectoryReference(AbsolutePath, AbsolutePath.ToLowerInvariant());
+		}
 	}
 
 	/// <summary>
@@ -413,6 +513,15 @@ namespace UnrealBuildTool
 		{
 		}
 
+		/// <summary>
+		/// Construct a FileReference from a FileInfo object.
+		/// </summary>
+		/// <param name="InInfo">Path to this file</param>
+		public FileReference(FileInfo InInfo)
+			: base(InInfo)
+		{
+		}
+
 		/// <summary>
 		/// Default constructor.
 		/// </summary>
@@ -578,51 +687,22 @@ namespace UnrealBuildTool
 		}
 
 		/// <summary>
-		/// Enumerate files from a given directory
-		/// </summary>
-		/// <returns>Sequence of file references</returns>
-		public static IEnumerable<FileReference> Enumerate(DirectoryReference BaseDirectory)
-		{
-			foreach(string FileName in System.IO.Directory.EnumerateFiles(BaseDirectory.FullName))
-			{
-				yield return new FileReference(FileName, FileName.ToLowerInvariant());
-			}
-		}
-
-		/// <summary>
-		/// Enumerate files from a given directory
-		/// </summary>
-		/// <param name="Pattern">Pattern to match</param>
-		/// <returns>Sequence of file references</returns>
-		public static IEnumerable<FileReference> Enumerate(DirectoryReference BaseDirectory, string Pattern)
-		{
-			foreach(string FileName in System.IO.Directory.EnumerateFiles(BaseDirectory.FullName, Pattern))
-			{
-				yield return new FileReference(FileName, FileName.ToLowerInvariant());
-			}
-		}
-
-		/// <summary>
-		/// Enumerate files from a given directory
+		/// Helper function to create a remote file reference. Unlike normal FileReference objects, these aren't converted to a full path in the local filesystem, but are
+		/// left as they are passed in.
 		/// </summary>
-		/// <param name="Pattern">Pattern to match</param>
-		/// <param name="Search">Options for the search</param>
-		/// <returns>Sequence of file references</returns>
-		public static IEnumerable<FileReference> Enumerate(DirectoryReference BaseDirectory, string Pattern, SearchOption Option)
+		/// <param name="AbsolutePath">The absolute path in the remote file system</param>
+		/// <returns>New file reference</returns>
+		public static FileReference MakeRemote(string AbsolutePath)
 		{
-			foreach(string FileName in System.IO.Directory.EnumerateFiles(BaseDirectory.FullName, Pattern, Option))
-			{
-				yield return new FileReference(FileName, FileName.ToLowerInvariant());
-			}
+			return new FileReference(AbsolutePath, AbsolutePath.ToLowerInvariant());
 		}
 
 		/// <summary>
-		/// Helper function to create a remote file reference. Unlike normal FileReference objects, these aren't converted to a full path in the local filesystem, but are
-		/// left as they are passed in.
+		/// Helper function to create a file reference from a raw platform path. The path provided *MUST* be exactly the same as that returned by Path.GetFullPath(). 
 		/// </summary>
-		/// <param name="AbsolutePath">The absolute path in the remote file system</param>
+		/// <param name="AbsolutePath">The absolute path in the file system</param>
 		/// <returns>New file reference</returns>
-		public static FileReference MakeRemote(string AbsolutePath)
+		public static FileReference MakeFromNormalizedFullPath(string AbsolutePath)
 		{
 			return new FileReference(AbsolutePath, AbsolutePath.ToLowerInvariant());
 		}
diff --git a/Engine/Source/Programs/UnrealBuildTool/System/ProjectFileGenerator.cs b/Engine/Source/Programs/UnrealBuildTool/System/ProjectFileGenerator.cs
--- a/Engine/Source/Programs/UnrealBuildTool/System/ProjectFileGenerator.cs
+++ b/Engine/Source/Programs/UnrealBuildTool/System/ProjectFileGenerator.cs
@@ -851,9 +851,7 @@ namespace UnrealBuildTool
 					if( GameConfigDirectory.Exists( ) )
 					{
 						var GameProjectFile = GameFolderAndProjectFile.Value;
-						var DirectoriesToSearch = new List<DirectoryReference>();
-						DirectoriesToSearch.Add( GameConfigDirectory );
-						GameProjectFile.AddFilesToProject( SourceFileSearch.FindFiles( DirectoriesToSearch ), GameProjectDirectory );
+						GameProjectFile.AddFilesToProject( SourceFileSearch.FindFiles( GameConfigDirectory ), GameProjectDirectory );
 					}
 				}
 			}
@@ -866,9 +864,7 @@ namespace UnrealBuildTool
 			var EngineLocalizationDirectory = DirectoryReference.Combine( UnrealBuildTool.EngineDirectory, "Content", "Localization" );
 			if( EngineLocalizationDirectory.Exists( ) )
 			{
-				var DirectoriesToSearch = new List<DirectoryReference>();
-				DirectoriesToSearch.Add( EngineLocalizationDirectory );
-				EngineProject.AddFilesToProject( SourceFileSearch.FindFiles( DirectoriesToSearch ), UnrealBuildTool.EngineDirectory );
+				EngineProject.AddFilesToProject( SourceFileSearch.FindFiles( EngineLocalizationDirectory ), UnrealBuildTool.EngineDirectory );
 			}
 		}
 
@@ -879,9 +875,7 @@ namespace UnrealBuildTool
             var EngineTemplateDirectory = DirectoryReference.Combine(UnrealBuildTool.EngineDirectory, "Content", "Editor", "Templates");
             if (EngineTemplateDirectory.Exists())
             {
-				var DirectoriesToSearch = new List<DirectoryReference>();
-				DirectoriesToSearch.Add( EngineTemplateDirectory );
-				EngineProject.AddFilesToProject( SourceFileSearch.FindFiles( DirectoriesToSearch ), UnrealBuildTool.EngineDirectory );
+				EngineProject.AddFilesToProject( SourceFileSearch.FindFiles( EngineTemplateDirectory ), UnrealBuildTool.EngineDirectory );
 			}
         }
 
@@ -892,10 +886,7 @@ namespace UnrealBuildTool
             var EngineConfigDirectory = DirectoryReference.Combine(UnrealBuildTool.EngineDirectory, "Config" );
 			if( EngineConfigDirectory.Exists( ) )
 			{
-				var DirectoriesToSearch = new List<DirectoryReference>();
-				DirectoriesToSearch.Add( EngineConfigDirectory );
-
-				EngineProject.AddFilesToProject( SourceFileSearch.FindFiles( DirectoriesToSearch ), UnrealBuildTool.EngineDirectory );
+				EngineProject.AddFilesToProject( SourceFileSearch.FindFiles( EngineConfigDirectory ), UnrealBuildTool.EngineDirectory );
 			}
 		}
 
@@ -905,10 +896,7 @@ namespace UnrealBuildTool
 			var UHTConfigDirectory = DirectoryReference.Combine(UnrealBuildTool.EngineDirectory, "Programs", "UnrealHeaderTool", "Config");
 			if (UHTConfigDirectory.Exists())
 			{
-				var DirectoriesToSearch = new List<DirectoryReference>();
-				DirectoriesToSearch.Add(UHTConfigDirectory);
-
-				EngineProject.AddFilesToProject(SourceFileSearch.FindFiles(DirectoriesToSearch), UnrealBuildTool.EngineDirectory);
+				EngineProject.AddFilesToProject(SourceFileSearch.FindFiles(UHTConfigDirectory), UnrealBuildTool.EngineDirectory);
 			}
 		}
 
@@ -1358,9 +1346,7 @@ namespace UnrealBuildTool
 					var ProgramConfigDirectory = DirectoryReference.Combine( ProgramDirectory, "Config" );
 					if( ProgramConfigDirectory.Exists( ) )
 					{
-						var DirectoriesToSearch = new List<DirectoryReference>();
-						DirectoriesToSearch.Add( ProgramConfigDirectory );
-						ProgramProjectFile.AddFilesToProject( SourceFileSearch.FindFiles( DirectoriesToSearch ), ProgramDirectory );
+						ProgramProjectFile.AddFilesToProject( SourceFileSearch.FindFiles( ProgramConfigDirectory ), ProgramDirectory );
 					}
 				}
 			}
@@ -1615,7 +1601,7 @@ namespace UnrealBuildTool
 								DirectoryReference PluginResourcesFolder = DirectoryReference.Combine(PluginFileName.Directory, "Resources");
 								if(PluginResourcesFolder.Exists())
 								{
-									ProjectFile.AddFilesToProject(SourceFileSearch.FindFiles(DirectoriesToSearch: new List<DirectoryReference>{ PluginResourcesFolder } ), BaseFolder );
+									ProjectFile.AddFilesToProject(SourceFileSearch.FindFiles(PluginResourcesFolder), BaseFolder );
 								}
 							}
 						}
@@ -1832,8 +1818,6 @@ namespace UnrealBuildTool
 			// Setup a project file entry for this module's project.  Remember, some projects may host multiple modules!
 			var ShadersDirectory = DirectoryReference.Combine( UnrealBuildTool.EngineDirectory, "Shaders" );
 
-			var DirectoriesToSearch = new List<DirectoryReference>();
-			DirectoriesToSearch.Add( ShadersDirectory );
 			var SubdirectoryNamesToExclude = new List<string>();
 			{
                 // Don't include binary shaders in the project file.
@@ -1842,9 +1826,8 @@ namespace UnrealBuildTool
 				SubdirectoryNamesToExclude.Add( "PDBDump" );
 				SubdirectoryNamesToExclude.Add( "WorkingDirectory" );
 			}
-			EngineProject.AddFilesToProject( SourceFileSearch.FindFiles(
-				DirectoriesToSearch: DirectoriesToSearch,
-				SubdirectoryNamesToExclude: SubdirectoryNamesToExclude ), UnrealBuildTool.EngineDirectory );
+
+			EngineProject.AddFilesToProject( SourceFileSearch.FindFiles( ShadersDirectory, SubdirectoryNamesToExclude ), UnrealBuildTool.EngineDirectory );
 		}
 
 
@@ -1853,14 +1836,10 @@ namespace UnrealBuildTool
 		{
 			var BuildDirectory = DirectoryReference.Combine( UnrealBuildTool.EngineDirectory, "Build" );
 
-			var DirectoriesToSearch = new List<DirectoryReference>();
-			DirectoriesToSearch.Add( BuildDirectory );
 			var SubdirectoryNamesToExclude = new List<string>();
 			SubdirectoryNamesToExclude.Add("Receipts");
 
-			EngineProject.AddFilesToProject( SourceFileSearch.FindFiles(
-				DirectoriesToSearch: DirectoriesToSearch,
-				SubdirectoryNamesToExclude: SubdirectoryNamesToExclude ), UnrealBuildTool.EngineDirectory );
+			EngineProject.AddFilesToProject( SourceFileSearch.FindFiles( BuildDirectory, SubdirectoryNamesToExclude ), UnrealBuildTool.EngineDirectory );
 		}
 
 
@@ -1875,9 +1854,6 @@ namespace UnrealBuildTool
 			{
 				Log.TraceVerbose( "Adding documentation files..." );
 
-				var DirectoriesToSearch = new List<DirectoryReference>();
-				DirectoriesToSearch.Add( DocumentationSourceDirectory );
-
 				var SubdirectoryNamesToExclude = new List<string>();
 				{
 					// We never want any of the images or attachment files included in our generated project
@@ -1891,9 +1867,7 @@ namespace UnrealBuildTool
 					SubdirectoryNamesToExclude.Add( "Javascript" );
 				}
 
-				var DocumentationFiles = SourceFileSearch.FindFiles(
-					DirectoriesToSearch: DirectoriesToSearch,
-					SubdirectoryNamesToExclude: SubdirectoryNamesToExclude );
+				var DocumentationFiles = SourceFileSearch.FindFiles( DocumentationSourceDirectory, SubdirectoryNamesToExclude );
 
 				// Filter out non-English documentation files if we were configured to do so
 				if( !bAllDocumentationLanguages )
diff --git a/Engine/Source/Programs/UnrealBuildTool/System/SourceFileSearch.cs b/Engine/Source/Programs/UnrealBuildTool/System/SourceFileSearch.cs
--- a/Engine/Source/Programs/UnrealBuildTool/System/SourceFileSearch.cs
+++ b/Engine/Source/Programs/UnrealBuildTool/System/SourceFileSearch.cs
@@ -11,25 +11,38 @@ namespace UnrealBuildTool
 {
 	public static class SourceFileSearch
 	{
-		private delegate void FileSearchDelegate(IEnumerable<string> FileList);
-
+		// Certain file types should never be added to project files. These extensions must all be lowercase.
+		static readonly string[] DefaultExcludedFileSuffixes = new string[] 
+		{
+			".vcxproj",				// Visual Studio project files
+			".vcxproj.filters",		// Visual Studio filter file
+			".vcxproj.user",		// Visual Studio project user file
+			".vcxproj.vspscc",		// Visual Studio source control state
+			".sln",					// Visual Studio solution file
+			".sdf",					// Visual Studio embedded symbol database file
+			".suo",					// Visual Studio solution option files
+			".sln.docstates.suo",	// Visual Studio temp file
+			".tmp",					// Temp files used by P4 diffing (e.g. t6884t87698.tmp)
+			".csproj",				// Visual Studio C# project files
+			".userprefs",			// MonoDevelop project settings
+			".ds_store",			// Mac Desktop Services Store hidden files
+			Path.DirectorySeparatorChar + "do_not_delete.txt",		// Perforce placeholder file
+		};
+
+		// Default directory names to exclude. Must be lowercase.
+		static readonly string[] DefaultExcludedDirectorySuffixes = new string[] 
+		{
+			Path.DirectorySeparatorChar + "intermediate"
+		};
 
 		/// Finds mouse source files
 		public static List<FileReference> FindModuleSourceFiles( FileReference ModuleRulesFile, bool SearchSubdirectories = true )
 		{
 			// The module's "base directory" is simply the directory where its xxx.Build.cs file is stored.  We'll always
 			// harvest source files for this module in this base directory directory and all of its sub-directories.
-			var ModuleBaseDirectory = ModuleRulesFile.Directory;
-
-			// Find all of the source files (and other files) and add them to the project
-			var DirectoriesToSearch = new List<DirectoryReference>();
-			DirectoriesToSearch.Add( ModuleBaseDirectory );
-			var FoundFiles = SourceFileSearch.FindFiles( DirectoriesToSearch, SubdirectoryNamesToExclude:null, SearchSubdirectories:SearchSubdirectories );
-
-			return FoundFiles;
+			return SourceFileSearch.FindFiles( ModuleRulesFile.Directory, SubdirectoryNamesToExclude:null, SearchSubdirectories:SearchSubdirectories );
 		}
 
-
 		/// <summary>
 		/// Fills in a project file with source files discovered in the specified list of paths
 		/// </summary>
@@ -37,131 +50,66 @@ namespace UnrealBuildTool
 		/// <param name="SubdirectoryNamesToExclude">Directory base names to ignore when searching subdirectories.  Can be null.</param>
 		/// <param name="SearchSubdirectories">True to include subdirectories, otherwise we only search the list of base directories</param>
 		/// <param name="IncludePrivateSourceCode">True if source files in the 'Private' directory should be included</param>
-		public static List<FileReference> FindFiles( List<DirectoryReference> DirectoriesToSearch, List<string> SubdirectoryNamesToExclude = null, bool SearchSubdirectories = true )
+		public static List<FileReference> FindFiles( DirectoryReference DirectoryToSearch, List<string> SubdirectoryNamesToExclude = null, bool SearchSubdirectories = true )
 		{
-			// Certain file types should never be added to project files
-			var ExcludedFileTypes = new string[] 
-						{
-							".vcxproj",				// Visual Studio project files
-							".vcxproj.filters",		// Visual Studio filter file
-							".vcxproj.user",		// Visual Studio project user file
-							".vcxproj.vspscc",		// Visual Studio source control state
-							".sln",					// Visual Studio solution file
-							".sdf",					// Visual Studio embedded symbol database file
-							".suo",					// Visual Studio solution option files
-							".sln.docstates.suo",	// Visual Studio temp file
-							".tmp",					// Temp files used by P4 diffing (e.g. t6884t87698.tmp)
-							".csproj",				// Visual Studio C# project files
-							".userprefs",			// MonoDevelop project settings
-							".DS_Store",			// Mac Desktop Services Store hidden files
-						};
-			var ExcludedFileNames = new string[] 
-						{
-							"DO_NOT_DELETE.txt",	// Placeholder .txt file in P4
-						};
+			// Build a list of directory names that we don't want to search under. We always ignore intermediate directories.
+			string[] ExcludedDirectorySuffixes;
+			if(SubdirectoryNamesToExclude == null)
+			{
+				ExcludedDirectorySuffixes = DefaultExcludedDirectorySuffixes;
+			}
+			else
+			{
+				ExcludedDirectorySuffixes = SubdirectoryNamesToExclude.Select(x => Path.DirectorySeparatorChar + x.ToLowerInvariant()).Union(DefaultExcludedDirectorySuffixes).ToArray();
+			}
 
-			HashSet<FileReference> FoundFiles = new HashSet<FileReference>();
-			foreach( var SearchDirectory in DirectoriesToSearch )
+			// Find all the files
+			List<FileReference> FoundFiles = new List<FileReference>();
+			if(SearchSubdirectories)
 			{
-				var FoundFileUnderIntermediateDirectory = string.Empty;
+				FindFilesInternalRecursive(DirectoryToSearch, ExcludedDirectorySuffixes, FoundFiles);
+			}
+			else
+			{
+				FindFilesInternal(DirectoryToSearch, ExcludedDirectorySuffixes, FoundFiles);
+			}
+			return FoundFiles;
+		}
 
-				// Look for files to add to the project within this module's directory
-				// NOTE: Due to issue with current Mono release (2.10.x), Directory.EnumerateFiles does not
-				// work correctly. Therefore, the logic here is encapsulated inside a delegate, and called
-				// with either the result of EnumerateFiles or GetFiles, depending on platform.
-				FileSearchDelegate DoSearch = delegate(IEnumerable<string> FileList) 
+		static void FindFilesInternal(DirectoryReference Directory, string[] ExcludedDirectorySuffixes, List<FileReference> FoundFiles)
+		{
+			foreach(FileReference File in Directory.EnumerateFileReferences())
+			{
+				if(ShouldInclude(File, DefaultExcludedFileSuffixes))
 				{
-					foreach(var CurSourceFilePath in FileList)
-					{
-						// Get the part of the path that is relative to the module folder
-						var RelativeSourceFilePath = new FileReference(CurSourceFilePath).MakeRelativeTo(SearchDirectory);
-
-						bool IncludeThisFile = true;
-
-						// Mask out blacklisted folders
-						if( IncludeThisFile && SubdirectoryNamesToExclude != null )
-						{
-							foreach( var BlacklistedSubdirectoryName in SubdirectoryNamesToExclude )
-							{
-								if( RelativeSourceFilePath.StartsWith( BlacklistedSubdirectoryName + Path.DirectorySeparatorChar, StringComparison.InvariantCultureIgnoreCase ) ||
-									RelativeSourceFilePath.IndexOf( Path.DirectorySeparatorChar + BlacklistedSubdirectoryName + Path.DirectorySeparatorChar, StringComparison.InvariantCultureIgnoreCase ) != -1 )
-								{
-									IncludeThisFile = false;
-									break;
-								}
-							}
-						}
-
-						// Mask out blacklisted file extensions
-						if( IncludeThisFile )
-						{
-							foreach( var BlacklistedFileExtension in ExcludedFileTypes )
-							{
-								if( RelativeSourceFilePath.EndsWith( BlacklistedFileExtension, StringComparison.InvariantCultureIgnoreCase ) )
-								{
-									IncludeThisFile = false;
-									break;
-								}
-							}
-
-							if( IncludeThisFile )
-							{
-								foreach( var BlacklistedFileName in ExcludedFileNames )
-								{
-									if( Path.GetFileName( RelativeSourceFilePath ).Equals( BlacklistedFileName, StringComparison.InvariantCultureIgnoreCase ) )
-									{
-										IncludeThisFile = false;
-										break;
-									}
-								}
-							}
-						}
-
-						// Watch out for "Intermediate" folders.  Module folders should not usually have an Intermediate subfolder, but we'll check for this in case
-						// there is something stale or broken; Our module is not going to compile properly if it tries to pull in existing intermediates as source!
-						if( IncludeThisFile )
-						{ 
-							if( RelativeSourceFilePath.StartsWith("Intermediate" + Path.DirectorySeparatorChar, StringComparison.InvariantCultureIgnoreCase) ||
-								RelativeSourceFilePath.IndexOf(Path.DirectorySeparatorChar + "Intermediate" + Path.DirectorySeparatorChar, StringComparison.InvariantCultureIgnoreCase) != -1)
-							{
-								if( !String.IsNullOrEmpty( FoundFileUnderIntermediateDirectory ) )
-								{ 
-									// Save off at least one example of a file that was bad, so that we can print a warning later
-									FoundFileUnderIntermediateDirectory = CurSourceFilePath;
-								}
-								IncludeThisFile = false;
-							}
-						}
+					FoundFiles.Add(File);
+				}
+			}
+		}
 
-						if( IncludeThisFile )
-						{
-							// Don't add duplicates
-							if( !FoundFiles.Contains( new FileReference(CurSourceFilePath) ) )
-							{
-								FoundFiles.Add( new FileReference(CurSourceFilePath) );
-							}
-						}
-					}
-				};
+		static void FindFilesInternalRecursive(DirectoryReference Directory, string[] ExcludedDirectorySuffixes, List<FileReference> FoundFiles)
+		{
+			FindFilesInternal(Directory, ExcludedDirectorySuffixes, FoundFiles);
 
-				// call the delegate to do the work, based on platform (see note above)
-				if(!Utils.IsRunningOnMono)
-				{
-					DoSearch(Directory.EnumerateFiles(SearchDirectory.FullName, "*.*", SearchSubdirectories ? SearchOption.AllDirectories : SearchOption.TopDirectoryOnly ));
-				}
-				else
+			foreach(DirectoryReference SubDirectory in Directory.EnumerateDirectoryReferences())
+			{
+				if(ShouldInclude(SubDirectory, ExcludedDirectorySuffixes))
 				{
-					DoSearch(Directory.GetFiles(SearchDirectory.FullName, "*.*", SearchSubdirectories ? SearchOption.AllDirectories : SearchOption.TopDirectoryOnly).AsEnumerable());
+					FindFilesInternalRecursive(SubDirectory, ExcludedDirectorySuffixes, FoundFiles);
 				}
+			}
+		}
 
-				if( !string.IsNullOrEmpty( FoundFileUnderIntermediateDirectory ) )
+		static bool ShouldInclude(FileSystemReference Reference, string[] InvariantExcludedSuffixes)
+		{
+			foreach(string InvariantExcludedSuffix in InvariantExcludedSuffixes)
+			{
+				if(Reference.CanonicalName.EndsWith(InvariantExcludedSuffix))
 				{
-					Console.WriteLine("WARNING: UnrealBuildTool skipped at least one file in an Intermediate folder while gathering module files in directory '{0}' (file: '{1}').  This is unexpected, please remove this Intermediate directory.", SearchDirectory, FoundFileUnderIntermediateDirectory );
+					return false;
 				}
 			}
-
-			return FoundFiles.ToList();
+			return true;
 		}
-
 	}
 }
diff --git a/Engine/Source/Programs/UnrealBuildTool/System/UnrealBuildTool.cs b/Engine/Source/Programs/UnrealBuildTool/System/UnrealBuildTool.cs
--- a/Engine/Source/Programs/UnrealBuildTool/System/UnrealBuildTool.cs
+++ b/Engine/Source/Programs/UnrealBuildTool/System/UnrealBuildTool.cs
@@ -2590,7 +2590,7 @@ namespace UnrealBuildTool
 							// Makefile was created. 
 							//if( Directory.GetLastWriteTimeUtc( FlatModule.ModuleSourceFolder ) > UBTMakefileInfo.LastWriteTimeUtc )
 							{
-								foreach( var SourceFilePath in FileReference.Enumerate( FlatModule.ModuleSourceFolder, "*.cpp", SearchOption.AllDirectories ) )
+								foreach( var SourceFilePath in FlatModule.ModuleSourceFolder.EnumerateFileReferences( "*.cpp", SearchOption.AllDirectories ) )
 								{
 									// Was the file written to since the Makefile was created?  We're not able to do a fully exhaustive
 									// check for working sets when in assembler mode, so we might get it wrong occasionally.  This is
diff --git a/Engine/Source/Programs/UnrealBuildTool/System/XcodeProjectFileGenerator.cs b/Engine/Source/Programs/UnrealBuildTool/System/XcodeProjectFileGenerator.cs
--- a/Engine/Source/Programs/UnrealBuildTool/System/XcodeProjectFileGenerator.cs
+++ b/Engine/Source/Programs/UnrealBuildTool/System/XcodeProjectFileGenerator.cs
@@ -885,7 +885,7 @@ namespace UnrealBuildTool
 		{
 			var XcodeProject = new XcodeProjectFile(ProjectPath);
 			DirectoryReference ProjectDirectory = ProjectPath.Directory;
-			XcodeProject.AddFilesToProject(SourceFileSearch.FindFiles(new List<DirectoryReference>() { ProjectDirectory }, SubdirectoryNamesToExclude:new List<string>() { "obj" }, SearchSubdirectories:true), null);
+			XcodeProject.AddFilesToProject(SourceFileSearch.FindFiles(ProjectDirectory, SubdirectoryNamesToExclude:new List<string>() { "obj" }, SearchSubdirectories:true), null);
 			XcodeProject.GenerateSectionsContents(ref PBXBuildFileSection, ref PBXFileReferenceSection, ref PBXSourcesBuildPhaseSection, ref Groups);
 		}
 
