Cache enumerated plugins between runs. Improves performance when generating project files.

diff --git a/Engine/Source/Programs/UnrealBuildTool/System/Plugins.cs b/Engine/Source/Programs/UnrealBuildTool/System/Plugins.cs
--- a/Engine/Source/Programs/UnrealBuildTool/System/Plugins.cs
+++ b/Engine/Source/Programs/UnrealBuildTool/System/Plugins.cs
@@ -17,7 +17,7 @@ namespace UnrealBuildTool
 		GameProject
 	}
 
-	[DebuggerDisplay("\\{{FileName}\\}")]
+	[DebuggerDisplay("\\{{File}\\}")]
 	public class PluginInfo
 	{
 		// Plugin name
@@ -52,6 +52,16 @@ namespace UnrealBuildTool
 
 	public class Plugins
 	{
+		/// <summary>
+		/// Cache of plugins under each directory
+		/// </summary>
+		static Dictionary<DirectoryReference, List<PluginInfo>> PluginInfoCache = new Dictionary<DirectoryReference,List<PluginInfo>>();
+
+		/// <summary>
+		/// Cache of plugin filenames under each directory
+		/// </summary>
+		static Dictionary<DirectoryReference, List<FileReference>> PluginFileCache = new Dictionary<DirectoryReference,List<FileReference>>();
+
 		/// <summary>
 		/// Read all the plugins available to a given project
 		/// </summary>
@@ -64,23 +74,37 @@ namespace UnrealBuildTool
 
 			// Read all the engine plugins
 			DirectoryReference EnginePluginsDirectoryName = DirectoryReference.Combine(EngineDirectoryName, "Plugins");
-			foreach(FileReference PluginFileName in EnumeratePlugins(EnginePluginsDirectoryName))
-			{
-				PluginInfo Plugin = new PluginInfo(PluginFileName, PluginLoadedFrom.Engine);
-				Plugins.Add(Plugin);
-			}
+			Plugins.AddRange(ReadAvailablePluginsInternal(EnginePluginsDirectoryName, PluginLoadedFrom.Engine));
 
 			// Read all the project plugins
 			if(ProjectFileName != null)
 			{
 				DirectoryReference ProjectPluginsDir = DirectoryReference.Combine(ProjectFileName.Directory, "Plugins");
-				foreach(FileReference PluginFileName in EnumeratePlugins(ProjectPluginsDir))
+				Plugins.AddRange(ReadAvailablePluginsInternal(ProjectPluginsDir, PluginLoadedFrom.GameProject));
+			}
+
+			return Plugins;
+		}
+
+		/// <summary>
+		/// Read all the plugin descriptors under the given directory
+		/// </summary>
+		/// <param name="ParentDirectory">The parent directory to look in.</param>
+		/// <param name="LoadedFrom">The directory type</param>
+		/// <returns>Sequence of the found PluginInfo object.</returns>
+		static IEnumerable<PluginInfo> ReadAvailablePluginsInternal(DirectoryReference ParentDirectory, PluginLoadedFrom LoadedFrom)
+		{
+			List<PluginInfo> Plugins;
+			if(!PluginInfoCache.TryGetValue(ParentDirectory, out Plugins))
+			{
+				Plugins = new List<PluginInfo>();
+				foreach(FileReference PluginFileName in EnumeratePlugins(ParentDirectory))
 				{
-					PluginInfo Plugin = new PluginInfo(PluginFileName, PluginLoadedFrom.GameProject);
+					PluginInfo Plugin = new PluginInfo(PluginFileName, LoadedFrom);
 					Plugins.Add(Plugin);
 				}
+				PluginInfoCache.Add(ParentDirectory, Plugins);
 			}
-
 			return Plugins;
 		}
 
@@ -90,10 +114,15 @@ namespace UnrealBuildTool
 		/// <param name="ParentDirectory">Parent directory to look in. Plugins will be found in any *subfolders* of this directory.</param>
 		public static IEnumerable<FileReference> EnumeratePlugins(DirectoryReference ParentDirectory)
 		{
-			List<FileReference> FileNames = new List<FileReference>();
-			if(ParentDirectory.Exists())
+			List<FileReference> FileNames;
+			if(!PluginFileCache.TryGetValue(ParentDirectory, out FileNames))
 			{
-				EnumeratePluginsInternal(new DirectoryInfo(ParentDirectory.FullName), FileNames);
+				FileNames = new List<FileReference>();
+				if(ParentDirectory.Exists())
+				{
+					EnumeratePluginsInternal(ParentDirectory, FileNames);
+				}
+				PluginFileCache.Add(ParentDirectory, FileNames);
 			}
 			return FileNames;
 		}
@@ -103,14 +132,14 @@ namespace UnrealBuildTool
 		/// </summary>
 		/// <param name="ParentDirectory">Parent directory to look in. Plugins will be found in any *subfolders* of this directory.</param>
 		/// <param name="FileNames">List of filenames. Will have all the discovered .uplugin files appended to it.</param>
-		static void EnumeratePluginsInternal(DirectoryInfo ParentDirectory, List<FileReference> FileNames)
+		static void EnumeratePluginsInternal(DirectoryReference ParentDirectory, List<FileReference> FileNames)
 		{
-			foreach(DirectoryInfo ChildDirectory in ParentDirectory.EnumerateDirectories())
+			foreach(DirectoryReference ChildDirectory in ParentDirectory.EnumerateDirectoryReferences())
 			{
 				int InitialFileNamesCount = FileNames.Count;
-				foreach(FileInfo PluginFileInfo in ChildDirectory.EnumerateFiles("*.uplugin", SearchOption.TopDirectoryOnly))
+				foreach(FileReference PluginFile in ChildDirectory.EnumerateFileReferences("*.uplugin"))
 				{
-					FileNames.Add(new FileReference(PluginFileInfo.FullName));
+					FileNames.Add(PluginFile);
 				}
 				if(FileNames.Count == InitialFileNamesCount)
 				{
