Only determine the module distribution level when it's actually required, and don't do it when generating project files.

diff --git a/Engine/Source/Programs/UnrealBuildTool/Configuration/UEBuildBinary.cs b/Engine/Source/Programs/UnrealBuildTool/Configuration/UEBuildBinary.cs
--- a/Engine/Source/Programs/UnrealBuildTool/Configuration/UEBuildBinary.cs
+++ b/Engine/Source/Programs/UnrealBuildTool/Configuration/UEBuildBinary.cs
@@ -359,7 +359,7 @@ namespace UnrealBuildTool
 		 * Checks whether the binary output paths are appropriate for the distribution
 		 * level of its direct module dependencies
 		 */
-		public void CheckOutputDistributionLevelAgainstDependencies()
+		public void CheckOutputDistributionLevelAgainstDependencies(Dictionary<UEBuildModule, UEBuildModuleDistribution> ModuleDistributionCache)
 		{
 			// Find maximum distribution level of its direct dependencies
 			var DistributionLevel = UEBuildModuleDistribution.Public;
@@ -367,18 +367,24 @@ namespace UnrealBuildTool
 			List<string>[] DependantModuleNames = new List<string>[Enum.GetNames(typeof(UEBuildModuleDistribution)).Length];
 			foreach (var Module in DependantModules)
 			{
-				if (Module.DistributionLevel != UEBuildModuleDistribution.Public)
+				UEBuildModuleDistribution ModuleDistributionLevel;
+				if(!ModuleDistributionCache.TryGetValue(Module, out ModuleDistributionLevel))
+				{
+					ModuleDistributionLevel = Module.DetermineDistributionLevel();
+					ModuleDistributionCache.Add(Module, ModuleDistributionLevel);
+				}
+				if (ModuleDistributionLevel != UEBuildModuleDistribution.Public)
 				{
 					// Make a list of non-public dependant modules so that exception
 					// message can be more helpful
-					int DistributionIndex = (int)Module.DistributionLevel;
+					int DistributionIndex = (int)ModuleDistributionLevel;
 					if (DependantModuleNames[DistributionIndex] == null)
 					{
 						DependantModuleNames[DistributionIndex] = new List<string>();
 					}
 					DependantModuleNames[DistributionIndex].Add(Module.Name);
 
-					DistributionLevel = Utils.Max(DistributionLevel, Module.DistributionLevel);
+					DistributionLevel = Utils.Max(DistributionLevel, ModuleDistributionLevel);
 				}
 			}
 
diff --git a/Engine/Source/Programs/UnrealBuildTool/Configuration/UEBuildModule.cs b/Engine/Source/Programs/UnrealBuildTool/Configuration/UEBuildModule.cs
--- a/Engine/Source/Programs/UnrealBuildTool/Configuration/UEBuildModule.cs
+++ b/Engine/Source/Programs/UnrealBuildTool/Configuration/UEBuildModule.cs
@@ -136,7 +136,7 @@ namespace UnrealBuildTool
 		/**
 		 * Determines the distribution level of a module based on its directory and includes.
 		 */
-		private void SetupModuleDistributionLevel()
+		public UEBuildModuleDistribution DetermineDistributionLevel()
 		{
 			List<string> PathsToCheck = new List<string>();
 			PathsToCheck.Add(ModuleDirectory.FullName);
@@ -146,13 +146,15 @@ namespace UnrealBuildTool
 			PathsToCheck.AddRange(PublicSystemIncludePaths);
 			PathsToCheck.AddRange(PublicLibraryPaths);
 
-			DistributionLevel = UEBuildModuleDistribution.Public;
+			UEBuildModuleDistribution DistributionLevel = UEBuildModuleDistribution.Public;
 
 			// Keep checking as long as we haven't reached the maximum level
 			for (int PathIndex = 0; PathIndex < PathsToCheck.Count && DistributionLevel != UEBuildModuleDistribution.NoRedist; ++PathIndex)
 			{
 				DistributionLevel = Utils.Max(DistributionLevel, UEBuildModule.GetModuleDistributionLevelBasedOnLocation(PathsToCheck[PathIndex]));
 			}
+
+			return DistributionLevel;
 		}
 
 		/** Converts an optional string list parameter to a well-defined hash set. */
@@ -173,9 +175,6 @@ namespace UnrealBuildTool
 		/** The rules for this module */
 		public ModuleRules Rules;
 
-		/** The distribution level of the module being built. Used to check where build products are placed. */
-		public UEBuildModuleDistribution DistributionLevel;
-
 		/** Path to the module directory */
 		public readonly DirectoryReference ModuleDirectory;
 
@@ -272,8 +271,6 @@ namespace UnrealBuildTool
 			Debug.Assert(InBuildCsFilename == null || InBuildCsFilename.EndsWith(".Build.cs", StringComparison.InvariantCultureIgnoreCase));
 			BuildCsFilenameField = InBuildCsFilename;
 
-			SetupModuleDistributionLevel();
-
 			Target.RegisterModule(this);
 		}
 
diff --git a/Engine/Source/Programs/UnrealBuildTool/Configuration/UEBuildTarget.cs b/Engine/Source/Programs/UnrealBuildTool/Configuration/UEBuildTarget.cs
--- a/Engine/Source/Programs/UnrealBuildTool/Configuration/UEBuildTarget.cs
+++ b/Engine/Source/Programs/UnrealBuildTool/Configuration/UEBuildTarget.cs
@@ -2239,12 +2239,13 @@ namespace UnrealBuildTool
 				}
 			}
 
-			if (ShouldCheckOutputDistributionLevel())
+			if (ShouldCheckOutputDistributionLevel() && !ProjectFileGenerator.bGenerateProjectFiles)
 			{
 				// Check the distribution level of all binaries based on the dependencies they have
+				Dictionary<UEBuildModule, UEBuildModuleDistribution> ModuleDistributionCache = new Dictionary<UEBuildModule, UEBuildModuleDistribution>();
 				foreach (var Binary in AppBinaries)
 				{
-					Binary.CheckOutputDistributionLevelAgainstDependencies();
+					Binary.CheckOutputDistributionLevelAgainstDependencies(ModuleDistributionCache);
 				}
 			}
 		}
