Use manual binary serialization for dependency caches, instead of the standard .NET serializer. Reduces load time from ~1s to <0.1s, and reduces size to half or less.

diff --git a/Engine/Source/Programs/UnrealBuildTool/System/DependencyCache.cs b/Engine/Source/Programs/UnrealBuildTool/System/DependencyCache.cs
--- a/Engine/Source/Programs/UnrealBuildTool/System/DependencyCache.cs
+++ b/Engine/Source/Programs/UnrealBuildTool/System/DependencyCache.cs
@@ -9,18 +9,9 @@ using System.Runtime.Serialization.Formatters.Binary;
 
 namespace UnrealBuildTool
 {
-	[Serializable]
+	[DebuggerDisplay("{IncludeName}")]
 	public class DependencyInclude
 	{
-		/// <summary>
-		/// Public ctor that initializes the include name (the resolved name won't be determined until later)
-		/// </summary>
-		/// <param name="InIncludeName"></param>
-		public DependencyInclude(string InIncludeName)
-		{
-			IncludeName = InIncludeName;
-		}
-
 		/// <summary>
 		/// These are direct include paths and cannot be resolved to an actual file on disk without using the proper list of include directories for this file's module 
 		/// </summary>
@@ -36,49 +27,43 @@ namespace UnrealBuildTool
 		/// </summary>
 		public FileReference IncludeResolvedNameIfSuccessful;
 
-		public override string ToString()
+		/// <summary>
+		/// Public ctor that initializes the include name (the resolved name won't be determined until later)
+		/// </summary>
+		/// <param name="InIncludeName"></param>
+		public DependencyInclude(string InIncludeName)
 		{
-			return IncludeName;
+			IncludeName = InIncludeName;
 		}
 	}
 
-	[Serializable]
-	struct DependencyInfo
-	{
-		/** List of files this file includes.  These are direct include paths and cannot be resolved to an actual file on 
-		    disk without using the proper list of include directories for this file's module */
-		public List<DependencyInclude> Includes;
-
-		/** Deprecated - always false */
-		public bool HasUObjects;
-	}
-
-
 	/**
 	 * Caches include dependency information to speed up preprocessing on subsequent runs.
 	 */
-	[Serializable]
 	public class DependencyCache
 	{
-		/** The time the cache was created. Used to invalidate entries. */
-		public DateTimeOffset CacheCreateDate;
+		/** The version number for binary serialization */
+		const int FileVersion = 1;
 
-		/** The time the cache was last updated. Stored as the creation date when saved. */
-		[NonSerialized]
-		private DateTimeOffset CacheUpdateDate;
+		/** The file signature for binary serialization */
+		const int FileSignature = ('D' << 24) | ('C' << 16) | FileVersion;
 
 		/** Path to store the cache data to. */
-		[NonSerialized]
-		private FileReference CachePath;
+		private FileReference BackingFile;
+
+		/** The time the cache was created. Used to invalidate entries. */
+		public DateTime CreationTimeUtc;
+
+		/** The time the cache was last updated. Stored as the creation date when saved. Not serialized. */
+		private DateTime UpdateTimeUtc;
 
-		/** Dependency lists, keyed (case-insensitively) on file's absolute path. */
-		private Dictionary<FileReference, DependencyInfo> DependencyMap;
+		/** Dependency lists, keyed on file's absolute path. */
+		private Dictionary<FileReference, List<DependencyInclude>> DependencyMap;
 
-		[NonSerialized]
+		/** A mapping of whether various files exist. Not serialized. */
 		private Dictionary<FileReference, bool> FileExistsInfo;
 
-		/** Whether the dependency cache is dirty and needs to be saved. */
-		[NonSerialized]
+		/** Whether the dependency cache is dirty and needs to be saved. Not serialized. */
 		private bool bIsDirty;
 
 		/**
@@ -86,26 +71,24 @@ namespace UnrealBuildTool
 		 * 
 		 * @param	CachePath	Name of the cache file to deserialize
 		 */
-		public static DependencyCache Create(FileReference CachePath)
+		public static DependencyCache Create(FileReference CacheFile)
 		{
 			// See whether the cache file exists.
-			FileItem Cache = FileItem.GetItemByFileReference(CachePath);
-			if (Cache.bExists)
+			if (CacheFile.Exists())
 			{
 				if (BuildConfiguration.bPrintPerformanceInfo)
 				{
-					Log.TraceInformation("Loading existing IncludeFileCache: " + Cache.AbsolutePath);
+					Log.TraceInformation("Loading existing IncludeFileCache: " + CacheFile.FullName);
 				}
 
 				var TimerStartTime = DateTime.UtcNow;
 
 				// Deserialize cache from disk if there is one.
-				DependencyCache Result = Load(Cache);
-				if (Result != null)
+				DependencyCache Result = Load(CacheFile);
+				if(Result != null)
 				{
 					// Successfully serialize, create the transient variables and return cache.
-					Result.CachePath = CachePath;
-					Result.CacheUpdateDate = DateTimeOffset.Now;
+					Result.UpdateTimeUtc = DateTime.UtcNow;
 
 					var TimerDuration = DateTime.UtcNow - TimerStartTime;
 					if (BuildConfiguration.bPrintPerformanceInfo)
@@ -118,7 +101,7 @@ namespace UnrealBuildTool
 				}
 			}
 			// Fall back to a clean cache on error or non-existance.
-			return new DependencyCache(Cache);
+			return new DependencyCache(CacheFile);
 		}
 
 		/**
@@ -126,12 +109,12 @@ namespace UnrealBuildTool
 		 * 
 		 * @param	Cache	File to deserialize from
 		 */
-		public static DependencyCache Load(FileItem Cache)
+		public static DependencyCache Load(FileReference CacheFile)
 		{
 			DependencyCache Result = null;
 			try
 			{
-				string CacheBuildMutexPath = Cache.AbsolutePath + ".buildmutex";
+				string CacheBuildMutexPath = CacheFile.FullName + ".buildmutex";
 
 				// If the .buildmutex file for the cache is present, it means that something went wrong between loading
 				// and saving the cache last time (most likely the UBT process being terminated), so we don't want to load
@@ -142,27 +125,82 @@ namespace UnrealBuildTool
 					{
 					}
 
-					using (FileStream Stream = new FileStream(Cache.AbsolutePath, FileMode.Open, FileAccess.Read))
+					using (BinaryReader Reader = new BinaryReader(new FileStream(CacheFile.FullName, FileMode.Open, FileAccess.Read)))
 					{
-						BinaryFormatter Formatter = new BinaryFormatter();
-						Result = Formatter.Deserialize(Stream) as DependencyCache;
+						if(Reader.ReadInt32() == FileSignature)
+						{
+							Result = DependencyCache.Deserialize(Reader);
+						}
 					}
-					Result.CreateFileExistsInfo();
-					Result.ResetUnresolvedDependencies();
 				}
 			}
 			catch (Exception Ex)
 			{
-				// Don't bother logging this expected error.
-				// It's due to a change in the CacheCreateDate type.
-				if (Ex.Message != "Object of type 'System.Collections.Generic.Dictionary`2[System.String,UnrealBuildTool.DependencyInfo]' cannot be converted to type 'System.Collections.Generic.Dictionary`2[UnrealBuildTool.FileReference,UnrealBuildTool.DependencyInfo]'."
-					&& Ex.Message != "Object of type 'System.DateTime' cannot be converted to type 'System.DateTimeOffset'" 
-					&& Ex.Message != "Object of type 'System.Collections.Generic.Dictionary`2[System.String,System.Collections.Generic.List`1[System.String]]' cannot be converted to type 'System.Collections.Generic.Dictionary`2[System.String,UnrealBuildTool.DependencyInfo]'.")	// To catch serialization differences added when we added the DependencyInfo struct
+				Console.Error.WriteLine("Failed to read dependency cache: {0}", Ex.Message);
+			}
+			return Result;
+		}
+
+		/// <summary>
+		/// Serializes the dependency cache to a binary writer
+		/// </summary>
+		/// <param name="Writer">Writer to output to</param>
+		void Serialize(BinaryWriter Writer)
+		{
+			Writer.Write(BackingFile);
+			Writer.Write(CreationTimeUtc.ToBinary());
+
+			Dictionary<FileReference, int> FileToUniqueId = new Dictionary<FileReference,int>();
+
+			Writer.Write(DependencyMap.Count);
+			foreach(KeyValuePair<FileReference, List<DependencyInclude>> Pair in DependencyMap)
+			{
+				Writer.Write(Pair.Key);
+				Writer.Write(Pair.Value.Count);
+				foreach(DependencyInclude Include in Pair.Value)
 				{
-					Console.Error.WriteLine("Failed to read dependency cache: {0}", Ex.Message);
+					Writer.Write(Include.IncludeName);
+					Writer.Write(Include.HasAttemptedResolve);
+					Writer.Write(Include.IncludeResolvedNameIfSuccessful, FileToUniqueId);
 				}
 			}
-			return Result;
+		}
+
+		/// <summary>
+		/// Deserialize the dependency cache from a binary reader
+		/// </summary>
+		/// <param name="Reader">Reader for the cache data</param>
+		/// <returns>New dependency cache object</returns>
+		static DependencyCache Deserialize(BinaryReader Reader)
+		{
+			DependencyCache Cache = new DependencyCache(Reader.ReadFileReference());
+			Cache.CreationTimeUtc = DateTime.FromBinary(Reader.ReadInt64());
+
+			int NumEntries = Reader.ReadInt32();
+			Cache.DependencyMap = new Dictionary<FileReference,List<DependencyInclude>>(NumEntries);
+
+			List<FileReference> UniqueFiles = new List<FileReference>();
+			for(int Idx = 0; Idx < NumEntries; Idx++)
+			{
+				FileReference File = Reader.ReadFileReference();
+
+				int NumIncludes = Reader.ReadInt32();
+				List<DependencyInclude> Includes = new List<DependencyInclude>(NumIncludes);
+
+				for(int IncludeIdx = 0; IncludeIdx < NumIncludes; IncludeIdx++)
+				{
+					DependencyInclude Include = new DependencyInclude(Reader.ReadString());
+					Include.HasAttemptedResolve = Reader.ReadBoolean();
+					Include.IncludeResolvedNameIfSuccessful = Reader.ReadFileReference(UniqueFiles);
+					Includes.Add(Include);
+				}
+
+				Cache.DependencyMap.Add(File, Includes);
+			}
+
+			Cache.CreateFileExistsInfo();
+			Cache.ResetUnresolvedDependencies();
+			return Cache;
 		}
 
 		/**
@@ -170,12 +208,12 @@ namespace UnrealBuildTool
 		 * 
 		 * @param	Cache	File associated with this cache
 		 */
-		protected DependencyCache(FileItem Cache)
+		protected DependencyCache(FileReference InBackingFile)
 		{
-			CacheCreateDate = DateTimeOffset.Now;
-			CacheUpdateDate = DateTimeOffset.Now;
-			CachePath = Cache.Reference;
-			DependencyMap = new Dictionary<FileReference, DependencyInfo>();
+			BackingFile = InBackingFile;
+			CreationTimeUtc = DateTime.UtcNow;
+			UpdateTimeUtc = DateTime.UtcNow;
+			DependencyMap = new Dictionary<FileReference, List<DependencyInclude>>();
 			bIsDirty = false;
 			CreateFileExistsInfo();
 		}
@@ -191,16 +229,16 @@ namespace UnrealBuildTool
 				var TimerStartTime = DateTime.UtcNow;
 
 				// Save update date as new creation date.
-				CacheCreateDate = CacheUpdateDate;
+				CreationTimeUtc = UpdateTimeUtc;
 
 				// Serialize the cache to disk.
 				try
 				{
-					CachePath.Directory.CreateDirectory();
-					using (FileStream Stream = new FileStream(CachePath.FullName, FileMode.Create, FileAccess.Write))
+					BackingFile.Directory.CreateDirectory();
+					using(BinaryWriter Writer = new BinaryWriter(new FileStream(BackingFile.FullName, FileMode.Create, FileAccess.Write)))
 					{
-						BinaryFormatter Formatter = new BinaryFormatter();
-						Formatter.Serialize(Stream, this);
+						Writer.Write(FileSignature);
+						Serialize(Writer);
 					}
 				}
 				catch (Exception Ex)
@@ -222,7 +260,7 @@ namespace UnrealBuildTool
 				}
 			}
 
-			FileReference MutexPath = CachePath + ".buildmutex";
+			FileReference MutexPath = BackingFile + ".buildmutex";
 			if(MutexPath.Exists())
 			{
 				try
@@ -245,14 +283,14 @@ namespace UnrealBuildTool
 		public List<DependencyInclude> GetCachedDependencyInfo(FileItem File)
 		{
 			// Check whether File is in cache.
-			DependencyInfo Result;
-			if (!DependencyMap.TryGetValue(File.Reference, out Result))
+			List<DependencyInclude> Includes;
+			if (!DependencyMap.TryGetValue(File.Reference, out Includes))
 			{
 				return null;
 			}
 
 			// File is in cache, now check whether last write time is prior to cache creation time.
-			if (File.LastWriteTime >= CacheCreateDate)
+			if (File.LastWriteTime.ToUniversalTime() >= CreationTimeUtc)
 			{
 				// Remove entry from cache as it's stale.
 				DependencyMap.Remove(File.Reference);
@@ -261,7 +299,7 @@ namespace UnrealBuildTool
 			}
 
 			// Check if any of the resolved includes is missing
-			foreach (var Include in Result.Includes)
+			foreach (var Include in Includes)
 			{
 				if (Include.IncludeResolvedNameIfSuccessful != null)
 				{
@@ -284,7 +322,7 @@ namespace UnrealBuildTool
 			}
 
 			// Cached version is up to date, return it.
-			return Result.Includes;
+			return Includes;
 		}
 
 		/**
@@ -296,7 +334,7 @@ namespace UnrealBuildTool
 		 */
 		public void SetDependencyInfo(FileItem File, List<DependencyInclude> Info)
 		{
-			DependencyMap[File.Reference] = new DependencyInfo{ Includes = Info, HasUObjects = false };
+			DependencyMap[File.Reference] = Info;
 			bIsDirty = true;
 		}
 
@@ -317,7 +355,7 @@ namespace UnrealBuildTool
 		{
 			foreach (var Dependency in DependencyMap)
 			{
-				foreach (var Include in Dependency.Value.Includes)
+				foreach (var Include in Dependency.Value)
 				{
 					if (Include.HasAttemptedResolve && Include.IncludeResolvedNameIfSuccessful == null)
 					{
@@ -339,7 +377,7 @@ namespace UnrealBuildTool
 		{
 			if (BuildConfiguration.bUseIncludeDependencyResolveCache)
 			{
-				var Includes = DependencyMap[File.Reference].Includes;
+				var Includes = DependencyMap[File.Reference];
 				var IncludeToResolve = Includes[DirectlyIncludedFileNameIndex];
 				if (BuildConfiguration.bTestIncludeDependencyResolveCache)
 				{
diff --git a/Engine/Source/Programs/UnrealBuildTool/System/FileSystemReference.cs b/Engine/Source/Programs/UnrealBuildTool/System/FileSystemReference.cs
--- a/Engine/Source/Programs/UnrealBuildTool/System/FileSystemReference.cs
+++ b/Engine/Source/Programs/UnrealBuildTool/System/FileSystemReference.cs
@@ -707,4 +707,77 @@ namespace UnrealBuildTool
 			return new FileReference(AbsolutePath, AbsolutePath.ToLowerInvariant());
 		}
 	}
+
+	static class FileReferenceExtensionMethods
+	{
+		/// <summary>
+		/// Manually serialize a file reference to a binary stream.
+		/// </summary>
+		/// <param name="Writer">Binary writer to write to</param>
+		public static void Write(this BinaryWriter Writer, FileReference File)
+		{
+			Writer.Write((File == null)? String.Empty : File.FullName);
+		}
+
+		/// <summary>
+		/// Serializes a file reference, using a lookup table to avoid serializing the same name more than once.
+		/// </summary>
+		/// <param name="Writer">The writer to save this reference to</param>
+		/// <param name="File">A file reference to output; may be null</param>
+		/// <param name="FileToUniqueId">A lookup table that caches previous files that have been output, and maps them to unique id's.</param>
+		public static void Write(this BinaryWriter Writer, FileReference File, Dictionary<FileReference, int> FileToUniqueId)
+		{
+			int UniqueId;
+			if(File == null)
+			{
+				Writer.Write(-1);
+			}
+			else if(FileToUniqueId.TryGetValue(File, out UniqueId))
+			{
+				Writer.Write(UniqueId);
+			}
+			else
+			{
+				Writer.Write(FileToUniqueId.Count);
+				Writer.Write(File);
+				FileToUniqueId.Add(File, FileToUniqueId.Count);
+			}
+		}
+
+		/// <summary>
+		/// Manually deserialize a file reference from a binary stream.
+		/// </summary>
+		/// <param name="Reader">Binary reader to read from</param>
+		/// <returns>New FileReference object</returns>
+		public static FileReference ReadFileReference(this BinaryReader Reader)
+		{
+			string FullName = Reader.ReadString();
+			return (FullName.Length == 0)? null : FileReference.MakeFromNormalizedFullPath(FullName);
+		}
+
+		/// <summary>
+		/// Deserializes a file reference, using a lookup table to avoid writing the same name more than once.
+		/// </summary>
+		/// <param name="Reader">The source to read from</param>
+		/// <param name="UniqueFiles">List of previously read file references. The index into this array is used in place of subsequent ocurrences of the file.</param>
+		/// <returns>The file reference that was read</returns>
+		public static FileReference ReadFileReference(this BinaryReader Reader, List<FileReference> UniqueFiles)
+		{
+			int UniqueId = Reader.ReadInt32();
+			if(UniqueId == -1)
+			{
+				return null;
+			}
+			else if(UniqueId < UniqueFiles.Count)
+			{
+				return UniqueFiles[UniqueId];
+			}
+			else
+			{
+				FileReference Result = Reader.ReadFileReference();
+				UniqueFiles.Add(Result);
+				return Result;
+			}
+		}
+	}
 }
diff --git a/Engine/Source/Programs/UnrealBuildTool/System/FlatCPPIncludeDepencencyCache.cs b/Engine/Source/Programs/UnrealBuildTool/System/FlatCPPIncludeDepencencyCache.cs
--- a/Engine/Source/Programs/UnrealBuildTool/System/FlatCPPIncludeDepencencyCache.cs
+++ b/Engine/Source/Programs/UnrealBuildTool/System/FlatCPPIncludeDepencencyCache.cs
@@ -12,7 +12,6 @@ namespace UnrealBuildTool
 	/// <summary>
 	/// Stores an ordered list of header files
 	/// </summary>
-	[Serializable]
 	class FlatCPPIncludeDependencyInfo
 	{
 		/// The PCH header this file is dependent on.
@@ -23,7 +22,6 @@ namespace UnrealBuildTool
 
 		/// Transient cached list of FileItems for all of the includes for a specific files.  This just saves a bunch of string
 		/// hash lookups as we locate FileItems for files that we've already requested dependencies for
-		[NonSerialized]
 		public List<FileItem> IncludeFileItems;
 	}
 
@@ -31,9 +29,23 @@ namespace UnrealBuildTool
 	/// <summary>
 	/// For a given target, caches all of the C++ source files and the files they are including
 	/// </summary>
-	[Serializable]
 	public class FlatCPPIncludeDependencyCache
 	{
+		/** The version number for binary serialization */
+		const int FileVersion = 1;
+
+		/** The file signature for binary serialization */
+		const int FileSignature = ('F' << 24) | ('C' << 16) | FileVersion;
+
+		/// File name of this cache, should be unique for every includes context (e.g. target)
+		private FileReference BackingFile;
+
+		/// True if the cache needs to be saved
+		private bool bIsDirty = false;
+
+		/// Dependency lists, keyed (case-insensitively) on file's absolute path.
+		private Dictionary<FileReference, FlatCPPIncludeDependencyInfo> DependencyMap;
+
 		/// <summary>
 		/// Creates the cache object
 		/// </summary>
@@ -41,21 +53,20 @@ namespace UnrealBuildTool
 		/// <returns>The new instance</returns>
 		public static FlatCPPIncludeDependencyCache Create( UEBuildTarget Target )
 		{
-			FileReference CachePath = FlatCPPIncludeDependencyCache.GetDependencyCachePathForTarget( Target );;
+			FileReference CacheFile = FlatCPPIncludeDependencyCache.GetDependencyCachePathForTarget( Target );;
 
 			// See whether the cache file exists.
-			FileItem Cache = FileItem.GetItemByFileReference(CachePath);
-			if (Cache.bExists)
+			if (CacheFile.Exists())
 			{
 				if (BuildConfiguration.bPrintPerformanceInfo)
 				{
-					Log.TraceInformation("Loading existing FlatCPPIncludeDependencyCache: " + Cache.AbsolutePath);
+					Log.TraceInformation("Loading existing FlatCPPIncludeDependencyCache: " + CacheFile.FullName);
 				}
 
 				var TimerStartTime = DateTime.UtcNow;
 
 				// Deserialize cache from disk if there is one.
-				FlatCPPIncludeDependencyCache Result = Load(Cache);
+				FlatCPPIncludeDependencyCache Result = Load(CacheFile);
 				if (Result != null)
 				{
 					var TimerDuration = DateTime.UtcNow - TimerStartTime;
@@ -75,7 +86,7 @@ namespace UnrealBuildTool
 			}
 
 			// Fall back to a clean cache on error or non-existence.
-			return new FlatCPPIncludeDependencyCache( Cache );
+			return new FlatCPPIncludeDependencyCache( CacheFile );
 		}
 
 		/// <summary>
@@ -83,12 +94,12 @@ namespace UnrealBuildTool
 		/// </summary>
 		/// <param name="Cache">The file to load</param>
 		/// <returns>The loaded instance</returns>
-		public static FlatCPPIncludeDependencyCache Load(FileItem Cache)
+		public static FlatCPPIncludeDependencyCache Load(FileReference BackingFile)
 		{
 			FlatCPPIncludeDependencyCache Result = null;
 			try
 			{
-				string CacheBuildMutexPath = Cache.AbsolutePath + ".buildmutex";
+				string CacheBuildMutexPath = BackingFile.FullName + ".buildmutex";
 
 				// If the .buildmutex file for the cache is present, it means that something went wrong between loading
 				// and saving the cache last time (most likely the UBT process being terminated), so we don't want to load
@@ -99,24 +110,20 @@ namespace UnrealBuildTool
 					{
 					}
 
-					using (FileStream Stream = new FileStream(Cache.AbsolutePath, FileMode.Open, FileAccess.Read))
+					using (BinaryReader Reader = new BinaryReader(new FileStream(BackingFile.FullName, FileMode.Open, FileAccess.Read)))
 					{
 						// @todo ubtmake: We can store the cache in a cheaper/smaller way using hash file names and indices into included headers, but it might actually slow down load times
 						// @todo ubtmake: If we can index PCHs here, we can avoid storing all of the PCH's included headers (PCH's action should have been invalidated, so we shouldn't even have to report the PCH's includes as our indirect includes)
-						BinaryFormatter Formatter = new BinaryFormatter();
-						Result = Formatter.Deserialize(Stream) as FlatCPPIncludeDependencyCache;
-						Result.CacheFileItem = Cache;
-						Result.bIsDirty = false;
+						if(Reader.ReadInt32() == FileSignature)
+						{
+							Result = Deserialize(Reader);
+						}
 					}
 				}
 			}
 			catch (Exception Ex)
 			{
-				// Don't bother failing if the file format has changed, simply abort the cache load
-				if (!Ex.Message.Contains( "cannot be converted to type" ))	// To catch serialization differences added when we added the DependencyInfo struct
-				{
-					Console.Error.WriteLine("Failed to read FlatCPPIncludeDependencyCache: {0}", Ex.Message);
-				}
+				Console.Error.WriteLine("Failed to read FlatCPPIncludeDependencyCache: {0}", Ex.Message);
 			}
 			return Result;
 		}
@@ -126,9 +133,9 @@ namespace UnrealBuildTool
 		/// Constructs a fresh cache, storing the file name that it should be saved as later on
 		/// </summary>
 		/// <param name="Cache">File name for this cache, usually unique per context (e.g. target)</param>
-		protected FlatCPPIncludeDependencyCache(FileItem Cache)
+		protected FlatCPPIncludeDependencyCache(FileReference InBackingFile)
 		{
-			CacheFileItem = Cache;
+			BackingFile = InBackingFile;
 			DependencyMap = new Dictionary<FileReference, FlatCPPIncludeDependencyInfo>();
 		}
 
@@ -146,11 +153,11 @@ namespace UnrealBuildTool
 				// Serialize the cache to disk.
 				try
 				{
-					Directory.CreateDirectory(Path.GetDirectoryName(CacheFileItem.AbsolutePath));
-					using (FileStream Stream = new FileStream(CacheFileItem.AbsolutePath, FileMode.Create, FileAccess.Write))
+					BackingFile.Directory.CreateDirectory();
+					using (BinaryWriter Writer = new BinaryWriter(new FileStream(BackingFile.FullName, FileMode.Create, FileAccess.Write)))
 					{
-						BinaryFormatter Formatter = new BinaryFormatter();
-						Formatter.Serialize(Stream, this);
+						Writer.Write(FileSignature);
+						Serialize(Writer);
 					}
 				}
 				catch (Exception Ex)
@@ -172,11 +179,11 @@ namespace UnrealBuildTool
 				}
 			}
 
-			if(File.Exists(CacheFileItem.AbsolutePath + ".buildmutex"))
+			if(File.Exists(BackingFile.FullName + ".buildmutex"))
 			{
 				try
 				{
-					File.Delete(CacheFileItem.AbsolutePath + ".buildmutex");
+					File.Delete(BackingFile.FullName + ".buildmutex");
 				}
 				catch
 				{
@@ -185,6 +192,71 @@ namespace UnrealBuildTool
 			}
 		}
 
+		/// <summary>
+		/// Serializes the dependency cache to a binary writer
+		/// </summary>
+		/// <param name="Writer">Writer to output to</param>
+		void Serialize(BinaryWriter Writer)
+		{
+			Dictionary<FileReference, int> FileToUniqueId = new Dictionary<FileReference,int>();
+			Writer.Write(BackingFile);
+
+			// Write out all the dependency info objects
+			Writer.Write(DependencyMap.Count);
+			foreach(KeyValuePair<FileReference, FlatCPPIncludeDependencyInfo> Pair in DependencyMap)
+			{
+				Writer.Write(Pair.Key);
+
+				FlatCPPIncludeDependencyInfo Info = Pair.Value;
+				Writer.Write(Info.PCHName, FileToUniqueId);
+
+				Writer.Write(Info.Includes.Count);
+				foreach(FileReference Include in Info.Includes)
+				{
+					Writer.Write(Include, FileToUniqueId);
+				}
+			}
+		}
+
+		/// <summary>
+		/// Deserialize the dependency cache from a binary reader
+		/// </summary>
+		/// <param name="Reader">Reader for the cache data</param>
+		/// <returns>New dependency cache object</returns>
+		static FlatCPPIncludeDependencyCache Deserialize(BinaryReader Reader)
+		{
+			FlatCPPIncludeDependencyCache Cache = new FlatCPPIncludeDependencyCache(Reader.ReadFileReference());
+
+			// Create the dependency map
+			int NumDependencyMapEntries = Reader.ReadInt32();
+			Cache.DependencyMap = new Dictionary<FileReference,FlatCPPIncludeDependencyInfo>(NumDependencyMapEntries);
+
+			// Read all the dependency map entries
+			List<FileReference> UniqueFiles = new List<FileReference>(NumDependencyMapEntries * 2);
+			for(int Idx = 0; Idx < NumDependencyMapEntries; Idx++)
+			{
+				FileReference File = Reader.ReadFileReference();
+
+				FlatCPPIncludeDependencyInfo Info = new FlatCPPIncludeDependencyInfo();
+
+				// Read the PCH name
+				Info.PCHName = Reader.ReadFileReference(UniqueFiles);
+
+				// Read the includes 
+				int NumIncludes = Reader.ReadInt32();
+				Info.Includes = new List<FileReference>(NumIncludes);
+
+				for(int IncludeIdx = 0; IncludeIdx < NumIncludes; IncludeIdx++)
+				{
+					Info.Includes.Add(Reader.ReadFileReference(UniqueFiles));
+				}
+
+				// Add it to the map
+				Cache.DependencyMap.Add(File, Info);
+			}
+
+			return Cache;
+		}
 
 		/// <summary>
 		/// Sets the new dependencies for the specified file
@@ -248,19 +320,5 @@ namespace UnrealBuildTool
 			}
 			return FileReference.Combine(PlatformIntermediatePath, Target.GetTargetName(), "FlatCPPIncludes.bin" );
 		}
-
-
-		/// File name of this cache, should be unique for every includes context (e.g. target)
-		[NonSerialized]
-		private FileItem CacheFileItem;
-
-		/// True if the cache needs to be saved
-		[NonSerialized]
-		private bool bIsDirty = false;
-
-		/// Dependency lists, keyed (case-insensitively) on file's absolute path.
-		private Dictionary<FileReference, FlatCPPIncludeDependencyInfo> DependencyMap;
-
-
 	}
 }
