Cache directory searches when looking for include paths. When generating project files we do this many times over.

diff --git a/Engine/Source/Programs/UnrealBuildTool/Configuration/UEBuildTarget.cs b/Engine/Source/Programs/UnrealBuildTool/Configuration/UEBuildTarget.cs
--- a/Engine/Source/Programs/UnrealBuildTool/Configuration/UEBuildTarget.cs
+++ b/Engine/Source/Programs/UnrealBuildTool/Configuration/UEBuildTarget.cs
@@ -3565,22 +3565,20 @@ namespace UnrealBuildTool
 
 			// Add the 'classes' directory, if it exists
 			DirectoryReference ClassesDirectory = DirectoryReference.Combine(ModuleFile.Directory, "Classes");
-			if(ClassesDirectory.Exists())
+			if(DirectoryLookupCache.DirectoryExists(ClassesDirectory))
 			{
 				RulesObject.PublicIncludePaths.Add(NormalizeIncludePath(ClassesDirectory));
 			}
 
 			// Add all the public directories
 			DirectoryReference PublicDirectory = DirectoryReference.Combine(ModuleFile.Directory, "Public");
-			if(PublicDirectory.Exists())
+			if(DirectoryLookupCache.DirectoryExists(PublicDirectory))
 			{
 				RulesObject.PublicIncludePaths.Add(NormalizeIncludePath(PublicDirectory));
 
-				// Add subdirectories of Public if present
-				DirectoryInfo PublicInfo = new DirectoryInfo(PublicDirectory.FullName);
-				foreach (DirectoryInfo SubDir in PublicInfo.GetDirectories("*", SearchOption.AllDirectories))
+				foreach(DirectoryReference PublicSubDirectory in DirectoryLookupCache.EnumerateDirectoriesRecursively(PublicDirectory))
 				{
-					RulesObject.PublicIncludePaths.Add(NormalizeIncludePath(new DirectoryReference(SubDir)));
+					RulesObject.PublicIncludePaths.Add(NormalizeIncludePath(PublicSubDirectory));
 				}
 			}
 		}
diff --git a/Engine/Source/Programs/UnrealBuildTool/System/DirectoryLookupCache.cs b/Engine/Source/Programs/UnrealBuildTool/System/DirectoryLookupCache.cs
--- a/Engine/Source/Programs/UnrealBuildTool/System/DirectoryLookupCache.cs
+++ b/Engine/Source/Programs/UnrealBuildTool/System/DirectoryLookupCache.cs
@@ -8,43 +8,108 @@ namespace UnrealBuildTool
 {
 	class DirectoryCache
 	{
-		public DirectoryCache( DirectoryReference BaseDirectory )
+		public readonly DirectoryReference Directory;
+		public readonly bool Exists;
+		HashSet<FileReference> Files;
+		HashSet<DirectoryReference> Directories;
+
+		public DirectoryCache(DirectoryReference InDirectory)
 		{
-			if( BaseDirectory.Exists() )
+			Directory = InDirectory;
+			Exists = Directory.Exists();
+
+			if(!Exists)
 			{
-				Files = new HashSet<FileReference>(BaseDirectory.EnumerateFileReferences());
+				Files = new HashSet<FileReference>();
+				Directories = new HashSet<DirectoryReference>();
 			}
-			else
+		}
+
+		private void CacheFiles()
+		{
+			if(Files == null)
 			{
-				Files = new HashSet<FileReference>();
+				Files = new HashSet<FileReference>(Directory.EnumerateFileReferences());
 			}
 		}
 
-		public bool HasFile( FileReference File )
+		public bool HasFile(FileReference CheckFile)
 		{
-			return Files.Contains(File);
+			CacheFiles();
+			return Files.Contains(CheckFile);
 		}
 
-		HashSet<FileReference> Files;
-	}
+		public IEnumerable<FileReference> EnumerateFiles()
+		{
+			CacheFiles();
+			return Files;
+		}
+
+		private void CacheDirectories()
+		{
+			if(Directories == null)
+			{
+				Directories = new HashSet<DirectoryReference>(Directory.EnumerateDirectoryReferences());
+			}
+		}
 
+		public IEnumerable<DirectoryReference> EnumerateDirectories()
+		{
+			CacheDirectories();
+			return Directories;
+		}
+	}
 
 	public static class DirectoryLookupCache
 	{
-		static public bool FileExists( FileReference FileFullPath )
+		static DirectoryCache FindOrCreateDirectoryCache(DirectoryReference Directory)
 		{
-			DirectoryReference FileDirectory = FileFullPath.Directory;
-
 			DirectoryCache FoundDirectoryCache;
-			if( !Directories.TryGetValue( FileDirectory, out FoundDirectoryCache ) )
+			if (!Directories.TryGetValue(Directory, out FoundDirectoryCache))
 			{
-				// First time we've seen this directory.  Gather information about files.
-				FoundDirectoryCache = new DirectoryCache( FileDirectory );
-				Directories.Add( FileDirectory, FoundDirectoryCache );
+				FoundDirectoryCache = new DirectoryCache(Directory);
+				Directories.Add(Directory, FoundDirectoryCache);
 			}
+			return FoundDirectoryCache;
+		}
+
+
+		static public bool FileExists(FileReference File)
+		{
+			DirectoryCache FoundDirectoryCache = FindOrCreateDirectoryCache(File.Directory);
+			return FoundDirectoryCache.HasFile(File);
+		}
 
-			// Try to find the file in this directory
-			return FoundDirectoryCache.HasFile( FileFullPath );
+		static public bool DirectoryExists(DirectoryReference Directory)
+		{
+			DirectoryCache FoundDirectoryCache = FindOrCreateDirectoryCache(Directory);
+			return FoundDirectoryCache.Exists;
+		}
+
+		static public IEnumerable<FileReference> EnumerateFiles(DirectoryReference Directory)
+		{
+			DirectoryCache FoundDirectoryCache = FindOrCreateDirectoryCache(Directory);
+			return FoundDirectoryCache.EnumerateFiles();
+		}
+
+		static public IEnumerable<DirectoryReference> EnumerateDirectories(DirectoryReference Directory)
+		{
+			DirectoryCache FoundDirectoryCache = FindOrCreateDirectoryCache(Directory);
+			return FoundDirectoryCache.EnumerateDirectories();
+		}
+
+		static public IEnumerable<DirectoryReference> EnumerateDirectoriesRecursively(DirectoryReference Directory)
+		{
+			DirectoryCache FoundDirectoryCache = FindOrCreateDirectoryCache(Directory);
+			foreach(DirectoryReference SubDirectory in FoundDirectoryCache.EnumerateDirectories())
+			{
+				yield return SubDirectory;
+
+				foreach(DirectoryReference ChildSubDirectory in EnumerateDirectoriesRecursively(SubDirectory))
+				{
+					yield return ChildSubDirectory;
+				}
+			}
 		}
 
 		static Dictionary<DirectoryReference, DirectoryCache> Directories = new Dictionary<DirectoryReference, DirectoryCache>();
