Reduce time taken to generate receipts by using the target's view of the required modules, rather than recursively finding them from all the binaries. Also use the FileReferences for build products to optimize the way we escape build product paths.

diff --git a/Engine/Source/Programs/UnrealBuildTool/Android/AndroidToolChain.cs b/Engine/Source/Programs/UnrealBuildTool/Android/AndroidToolChain.cs
--- a/Engine/Source/Programs/UnrealBuildTool/Android/AndroidToolChain.cs
+++ b/Engine/Source/Programs/UnrealBuildTool/Android/AndroidToolChain.cs
@@ -1257,7 +1257,7 @@ namespace UnrealBuildTool
 			return Outputs.ToArray();
 		}
 
-		public override void AddFilesToReceipt(TargetReceipt Receipt, UEBuildBinary Binary)
+		public override void ModifyBuildProducts(UEBuildBinary Binary, Dictionary<FileReference, BuildProductType> BuildProducts)
 		{
 			// the binary will have all of the .so's in the output files, we need to trim down to the shared apk (which is what needs to go into the manifest)
 			if (Binary.Config.Type != UEBuildBinaryType.StaticLibrary)
@@ -1265,7 +1265,7 @@ namespace UnrealBuildTool
 				foreach (FileReference BinaryPath in Binary.Config.OutputFilePaths)
 				{
 					FileReference ApkFile = BinaryPath.ChangeExtension(".apk");
-					Receipt.AddBuildProduct(ApkFile, BuildProductType.Executable);
+					BuildProducts.Add(ApkFile, BuildProductType.Executable);
 				}
 			}
 		}
diff --git a/Engine/Source/Programs/UnrealBuildTool/Configuration/UEBuildBinary.cs b/Engine/Source/Programs/UnrealBuildTool/Configuration/UEBuildBinary.cs
--- a/Engine/Source/Programs/UnrealBuildTool/Configuration/UEBuildBinary.cs
+++ b/Engine/Source/Programs/UnrealBuildTool/Configuration/UEBuildBinary.cs
@@ -285,14 +285,12 @@ namespace UnrealBuildTool
 		}
 
 		/// <summary>
-		/// Creates a receipt for this binary.
+		/// Gets all build products produced by this binary
 		/// </summary>
-		/// <param name="ToolChain">Toolchain for the target platform</param>
-		/// <param name="BuildPlatform">Platform that we're building for</param>
-		public virtual TargetReceipt MakeReceipt(IUEToolChain ToolChain)
+		/// <param name="ToolChain">The platform toolchain</param>
+		/// <param name="BuildProducts">Mapping of produced build product to type</param>
+		public virtual void GetBuildProducts(IUEToolChain ToolChain, Dictionary<FileReference, BuildProductType> BuildProducts)
 		{
-			TargetReceipt Receipt = new TargetReceipt();
-
 			// Get the type of build products we're creating
 			BuildProductType Type = BuildProductType.RequiredResource;
 			switch(Config.Type)
@@ -312,7 +310,7 @@ namespace UnrealBuildTool
 			string DebugExtension = UEBuildPlatform.GetBuildPlatform(Target.Platform).GetDebugInfoExtension(Config.Type);
 			foreach (FileReference OutputFilePath in Config.OutputFilePaths)
 			{
-				AddBuildProductAndDebugFile(OutputFilePath, Type, DebugExtension, Receipt, ToolChain);
+				AddBuildProductAndDebugFile(OutputFilePath, Type, DebugExtension, BuildProducts, ToolChain);
 			}
 
 			// Add the console app, if there is one
@@ -320,13 +318,12 @@ namespace UnrealBuildTool
 			{
 				foreach (FileReference OutputFilePath in Config.OutputFilePaths)
 				{
-					AddBuildProductAndDebugFile(GetAdditionalConsoleAppPath(OutputFilePath), Type, DebugExtension, Receipt, ToolChain);
+					AddBuildProductAndDebugFile(GetAdditionalConsoleAppPath(OutputFilePath), Type, DebugExtension, BuildProducts, ToolChain);
 				}
 			}
 
 			// Add any extra files from the toolchain
-			ToolChain.AddFilesToReceipt(Receipt, this);
-			return Receipt;
+			ToolChain.ModifyBuildProducts(this, BuildProducts);
 		}
 
 		/// <summary>
@@ -335,13 +332,13 @@ namespace UnrealBuildTool
 		/// <param name="OutputFile">Build product to add</param>
 		/// <param name="DebugExtension">Extension for the matching debug file (may be null).</param>
 		/// <param name="Receipt">Receipt to add to</param>
-		static void AddBuildProductAndDebugFile(FileReference OutputFile, BuildProductType OutputType, string DebugExtension, TargetReceipt Receipt, IUEToolChain ToolChain)
+		static void AddBuildProductAndDebugFile(FileReference OutputFile, BuildProductType OutputType, string DebugExtension, Dictionary<FileReference, BuildProductType> BuildProducts, IUEToolChain ToolChain)
 		{
-			Receipt.AddBuildProduct(OutputFile, OutputType);
+			BuildProducts.Add(OutputFile, OutputType);
 
 			if(!String.IsNullOrEmpty(DebugExtension) && ToolChain.ShouldAddDebugFileToReceipt(OutputFile, OutputType))
 			{
-				Receipt.AddBuildProduct(OutputFile.ChangeExtension(DebugExtension), BuildProductType.SymbolFile);
+				BuildProducts.Add(OutputFile.ChangeExtension(DebugExtension), BuildProductType.SymbolFile);
 			}
 		}
 
@@ -637,50 +634,20 @@ namespace UnrealBuildTool
 		}
 
 		/// <summary>
-		/// Overrides base class to add module runtime dependencies to the build receipt.
+		/// Gets all build products produced by this binary
 		/// </summary>
 		/// <param name="ToolChain">The platform toolchain</param>
-		public override TargetReceipt MakeReceipt(IUEToolChain ToolChain)
+		/// <param name="BuildProducts">Mapping of produced build product to type</param>
+		public override void GetBuildProducts(IUEToolChain ToolChain, Dictionary<FileReference, BuildProductType> BuildProducts)
 		{
-			TargetReceipt Receipt = base.MakeReceipt(ToolChain);
-
-			// Set the IsPrecompiled flag on all the build products if we're not actually building this binary
-			if(!Config.bAllowCompilation)
-			{
-				foreach(BuildProduct BuildProduct in Receipt.BuildProducts)
-				{
-					BuildProduct.IsPrecompiled = true;
-				}
-			}
+			base.GetBuildProducts(ToolChain, BuildProducts);
 
 			// Add the compiled resource file if we're building a static library containing the launch module on Windows
 			if(Config.Type == UEBuildBinaryType.StaticLibrary && ModuleNames.Contains("Launch") && (Target.Platform == UnrealTargetPlatform.Win32 || Target.Platform == UnrealTargetPlatform.Win64))
 			{
 				FileReference ResourceFilePath = FileReference.Combine(Config.IntermediateDirectory, "Launch", "PCLaunch.rc.res");
-				Receipt.AddBuildProduct(ResourceFilePath.FullName, BuildProductType.StaticLibrary);
-			}
-			
-			// Add runtime dependencies for all the modules in this binary, and build up a list of all the referenced modules
-			var ReferencedModules = new CaselessDictionary<UEBuildModule.ModuleIndexPair>();
-			foreach (string ModuleName in ModuleNames)
-			{
-				UEBuildModule Module = Target.GetModuleByName(ModuleName);
-				foreach(RuntimeDependency RuntimeDependency in Module.RuntimeDependencies)
-				{
-					Receipt.RuntimeDependencies.Add(new RuntimeDependency(RuntimeDependency));
-				}
-				Module.GetAllDependencyModules(ReferencedModules, true, false, true);
-			}
-
-			// Add runtime dependencies for all the referenced external modules. These may be introduce dependencies for the binary without actually being listed for inclusion in it.
-			foreach(UEBuildExternalModule ExternalModule in ReferencedModules.Values.OrderBy(x => x.Index).Select(x => x.Module).OfType<UEBuildExternalModule>())
-			{
-				foreach(RuntimeDependency RuntimeDependency in ExternalModule.RuntimeDependencies)
-				{
-					Receipt.RuntimeDependencies.Add(new RuntimeDependency(RuntimeDependency));
-				}
+				BuildProducts.Add(ResourceFilePath, BuildProductType.StaticLibrary);
 			}
-			return Receipt;
 		}
 
 		// Object interface.
diff --git a/Engine/Source/Programs/UnrealBuildTool/Configuration/UEBuildTarget.cs b/Engine/Source/Programs/UnrealBuildTool/Configuration/UEBuildTarget.cs
--- a/Engine/Source/Programs/UnrealBuildTool/Configuration/UEBuildTarget.cs
+++ b/Engine/Source/Programs/UnrealBuildTool/Configuration/UEBuildTarget.cs
@@ -1715,18 +1715,36 @@ namespace UnrealBuildTool
 			// that already contains a manifest, we'll reuse the build id that's already in there (see below).
 			string BuildId = (OnlyModules.Count == 0)? Guid.NewGuid().ToString() : "";
 
-			// Create the receipt, and merge all the receipts from the individual binaries
+			// Find all the build products and modules from this binary
 			Receipt = new TargetReceipt(TargetName, Platform, Configuration, BuildId, Version);
 			foreach(UEBuildBinary Binary in AppBinaries)
 			{
-				TargetReceipt BinaryReceipt = Binary.MakeReceipt(ToolChain);
-				if(Binary.Config.Type == UEBuildBinaryType.StaticLibrary)
+				// Get all the build products for this binary
+				Dictionary<FileReference, BuildProductType> BuildProducts = new Dictionary<FileReference,BuildProductType>();
+				Binary.GetBuildProducts(ToolChain, BuildProducts);
+
+				// Add them to the receipt
+				foreach(KeyValuePair<FileReference, BuildProductType> BuildProductPair in BuildProducts)
 				{
-					BinaryReceipt.RuntimeDependencies.Clear();
+					string NormalizedPath = TargetReceipt.InsertPathVariables(BuildProductPair.Key, UnrealBuildTool.EngineDirectory, ProjectDirectory);
+					BuildProduct BuildProduct = Receipt.AddBuildProduct(NormalizedPath, BuildProductPair.Value);
+					BuildProduct.IsPrecompiled = !Binary.Config.bAllowCompilation;
+				}
+			}
+
+			// Find all the modules which are part of this target
+			foreach(UEBuildModule Module in Modules.Values)
+			{
+				if(Module.bIncludedInTarget && Module.Binary != null && Module.Binary.Config.Type != UEBuildBinaryType.StaticLibrary)
+				{
+					foreach(RuntimeDependency RuntimeDependency in Module.RuntimeDependencies)
+					{
+						string SourcePath = TargetReceipt.InsertPathVariables(RuntimeDependency.Path, UnrealBuildTool.EngineDirectory, ProjectDirectory);
+						string TargetPath = TargetReceipt.InsertPathVariables(RuntimeDependency.StagePath, UnrealBuildTool.EngineDirectory, ProjectDirectory);
+						Receipt.AddRuntimeDependency(SourcePath, TargetPath);
+					}
 				}
-				Receipt.Merge(BinaryReceipt);
 			}
-			Receipt.InsertStandardPathVariables(UnrealBuildTool.EngineDirectory, ProjectDirectory);
 
 			// Prepare all the version manifests
 			Dictionary<FileReference, VersionManifest> FileNameToVersionManifest = new Dictionary<FileReference, VersionManifest>();
diff --git a/Engine/Source/Programs/UnrealBuildTool/HTML5/HTML5ToolChain.cs b/Engine/Source/Programs/UnrealBuildTool/HTML5/HTML5ToolChain.cs
--- a/Engine/Source/Programs/UnrealBuildTool/HTML5/HTML5ToolChain.cs
+++ b/Engine/Source/Programs/UnrealBuildTool/HTML5/HTML5ToolChain.cs
@@ -536,13 +536,13 @@ namespace UnrealBuildTool
 			throw new BuildException("HTML5 cannot compile C# files");
 		}
 
-        public override void AddFilesToReceipt(TargetReceipt Receipt, UEBuildBinary Binary)
+		public override void ModifyBuildProducts(UEBuildBinary Binary, Dictionary<FileReference, BuildProductType> BuildProducts)
         {
             // we need to include the generated .mem and .symbols file.  
 			if(Binary.Config.Type != UEBuildBinaryType.StaticLibrary)
 			{
-	            Receipt.AddBuildProduct(Binary.Config.OutputFilePath + ".mem", BuildProductType.RequiredResource);
-				Receipt.AddBuildProduct(Binary.Config.OutputFilePath + ".symbols", BuildProductType.RequiredResource);
+	            BuildProducts.Add(Binary.Config.OutputFilePath + ".mem", BuildProductType.RequiredResource);
+				BuildProducts.Add(Binary.Config.OutputFilePath + ".symbols", BuildProductType.RequiredResource);
 			}
         }
 
diff --git a/Engine/Source/Programs/UnrealBuildTool/IOS/IOSToolChain.cs b/Engine/Source/Programs/UnrealBuildTool/IOS/IOSToolChain.cs
--- a/Engine/Source/Programs/UnrealBuildTool/IOS/IOSToolChain.cs
+++ b/Engine/Source/Programs/UnrealBuildTool/IOS/IOSToolChain.cs
@@ -94,12 +94,12 @@ namespace UnrealBuildTool
 			SetupXcodePaths(true);
 		}
 
-		public override void AddFilesToReceipt(TargetReceipt Receipt, UEBuildBinary Binary)
+		public override void ModifyBuildProducts(UEBuildBinary Binary, Dictionary<FileReference, BuildProductType> BuildProducts)
 		{
 			if (BuildConfiguration.bCreateStubIPA && Binary.Config.Type != UEBuildBinaryType.StaticLibrary)
 			{
 				FileReference StubFile = FileReference.Combine(Binary.Config.OutputFilePath.Directory, Binary.Config.OutputFilePath.GetFileNameWithoutExtension() + ".stub");
-				Receipt.AddBuildProduct(StubFile, BuildProductType.Executable);
+				BuildProducts.Add(StubFile, BuildProductType.Executable);
 			}
 		}
 
diff --git a/Engine/Source/Programs/UnrealBuildTool/Mac/MacToolChain.cs b/Engine/Source/Programs/UnrealBuildTool/Mac/MacToolChain.cs
--- a/Engine/Source/Programs/UnrealBuildTool/Mac/MacToolChain.cs
+++ b/Engine/Source/Programs/UnrealBuildTool/Mac/MacToolChain.cs
@@ -7,6 +7,7 @@ using System.IO;
 using System.Diagnostics;
 using System.Security.AccessControl;
 using System.Text;
+using System.Linq;
 
 namespace UnrealBuildTool
 {
@@ -1295,26 +1296,28 @@ namespace UnrealBuildTool
 
 		static private DirectoryReference BundleContentsDirectory;
 
-        public override void AddFilesToReceipt(TargetReceipt Receipt, UEBuildBinary Binary)
+        public override void ModifyBuildProducts(UEBuildBinary Binary, Dictionary<FileReference, BuildProductType> BuildProducts)
 		{
 			// The cross-platform code adds .dSYMs for static libraries, which is just wrong, so
 			// eliminate them here for now.
 			string DebugExtension = UEBuildPlatform.GetBuildPlatform(Binary.Target.Platform).GetDebugInfoExtension(Binary.Config.Type);
 			if(DebugExtension == ".dSYM")
 			{
-				for (int i = 0; i < Receipt.BuildProducts.Count; i++)
+				KeyValuePair<FileReference, BuildProductType>[] BuildProductsArray = BuildProducts.ToArray();
+
+				foreach(KeyValuePair<FileReference, BuildProductType> BuildProductPair in BuildProductsArray)
 				{
-					if(Path.GetExtension(Receipt.BuildProducts[i].Path) == DebugExtension)
+					if(BuildProductPair.Key.HasExtension(DebugExtension))
 					{
-						Receipt.BuildProducts.RemoveAt(i--);
+						BuildProducts.Remove(BuildProductPair.Key);
 					}
 				}
 
-				for (int i = 0; i < Receipt.BuildProducts.Count; i++)
+				foreach(KeyValuePair<FileReference, BuildProductType> BuildProductPair in BuildProductsArray)
 				{
-					if(Receipt.BuildProducts[i].Type == BuildProductType.Executable || Receipt.BuildProducts[i].Type == BuildProductType.DynamicLibrary)
+					if(BuildProductPair.Value == BuildProductType.Executable || BuildProductPair.Value == BuildProductType.DynamicLibrary)
 					{
-						Receipt.AddBuildProduct(new FileReference(Path.ChangeExtension(Receipt.BuildProducts[i].Path, DebugExtension)), BuildProductType.SymbolFile);
+						BuildProducts.Add(BuildProductPair.Key.ChangeExtension(DebugExtension), BuildProductType.SymbolFile);
 					}
 				}
 			}
@@ -1349,7 +1352,7 @@ namespace UnrealBuildTool
 						if (Path.GetExtension(AdditionalLibrary) == ".dylib" && BundleContentsDirectory != null)
 						{
 							FileReference Entry = FileReference.Combine(BundleContentsDirectory, "MacOS", LibName);
-							Receipt.AddBuildProduct(Entry, BuildProductType.DynamicLibrary);
+							BuildProducts.Add(Entry, BuildProductType.DynamicLibrary);
 						}
 					}
 				}
@@ -1360,12 +1363,12 @@ namespace UnrealBuildTool
 					{
 						foreach (string ResourceFile in Directory.GetFiles(Resource.ResourcePath, "*", SearchOption.AllDirectories))
 						{
-							Receipt.AddBuildProduct(FileReference.Combine(BundleContentsDirectory, Resource.BundleContentsSubdir, ResourceFile.Substring(Path.GetDirectoryName(Resource.ResourcePath).Length + 1)), BuildProductType.RequiredResource);
+							BuildProducts.Add(FileReference.Combine(BundleContentsDirectory, Resource.BundleContentsSubdir, ResourceFile.Substring(Path.GetDirectoryName(Resource.ResourcePath).Length + 1)), BuildProductType.RequiredResource);
 						}
 					}
 					else
 					{
-						Receipt.AddBuildProduct(FileReference.Combine(BundleContentsDirectory, Resource.BundleContentsSubdir, Path.GetFileName(Resource.ResourcePath)), BuildProductType.RequiredResource);
+						BuildProducts.Add(FileReference.Combine(BundleContentsDirectory, Resource.BundleContentsSubdir, Path.GetFileName(Resource.ResourcePath)), BuildProductType.RequiredResource);
 					}
 				}
 			}
@@ -1373,17 +1376,17 @@ namespace UnrealBuildTool
 			if (Binary.Config.Type == UEBuildBinaryType.Executable)
 			{
 				// And we also need all the resources
-				Receipt.AddBuildProduct(FileReference.Combine(BundleContentsDirectory, "Info.plist"), BuildProductType.RequiredResource);
-				Receipt.AddBuildProduct(FileReference.Combine(BundleContentsDirectory, "PkgInfo"), BuildProductType.RequiredResource);
+				BuildProducts.Add(FileReference.Combine(BundleContentsDirectory, "Info.plist"), BuildProductType.RequiredResource);
+				BuildProducts.Add(FileReference.Combine(BundleContentsDirectory, "PkgInfo"), BuildProductType.RequiredResource);
 
 				if (Binary.Target.TargetType == TargetRules.TargetType.Editor)
 				{
-					Receipt.AddBuildProduct(FileReference.Combine(BundleContentsDirectory, "Resources/UE4Editor.icns"), BuildProductType.RequiredResource);
-					Receipt.AddBuildProduct(FileReference.Combine(BundleContentsDirectory, "Resources/UProject.icns"), BuildProductType.RequiredResource);
+					BuildProducts.Add(FileReference.Combine(BundleContentsDirectory, "Resources/UE4Editor.icns"), BuildProductType.RequiredResource);
+					BuildProducts.Add(FileReference.Combine(BundleContentsDirectory, "Resources/UProject.icns"), BuildProductType.RequiredResource);
 				}
 				else
 				{
-					Receipt.AddBuildProduct(FileReference.Combine(BundleContentsDirectory, "Resources/" + Binary.Target.TargetName + ".icns"), BuildProductType.RequiredResource);
+					BuildProducts.Add(FileReference.Combine(BundleContentsDirectory, "Resources/" + Binary.Target.TargetName + ".icns"), BuildProductType.RequiredResource);
 				}
 			}
 		}
diff --git a/Engine/Source/Programs/UnrealBuildTool/System/TargetReceipt.cs b/Engine/Source/Programs/UnrealBuildTool/System/TargetReceipt.cs
--- a/Engine/Source/Programs/UnrealBuildTool/System/TargetReceipt.cs
+++ b/Engine/Source/Programs/UnrealBuildTool/System/TargetReceipt.cs
@@ -196,17 +196,6 @@ namespace UnrealBuildTool
 			return NewBuildProduct;
 		}
 
-		/// <summary>
-		/// Adds a build product to the receipt. Does not check whether it already exists.
-		/// </summary>
-		/// <param name="Path">Path to the build product.</param>
-		/// <param name="Type">Type of build product.</param>
-		/// <returns>The BuildProduct object that was created</returns>
-		public BuildProduct AddBuildProduct(FileReference Path, BuildProductType Type)
-		{
-			return AddBuildProduct(Path.FullName, Type);
-		}
-
 		/// <summary>
 		/// Constructs a runtime dependency object and adds it to the receipt.
 		/// </summary>
@@ -277,52 +266,43 @@ namespace UnrealBuildTool
 		}
 
 		/// <summary>
-		/// Inserts standard $(EngineDir) and $(ProjectDir) variables into any path strings, so it can be used on different machines.
+		/// Inserts $(EngineDir) and $(ProjectDir) variables into a path string, so it can be used on different machines.
 		/// </summary>
+		/// <param name="InputPath">Input path</param>
 		/// <param name="EngineDir">The engine directory. Relative paths are ok.</param>
 		/// <param name="ProjectDir">The project directory. Relative paths are ok.</param>
-		public void InsertStandardPathVariables(DirectoryReference EngineDir, DirectoryReference ProjectDir)
+		/// <returns>New string with the base directory replaced, or the original string</returns>
+		public static string InsertPathVariables(string InputPath, DirectoryReference EngineDir, DirectoryReference ProjectDir)
 		{
-			string EnginePrefix = EngineDir.FullName.TrimEnd(Path.DirectorySeparatorChar) + Path.DirectorySeparatorChar;
-			string ProjectPrefix = ProjectDir.FullName.TrimEnd(Path.DirectorySeparatorChar) + Path.DirectorySeparatorChar;
-
-			foreach(BuildProduct BuildProduct in BuildProducts)
-			{
-				BuildProduct.Path = InsertStandardPathVariablesToString(BuildProduct.Path, EnginePrefix, ProjectPrefix);
-			}
-			foreach(RuntimeDependency RuntimeDependency in RuntimeDependencies)
+			string Result = InputPath;
+			if(InputPath != null && !InputPath.StartsWith("$("))
 			{
-				RuntimeDependency.Path = InsertStandardPathVariablesToString(RuntimeDependency.Path, EnginePrefix, ProjectPrefix);
-				if(RuntimeDependency.StagePath != null)
-				{
-					RuntimeDependency.StagePath = InsertStandardPathVariablesToString(RuntimeDependency.StagePath, EnginePrefix, ProjectPrefix);
-				}
+				Result = InsertPathVariables(new FileReference(InputPath), EngineDir, ProjectDir);
 			}
+			return Result;
 		}
 
 		/// <summary>
-		/// Inserts $(EngineDir) and $(ProjectDir) variables into a path string, so it can be used on different machines.
+		/// Inserts variables to make a file relative to $(EngineDir) or $(ProjectDir)
 		/// </summary>
-		/// <param name="InputPath">Input path</param>
-		/// <param name="EngineDir">The engine directory. Relative paths are ok.</param>
-		/// <param name="ProjectDir">The project directory. Relative paths are ok.</param>
-		/// <returns>New string with the base directory replaced, or the original string</returns>
-		static public string InsertStandardPathVariablesToString(string InputPath, string EnginePrefix, string ProjectPrefix)
+		/// <param name="File">The file to insert variables into.</param>
+		/// <param name="EngineDir">Value of the $(EngineDir) variable.</param>
+		/// <param name="ProjectDir">Value of the $(ProjectDir) variable.</param>
+		/// <returns>Converted path for the file.</returns>
+		public static string InsertPathVariables(FileReference File, DirectoryReference EngineDir, DirectoryReference ProjectDir)
 		{
-			string Result = InputPath;
-			if(!InputPath.StartsWith("$("))
+			if(File.IsUnderDirectory(EngineDir))
 			{
-				string FullInputPath = Path.GetFullPath(InputPath);
-				if(FullInputPath.StartsWith(EnginePrefix))
-				{
-					Result = "$(EngineDir)" + FullInputPath.Substring(EnginePrefix.Length - 1);
-				}
-				else if(FullInputPath.StartsWith(ProjectPrefix))
-				{
-					Result = "$(ProjectDir)" + FullInputPath.Substring(ProjectPrefix.Length - 1);
-				}
+				return "$(EngineDir)" + Path.DirectorySeparatorChar + File.MakeRelativeTo(EngineDir);
+			}
+			else if(File.IsUnderDirectory(ProjectDir))
+			{
+				return "$(ProjectDir)" + Path.DirectorySeparatorChar + File.MakeRelativeTo(ProjectDir);
+			}
+			else
+			{
+				return File.FullName;
 			}
-			return Result;
 		}
 
 		/// <summary>
diff --git a/Engine/Source/Programs/UnrealBuildTool/ToolChain/UEToolChain.cs b/Engine/Source/Programs/UnrealBuildTool/ToolChain/UEToolChain.cs
--- a/Engine/Source/Programs/UnrealBuildTool/ToolChain/UEToolChain.cs
+++ b/Engine/Source/Programs/UnrealBuildTool/ToolChain/UEToolChain.cs
@@ -38,7 +38,7 @@ namespace UnrealBuildTool
 
 		void SetUpGlobalEnvironment();
 
-		void AddFilesToReceipt(TargetReceipt Receipt, UEBuildBinary Binary);
+		void ModifyBuildProducts(UEBuildBinary Binary, Dictionary<FileReference, BuildProductType> BuildProducts);
 
 		bool ShouldAddDebugFileToReceipt(FileReference OutputFile, BuildProductType OutputType);
 
@@ -171,7 +171,7 @@ namespace UnrealBuildTool
 		{
 		}
 
-		public virtual void AddFilesToReceipt(TargetReceipt Receipt, UEBuildBinary Binary)
+		public virtual void ModifyBuildProducts(UEBuildBinary Binary, Dictionary<FileReference, BuildProductType> BuildProducts)
 		{
 		}
 
diff --git a/Engine/Source/Programs/UnrealBuildTool/Windows/VCToolChain.cs b/Engine/Source/Programs/UnrealBuildTool/Windows/VCToolChain.cs
--- a/Engine/Source/Programs/UnrealBuildTool/Windows/VCToolChain.cs
+++ b/Engine/Source/Programs/UnrealBuildTool/Windows/VCToolChain.cs
@@ -1777,11 +1777,11 @@ namespace UnrealBuildTool
 			return IncludePaths;
 		}
 
-		public override void AddFilesToReceipt(TargetReceipt Receipt, UEBuildBinary Binary)
+		public override void ModifyBuildProducts(UEBuildBinary Binary, Dictionary<FileReference, BuildProductType> BuildProducts)
 		{
 			if (Binary.Config.Type == UEBuildBinaryType.DynamicLinkLibrary)
 			{
-				Receipt.AddBuildProduct(FileReference.Combine(Binary.Config.IntermediateDirectory, Binary.Config.OutputFilePath.GetFileNameWithoutExtension() + ".lib"), BuildProductType.ImportLibrary);
+				BuildProducts.Add(FileReference.Combine(Binary.Config.IntermediateDirectory, Binary.Config.OutputFilePath.GetFileNameWithoutExtension() + ".lib"), BuildProductType.ImportLibrary);
 			}
 		}
 
