Move retrieval of VULCAN_SDK_PATH to a static variable; it gets called a lot when generating project files and shows up on profiles.

diff --git a/Engine/Source/Programs/UnrealBuildTool/Windows/UEBuildWindows.cs b/Engine/Source/Programs/UnrealBuildTool/Windows/UEBuildWindows.cs
--- a/Engine/Source/Programs/UnrealBuildTool/Windows/UEBuildWindows.cs
+++ b/Engine/Source/Programs/UnrealBuildTool/Windows/UEBuildWindows.cs
@@ -480,6 +480,9 @@ namespace UnrealBuildTool
             UEBuildConfiguration.bCompileICU = true;
         }
 
+//#todo-rco: Remove when public
+		public static string VulkanSDKPath = Environment.GetEnvironmentVariable("VK_SDK_PATH");
+
 		public override void ModifyModuleRules(string ModuleName, ModuleRules Rules, TargetInfo Target)
         {
             if ((Target.Platform == UnrealTargetPlatform.Win32) || (Target.Platform == UnrealTargetPlatform.Win64))
@@ -511,7 +514,6 @@ namespace UnrealBuildTool
 
 //#todo-rco: Remove when public
 					{
-						string VulkanSDKPath = Environment.GetEnvironmentVariable("VK_SDK_PATH");
 						if (!String.IsNullOrEmpty(VulkanSDKPath))
 						{
 							Rules.DynamicallyLoadedModuleNames.Add("VulkanShaderFormat");
