Store module references rather than module names in UEBuildBinary. Removes several code paths that look them up more than once, or create redundant instances of module rules classes.

diff --git a/Engine/Source/Programs/UnrealBuildTool/Configuration/UEBuildBinary.cs b/Engine/Source/Programs/UnrealBuildTool/Configuration/UEBuildBinary.cs
--- a/Engine/Source/Programs/UnrealBuildTool/Configuration/UEBuildBinary.cs
+++ b/Engine/Source/Programs/UnrealBuildTool/Configuration/UEBuildBinary.cs
@@ -280,7 +280,7 @@ namespace UnrealBuildTool
 		/// Adds a module to the binary.
 		/// </summary>
 		/// <param name="ModuleName">The module to add</param>
-		public virtual void AddModule( string ModuleName )
+		public virtual void AddModule( UEBuildModule Module )
 		{
 		}
 
@@ -409,7 +409,7 @@ namespace UnrealBuildTool
 	/// </summary>
 	public class UEBuildBinaryCPP : UEBuildBinary
 	{
-		public readonly List<string> ModuleNames;
+		public readonly List<UEBuildModule> Modules = new List<UEBuildModule>();
 		private bool bCreateImportLibrarySeparately;
 		private bool bIncludeDependentLibrariesInLibrary;
 
@@ -420,7 +420,6 @@ namespace UnrealBuildTool
 		public UEBuildBinaryCPP( UEBuildTarget InTarget, UEBuildBinaryConfiguration InConfig )
 			: base( InTarget, InConfig )
 		{
-			ModuleNames = new List<string>(InConfig.ModuleNames);
 			bCreateImportLibrarySeparately = InConfig.bCreateImportLibrarySeparately;
 			bIncludeDependentLibrariesInLibrary = InConfig.bIncludeDependentLibrariesInLibrary;
 		}
@@ -429,11 +428,11 @@ namespace UnrealBuildTool
 		/// Adds a module to the binary.
 		/// </summary>
 		/// <param name="ModuleName">The module to add</param>
-		public override void AddModule(string ModuleName)
+		public override void AddModule(UEBuildModule Module)
 		{
-			if( !ModuleNames.Contains( ModuleName ) )
+			if(!Modules.Contains(Module))
 			{
-				ModuleNames.Add( ModuleName );
+				Modules.Add(Module);
 			}
 		}
 
@@ -446,12 +445,10 @@ namespace UnrealBuildTool
 		/// <param name="Target">The target info</param>
 		public override void BindModules()
 		{
-			foreach(var ModuleName in ModuleNames)
+			foreach(var Module in Modules)
 			{
-				UEBuildModule Module = null;
 				if (Config.bHasModuleRules)
 				{
-					Module = Target.FindOrCreateModuleByName(ModuleName);
 					if(Module.Binary == null)
 					{
 						Module.Binary = this;
@@ -459,7 +456,7 @@ namespace UnrealBuildTool
 					}
 					else if(Module.Binary.Config.Type != UEBuildBinaryType.StaticLibrary)
 					{
-						throw new BuildException("Module \"{0}\" linked into both {1} and {2}, which creates ambiguous linkage for dependents.", ModuleName, Module.Binary.Config.OutputFilePath, Config.OutputFilePath);
+						throw new BuildException("Module \"{0}\" linked into both {1} and {2}, which creates ambiguous linkage for dependents.", Module.Name, Module.Binary.Config.OutputFilePath, Config.OutputFilePath);
 					}
 				}
 			}
@@ -474,16 +471,15 @@ namespace UnrealBuildTool
 		public override List<UEBuildModule> GetAllDependencyModules(bool bIncludeDynamicallyLoaded, bool bForceCircular)
 		{
 			var ReferencedModules = new CaselessDictionary<UEBuildModule.ModuleIndexPair>();
-			foreach( var ModuleName in ModuleNames )
+			foreach( var Module in Modules )
 			{
-				if( !ReferencedModules.ContainsKey( ModuleName ) )
+				if( !ReferencedModules.ContainsKey( Module.Name ) )
 				{
-					var Module = Target.GetModuleByName( ModuleName );
-					ReferencedModules[ ModuleName ] = null;
+					ReferencedModules[ Module.Name ] = null;
 
 					Module.GetAllDependencyModules(ReferencedModules, bIncludeDynamicallyLoaded, bForceCircular, bOnlyDirectDependencies: false);
 
-					ReferencedModules[ ModuleName ] = new UEBuildModule.ModuleIndexPair{ Module = Module, Index = ReferencedModules.Count };
+					ReferencedModules[ Module.Name ] = new UEBuildModule.ModuleIndexPair{ Module = Module, Index = ReferencedModules.Count };
 				}
 			}
 
@@ -500,9 +496,8 @@ namespace UnrealBuildTool
 			if (Config.bHasModuleRules)
 			{
 				// Modules may be added to this binary during this process, so don't foreach over ModuleNames
-				for(int Idx = 0; Idx < ModuleNames.Count; Idx++)
+				foreach(UEBuildModule Module in Modules)
 				{
-					UEBuildModule Module = Target.FindOrCreateModuleByName(ModuleNames[Idx]);
 					Module.RecursivelyCreateModules();
 					Module.RecursivelyProcessUnboundModules();
 				}
@@ -556,7 +551,7 @@ namespace UnrealBuildTool
 		public override IEnumerable<FileItem> Build(IUEToolChain TargetToolChain, CPPEnvironment CompileEnvironment, LinkEnvironment LinkEnvironment)
 		{
 			// UnrealCodeAnalyzer produces output files only for a specific module.
-			if (BuildConfiguration.bRunUnrealCodeAnalyzer && !(ModuleNames.Contains(BuildConfiguration.UCAModuleToAnalyze)))
+			if (BuildConfiguration.bRunUnrealCodeAnalyzer && !(Modules.Any(x => x.Name == BuildConfiguration.UCAModuleToAnalyze)))
 			{
 				return new List<FileItem>();
 			}
@@ -592,7 +587,7 @@ namespace UnrealBuildTool
 			}
 			
 			// If we're linking against static library containing the launch module on windows, we need to add the compiled resource separately. We can't link it through the static library.
-			if(Config.Type == UEBuildBinaryType.StaticLibrary && ModuleNames.Contains("Launch") && (Target.Platform == UnrealTargetPlatform.Win32 || Target.Platform == UnrealTargetPlatform.Win64))
+			if(Config.Type == UEBuildBinaryType.StaticLibrary && Modules.Any(x => x.Name == "Launch") && (Target.Platform == UnrealTargetPlatform.Win32 || Target.Platform == UnrealTargetPlatform.Win64))
 			{
 				FileReference ResourceFileRef = FileReference.Combine(Config.IntermediateDirectory, "Launch", "PCLaunch.rc.res");
 				DependentLinkEnvironment.InputFiles.Add(FileItem.GetItemByFileReference(ResourceFileRef));
@@ -606,11 +601,11 @@ namespace UnrealBuildTool
 		/// <returns>The OnlyModule if found, null if not</returns>
 		public override OnlyModule FindOnlyModule(List<OnlyModule> OnlyModules)
 		{
-			foreach (var ModuleName in ModuleNames)
+			foreach (var Module in Modules)
 			{
 				foreach (var OnlyModule in OnlyModules)
 				{
-					if (OnlyModule.OnlyModuleName.ToLower() == ModuleName.ToLower())
+					if (OnlyModule.OnlyModuleName.ToLower() == Module.Name.ToLower())
 					{
 						return OnlyModule;
 					}
@@ -622,9 +617,8 @@ namespace UnrealBuildTool
 		public override List<UEBuildModule> FindGameModules()
 		{
 			var GameModules = new List<UEBuildModule>();
-			foreach (var ModuleName in ModuleNames)
+			foreach (var Module in Modules)
 			{
-				UEBuildModule Module = Target.GetModuleByName(ModuleName);
 				if (!Module.ModuleDirectory.IsUnderDirectory(UnrealBuildTool.EngineDirectory))
 				{
 					GameModules.Add(Module);
@@ -643,7 +637,7 @@ namespace UnrealBuildTool
 			base.GetBuildProducts(ToolChain, BuildProducts);
 
 			// Add the compiled resource file if we're building a static library containing the launch module on Windows
-			if(Config.Type == UEBuildBinaryType.StaticLibrary && ModuleNames.Contains("Launch") && (Target.Platform == UnrealTargetPlatform.Win32 || Target.Platform == UnrealTargetPlatform.Win64))
+			if(Config.Type == UEBuildBinaryType.StaticLibrary && Modules.Any(x => x.Name == "Launch") && (Target.Platform == UnrealTargetPlatform.Win32 || Target.Platform == UnrealTargetPlatform.Win64))
 			{
 				FileReference ResourceFilePath = FileReference.Combine(Config.IntermediateDirectory, "Launch", "PCLaunch.rc.res");
 				BuildProducts.Add(ResourceFilePath, BuildProductType.StaticLibrary);
@@ -682,15 +676,13 @@ namespace UnrealBuildTool
 				Config.OutputFilePaths[0].GetFileName();
 			BinaryCompileEnvironment.Config.Definitions.Add("ORIGINAL_FILE_NAME=\"" + OriginalFilename + "\"");
 
-			foreach (var ModuleName in ModuleNames)
+			foreach (var Module in Modules)
 			{
-				var Module = Target.GetModuleByName(ModuleName);
-
 				List<FileItem> LinkInputFiles; 
 				if(Module.Binary == null || Module.Binary == this)
 				{
 					// Compile each module.
-					Log.TraceVerbose("Compile module: " + ModuleName);
+					Log.TraceVerbose("Compile module: " + Module.Name);
 					LinkInputFiles = Module.Compile(CompileEnvironment, BinaryCompileEnvironment);
 
 					// NOTE: Because of 'Shared PCHs', in monolithic builds the same PCH file may appear as a link input
@@ -819,7 +811,7 @@ namespace UnrealBuildTool
 		private List<FileItem> CreateOutputFilesForUCA(LinkEnvironment BinaryLinkEnvironment)
 		{
 			var OutputFiles = new List<FileItem>();
-			var ModuleName = ModuleNames.First(Name => Name.CompareTo(BuildConfiguration.UCAModuleToAnalyze) == 0);
+			var ModuleName = Modules.Select(Module => Module.Name).First(Name => Name.CompareTo(BuildConfiguration.UCAModuleToAnalyze) == 0);
 			var ModuleCPP = (UEBuildModuleCPP)Target.GetModuleByName(ModuleName);
 			var ModulePrivatePCH = ModuleCPP.ProcessedDependencies.UniquePCHHeaderFile;
 			var IntermediatePath = Path.Combine(Target.ProjectIntermediateDirectory.FullName, ModuleName);
diff --git a/Engine/Source/Programs/UnrealBuildTool/Configuration/UEBuildModule.cs b/Engine/Source/Programs/UnrealBuildTool/Configuration/UEBuildModule.cs
--- a/Engine/Source/Programs/UnrealBuildTool/Configuration/UEBuildModule.cs
+++ b/Engine/Source/Programs/UnrealBuildTool/Configuration/UEBuildModule.cs
@@ -1717,7 +1717,7 @@ namespace UnrealBuildTool
 						bool bIsCrossTarget = Rules.PlatformSpecificDynamicallyLoadedModuleNames.Contains(DependencyModule.Name) && !Rules.DynamicallyLoadedModuleNames.Contains(DependencyModule.Name);
 
 						// Get the binary that this module should be bound to
-						UEBuildBinary BinaryToBindTo = Target.FindOrAddBinaryForModule(DependencyModule.Name, bIsCrossTarget);
+						UEBuildBinary BinaryToBindTo = Target.FindOrAddBinaryForModule(DependencyModule, bIsCrossTarget);
 
 						// Bind this module
 						DependencyModule.Binary = BinaryToBindTo;
diff --git a/Engine/Source/Programs/UnrealBuildTool/Configuration/UEBuildTarget.cs b/Engine/Source/Programs/UnrealBuildTool/Configuration/UEBuildTarget.cs
--- a/Engine/Source/Programs/UnrealBuildTool/Configuration/UEBuildTarget.cs
+++ b/Engine/Source/Programs/UnrealBuildTool/Configuration/UEBuildTarget.cs
@@ -1280,7 +1280,7 @@ namespace UnrealBuildTool
 					{
 						// Collect all modules used by this binary.
 						Log.TraceVerbose("\tProcessing AppBinary " + AppBin.Config.OutputFilePaths[0]);
-						foreach (string ModuleName in AppBinCPP.ModuleNames)
+						foreach (string ModuleName in AppBinCPP.Modules.Select(x => x.Name))
 						{					
 							if (ModuleList.Contains(ModuleName) == false)
 							{
@@ -1864,7 +1864,7 @@ namespace UnrealBuildTool
 				{
 					if (CurModule.Binary == null || CurModule.Binary == ExecutableBinary || CurModule.Binary.Config.Type == UEBuildBinaryType.StaticLibrary)
 					{
-						ExecutableBinary.AddModule(CurModule.Name);
+						ExecutableBinary.AddModule(CurModule);
 					}
 				}
 
@@ -2099,6 +2099,16 @@ namespace UnrealBuildTool
 			// Setup the target's plugins
 			SetupPlugins();
 
+			// Create all the modules for each binary
+			foreach(UEBuildBinary Binary in AppBinaries)
+			{
+				foreach(string ModuleName in Binary.Config.ModuleNames)
+				{
+					UEBuildModule Module = FindOrCreateModuleByName(ModuleName);
+					Binary.AddModule(Module);
+				}
+			}
+
 			// Add the enabled plugins to the build
 			foreach (PluginInfo BuildPlugin in BuildPlugins)
 			{
@@ -2336,12 +2346,12 @@ namespace UnrealBuildTool
 			UEBuildBinaryCPP BinaryCPP = ExecutableBinary as UEBuildBinaryCPP;
 			if (BinaryCPP != null)
 			{
-				foreach (var TargetModuleName in BinaryCPP.ModuleNames)
+				foreach (var TargetModule in BinaryCPP.Modules)
 				{
-					ModuleRules CheckRules = RulesAssembly.CreateModuleRules(TargetModuleName, TargetInfo);
+					ModuleRules CheckRules = TargetModule.Rules;
 					if (CheckRules.Type != ModuleRules.ModuleType.External)
 					{
-						PrivateDependencyModuleNames.Add(TargetModuleName);
+						PrivateDependencyModuleNames.Add(TargetModule.Name);
 					}
 				}
 			}
@@ -2468,7 +2478,7 @@ namespace UnrealBuildTool
 			Module.CachePCHUsageForModuleSourceFiles(Module.CreateModuleCompileEnvironment(GlobalCompileEnvironment));
 
 			// Add module to binary
-			Binary.AddModule(Module.Name);
+			Binary.AddModule(Module);
 		}
 
 		/// <summary>
@@ -2516,9 +2526,9 @@ namespace UnrealBuildTool
 				var CPPBinary = Binary as UEBuildBinaryCPP;
 				if( CPPBinary != null )
 				{
-					foreach( var ModuleName in CPPBinary.ModuleNames )
+					foreach( var Module in CPPBinary.Modules )
 					{
-						var CPPModule = GetModuleByName( ModuleName ) as UEBuildModuleCPP;
+						var CPPModule = Module as UEBuildModuleCPP;
 						if( CPPModule != null )
 						{
 							if( !String.IsNullOrEmpty( CPPModule.Rules.SharedPCHHeaderFile ) && CPPModule.Binary.Config.bAllowCompilation )
@@ -2607,15 +2617,17 @@ namespace UnrealBuildTool
 				{
 					if (Module.IsCompiledInConfiguration(Platform, TargetType, UEBuildConfiguration.bBuildDeveloperTools, UEBuildConfiguration.bBuildEditor))
 					{
+						UEBuildModule ModuleInstance = FindOrCreateModuleByName(Module.Name);
+
 						// Add the corresponding binary for it
 						FileReference ModuleFileName = RulesAssembly.GetModuleFileName(Module.Name);
 						bool bHasSource = (ModuleFileName != null && Directory.EnumerateFiles(ModuleFileName.Directory.FullName, "*.cpp", SearchOption.AllDirectories).Any());
-						AddBinaryForModule(Module.Name, BinaryType, bAllowCompilation: bHasSource, bIsCrossTarget: false);
+						AddBinaryForModule(ModuleInstance, BinaryType, bAllowCompilation: bHasSource, bIsCrossTarget: false);
 
 						// Add it to the binary if we're compiling monolithic (and it's enabled)
 						if(ShouldCompileMonolithic() && EnabledPlugins.Contains(Plugin))
 						{
-							AppBinaries[0].AddModule(Module.Name);
+							AppBinaries[0].AddModule(ModuleInstance);
 						}
 					}
 				}
@@ -2629,8 +2641,9 @@ namespace UnrealBuildTool
 			// Add extra modules that will either link into the main binary (monolithic), or be linked into separate DLL files (modular)
 			foreach (var ModuleName in ExtraModuleNames)
 			{
-				UEBuildBinaryCPP Binary = FindOrAddBinaryForModule(ModuleName, false);
-				Binary.AddModule(ModuleName);
+				UEBuildModule Module = FindOrCreateModuleByName(ModuleName);
+				UEBuildBinaryCPP Binary = FindOrAddBinaryForModule(Module, false);
+				Binary.AddModule(Module);
 			}
 		}
 
@@ -2716,7 +2729,7 @@ namespace UnrealBuildTool
 					if(PrecompiledModule is UEBuildModuleCPP && !BoundModuleNames.Contains(PrecompiledModule.Name))
 					{
 						UEBuildBinaryType BinaryType = bCompileMonolithic? UEBuildBinaryType.StaticLibrary : UEBuildBinaryType.DynamicLinkLibrary;
-						UEBuildBinary Binary = AddBinaryForModule(PrecompiledModule.Name, BinaryType, bAllowCompilation: !bUsePrecompiled || RulesAssembly.IsGameModule(PrecompiledModule.Name), bIsCrossTarget: false);
+						UEBuildBinary Binary = AddBinaryForModule(PrecompiledModule, BinaryType, bAllowCompilation: !bUsePrecompiled || RulesAssembly.IsGameModule(PrecompiledModule.Name), bIsCrossTarget: false);
 						PrecompiledBinaries.Add(Binary);
 						BoundModuleNames.Add(PrecompiledModule.Name);
 					}
@@ -2724,7 +2737,7 @@ namespace UnrealBuildTool
 			}
 		}
 
-		public UEBuildBinaryCPP FindOrAddBinaryForModule(string ModuleName, bool bIsCrossTarget)
+		public UEBuildBinaryCPP FindOrAddBinaryForModule(UEBuildModule Module, bool bIsCrossTarget)
 		{
 			UEBuildBinaryCPP Binary;
 			if (ShouldCompileMonolithic())
@@ -2735,7 +2748,7 @@ namespace UnrealBuildTool
 			else
 			{
 				// Otherwise create a new module for it
-				Binary = AddBinaryForModule(ModuleName, UEBuildBinaryType.DynamicLinkLibrary, bAllowCompilation: true, bIsCrossTarget: bIsCrossTarget);
+				Binary = AddBinaryForModule(Module, UEBuildBinaryType.DynamicLinkLibrary, bAllowCompilation: true, bIsCrossTarget: bIsCrossTarget);
 			}
 			return Binary;
 		}
@@ -2748,10 +2761,10 @@ namespace UnrealBuildTool
 		/// <param name="bAllowCompilation">Whether this binary can be compiled. The function will check whether plugin binaries can be compiled.</param>
 		/// <param name="bIsCrossTarget">True if module is for supporting a different target-platform</param>
 		/// <returns>The new binary</returns>
-		private UEBuildBinaryCPP AddBinaryForModule(string ModuleName, UEBuildBinaryType BinaryType, bool bAllowCompilation, bool bIsCrossTarget)
+		private UEBuildBinaryCPP AddBinaryForModule(UEBuildModule Module, UEBuildBinaryType BinaryType, bool bAllowCompilation, bool bIsCrossTarget)
 		{
 			// Get the plugin info for this module
-			PluginInfo Plugin = FindPluginForModule(ModuleName);
+			PluginInfo Plugin = FindPluginForModule(Module.Name);
 
 			// Get the root output directory and base name (target name/app name) for this binary
 			DirectoryReference BaseOutputDirectory;
@@ -2759,7 +2772,7 @@ namespace UnrealBuildTool
 			{
 				BaseOutputDirectory = Plugin.Directory;
 			}
-			else if(RulesAssembly.IsGameModule(ModuleName) || !bUseSharedBuildEnvironment)
+			else if(RulesAssembly.IsGameModule(Module.Name) || !bUseSharedBuildEnvironment)
 			{
 				BaseOutputDirectory = ProjectDirectory;
 			}
@@ -2770,7 +2783,7 @@ namespace UnrealBuildTool
 
 			// Get the configuration that this module will be built in. Engine modules compiled in DebugGame will use Development.
 			UnrealTargetConfiguration ModuleConfiguration = Configuration;
-			if(Configuration == UnrealTargetConfiguration.DebugGame && !RulesAssembly.IsGameModule(ModuleName))
+			if(Configuration == UnrealTargetConfiguration.DebugGame && !RulesAssembly.IsGameModule(Module.Name))
 			{
 				ModuleConfiguration = UnrealTargetConfiguration.Development;
 			}
@@ -2780,15 +2793,14 @@ namespace UnrealBuildTool
  			DirectoryReference IntermediateDirectory = DirectoryReference.Combine(BaseOutputDirectory, BuildConfiguration.PlatformIntermediateFolder, AppName, ModuleConfiguration.ToString());
 
 			// Append a subdirectory if the module rules specifies one
-			ModuleRules ModuleRules;
-			if(RulesAssembly.TryCreateModuleRules(ModuleName, TargetInfo, out ModuleRules) && !String.IsNullOrEmpty(ModuleRules.BinariesSubFolder))
+			if(Module.Rules != null && !String.IsNullOrEmpty(Module.Rules.BinariesSubFolder))
 			{
-				OutputDirectory = DirectoryReference.Combine(OutputDirectory, ModuleRules.BinariesSubFolder);
-				IntermediateDirectory = DirectoryReference.Combine(IntermediateDirectory, ModuleRules.BinariesSubFolder);
+				OutputDirectory = DirectoryReference.Combine(OutputDirectory, Module.Rules.BinariesSubFolder);
+				IntermediateDirectory = DirectoryReference.Combine(IntermediateDirectory, Module.Rules.BinariesSubFolder);
 			}
 
 			// Get the output filenames
-			FileReference BaseBinaryPath = FileReference.Combine(OutputDirectory, MakeBinaryFileName(AppName + "-" + ModuleName, Platform, ModuleConfiguration, Rules.UndecoratedConfiguration, BinaryType));
+			FileReference BaseBinaryPath = FileReference.Combine(OutputDirectory, MakeBinaryFileName(AppName + "-" + Module.Name, Platform, ModuleConfiguration, Rules.UndecoratedConfiguration, BinaryType));
 			List<FileReference> OutputFilePaths = UEBuildPlatform.GetBuildPlatform(Platform).FinalizeBinaryPaths(BaseBinaryPath);
 
 			// Prepare the configuration object
@@ -2796,15 +2808,16 @@ namespace UnrealBuildTool
 				InType: BinaryType,
 				InOutputFilePaths: OutputFilePaths,
 				InIntermediateDirectory: IntermediateDirectory,
-				bInHasModuleRules: ModuleRules != null,
+				bInHasModuleRules: Module.Rules != null,
 				bInAllowExports: BinaryType == UEBuildBinaryType.DynamicLinkLibrary,
 				bInAllowCompilation: bAllowCompilation,
 				bInIsCrossTarget: bIsCrossTarget,
-				InModuleNames: new string[]{ ModuleName }
+				InModuleNames: new string[]{ Module.Name }
 			);
 
 			// Create the new binary
 			UEBuildBinaryCPP Binary = new UEBuildBinaryCPP(this, Config);
+			Binary.AddModule(Module);
 			AppBinaries.Add(Binary);
 			return Binary;
 		}
