Restore missing runtime dependencies in target receipts for editor configurations. Includes runtime dependencies for precompiled modules in monolithic targets at the moment.

diff --git a/Engine/Source/Programs/UnrealBuildTool/Configuration/UEBuildTarget.cs b/Engine/Source/Programs/UnrealBuildTool/Configuration/UEBuildTarget.cs
--- a/Engine/Source/Programs/UnrealBuildTool/Configuration/UEBuildTarget.cs
+++ b/Engine/Source/Programs/UnrealBuildTool/Configuration/UEBuildTarget.cs
@@ -1735,7 +1735,7 @@ namespace UnrealBuildTool
 			// Find all the modules which are part of this target
 			foreach(UEBuildModule Module in Modules.Values)
 			{
-				if(Module.bIncludedInTarget && Module.Binary != null && Module.Binary.Config.Type != UEBuildBinaryType.StaticLibrary)
+				if(Module.bIncludedInTarget && !PrecompiledBinaries.Contains(Module.Binary))
 				{
 					foreach(RuntimeDependency RuntimeDependency in Module.RuntimeDependencies)
 					{
