Use the directory cache when scanning for source files, and remove code which uses custom attributes (was taking ~2% CPU time when generating project files).

diff --git a/Engine/Source/Programs/UnrealBuildTool/System/RulesCompiler.cs b/Engine/Source/Programs/UnrealBuildTool/System/RulesCompiler.cs
--- a/Engine/Source/Programs/UnrealBuildTool/System/RulesCompiler.cs
+++ b/Engine/Source/Programs/UnrealBuildTool/System/RulesCompiler.cs
@@ -1533,110 +1533,24 @@ namespace UnrealBuildTool
 	public class RulesCompiler
 	{
 		/// <summary>
-		/// Helper class to avoid adding extra conditions when getting file extensions and suffixes for
-		/// rule files in FindAllRulesFilesRecursively.
+		/// Enum for types of rules files. Should match extensions in RulesFileExtensions.
 		/// </summary>
-		public class RulesTypePropertiesAttribute : Attribute
-		{
-			public string Suffix;
-			public string Extension;
-			public RulesTypePropertiesAttribute(string Suffix, string Extension)
-			{
-				this.Suffix = Suffix;
-				this.Extension = Extension;
-			}
-		}
-
 		public enum RulesFileType
 		{
-			/// *.Build.cs files
-			[RulesTypeProperties(Suffix: "Build", Extension: ".cs")]
 			Module,
-
-			/// *.Target.cs files
-			[RulesTypeProperties(Suffix: "Target", Extension: ".cs")]
 			Target,
-
-			/// *.Automation.cs files
-			[RulesTypeProperties(Suffix: "Automation", Extension: ".cs")]
 			Automation,
-
-			/// *.Automation.csproj files
-			[RulesTypeProperties(Suffix: "Automation", Extension: ".csproj")]
 			AutomationModule
 		}
 
+		/// <summary>
+		/// Cached list of rules files in each directory of each type
+		/// </summary>
 		class RulesFileCache
 		{
-			/// List of rules file paths for each of the known types in RulesFileType
-			public List<FileReference>[] RulesFilePaths = new List<FileReference>[ typeof( RulesFileType ).GetEnumValues().Length ];
-		}
-
-
-		private static void FindAllRulesFilesRecursively( DirectoryInfo DirInfo, RulesFileCache RulesFileCache )
-		{
-			if( DirInfo.Exists )
-			{
-				var RulesFileTypeEnum = typeof(RulesFileType);
-				bool bFoundModuleRulesFile = false;
-				var RulesFileTypes = typeof( RulesFileType ).GetEnumValues();
-				foreach( RulesFileType CurRulesType in RulesFileTypes )
-				{
-					// Get the suffix and extension associated with this RulesFileType enum value.
-					var MemberInfo = RulesFileTypeEnum.GetMember(CurRulesType.ToString());
-					var Attributes = MemberInfo[0].GetCustomAttributes(typeof(RulesTypePropertiesAttribute), false);
-					var EnumProperties = (RulesTypePropertiesAttribute)Attributes[0];
-					
-					var SearchRuleSuffix = "." + EnumProperties.Suffix + EnumProperties.Extension; // match files with the right suffix and extension.
-					var FilesInDirectory = DirInfo.GetFiles("*" + EnumProperties.Extension);
-					foreach (var RuleFile in FilesInDirectory)
-					{
-						// test if filename has the appropriate suffix.
-						// this handles filenames such as Foo.build.cs, Foo.Build.cs, foo.bUiLd.cs to fix bug 266743 on platforms where case-sensitivity matters
-						if (RuleFile.Name.EndsWith(SearchRuleSuffix, StringComparison.InvariantCultureIgnoreCase))
-						{
-							// Skip Uncooked targets, as those are no longer valid.  This is just for easier backwards compatibility with existing projects.
-							// @todo: Eventually we can eliminate this conditional and just allow it to be an error when these are compiled
-							if( CurRulesType != RulesFileType.Target || !RuleFile.Name.EndsWith( "Uncooked" + SearchRuleSuffix, StringComparison.InvariantCultureIgnoreCase ) )
-							{
-								if (RulesFileCache.RulesFilePaths[(int)CurRulesType] == null)
-								{
-									RulesFileCache.RulesFilePaths[(int)CurRulesType] = new List<FileReference>();
-								}
-
-								// Convert file info to the full file path for this file and update our cache
-								RulesFileCache.RulesFilePaths[(int)CurRulesType].Add(new FileReference(RuleFile.FullName));
-
-								// NOTE: Multiple rules files in the same folder are supported.  We'll continue iterating along.
-								if( CurRulesType == RulesFileType.Module )
-								{
-									bFoundModuleRulesFile = true;
-								}
-							}
-							else
-							{
-								Log.TraceVerbose("Skipped deprecated Target rules file with Uncooked extension: " + RuleFile.Name );
-							}
-						}
-					}
-				}
-
-				// Only recurse if we didn't find a module rules file.  In the interest of performance and organizational sensibility
-				// we don't want to support folders with Build.cs files containing other folders with Build.cs files.  Performance-
-				// wise, this is really important to avoid scanning every folder in the Source/ThirdParty directory, for example.
-				if( !bFoundModuleRulesFile )
-				{
-					// Add all the files recursively
-					foreach( DirectoryInfo SubDirInfo in DirInfo.GetDirectories() )
-					{
-						if( SubDirInfo.Name.Equals( "Intermediate", StringComparison.InvariantCultureIgnoreCase ) )
-						{
-							Console.WriteLine( "WARNING: UnrealBuildTool found an Intermediate folder while looking for rules '{0}'.  It should only ever be searching under 'Source' folders -- an Intermediate folder is unexpected and will greatly decrease iteration times!", SubDirInfo.FullName );
-						}
-						FindAllRulesFilesRecursively(SubDirInfo, RulesFileCache);
-					}
-				}
-			}
+			public List<FileReference> ModuleRules = new List<FileReference>();
+			public List<FileReference> TargetRules = new List<FileReference>();
+			public List<FileReference> AutomationModules = new List<FileReference>();
 		}
 
 		/// Map of root folders to a cached list of all UBT-related source files in that folder or any of its sub-folders.
@@ -1718,6 +1632,7 @@ namespace UnrealBuildTool
 
 			// Iterate over all the folders to check
 			List<FileReference> SourceFiles = new List<FileReference>();
+			HashSet<FileReference> UniqueSourceFiles = new HashSet<FileReference>();
 			foreach( DirectoryReference Folder in Folders )
 			{
 				// Check to see if we've already cached source files for this folder
@@ -1725,37 +1640,73 @@ namespace UnrealBuildTool
 				if (!RootFolderToRulesFileCache.TryGetValue(Folder, out FolderRulesFileCache))
 				{
 					FolderRulesFileCache = new RulesFileCache();
-					FindAllRulesFilesRecursively(new DirectoryInfo(Folder.FullName), FolderRulesFileCache);
+					FindAllRulesFilesRecursively(Folder, FolderRulesFileCache);
 					RootFolderToRulesFileCache[Folder] = FolderRulesFileCache;
+				}
 
-					if (BuildConfiguration.bPrintDebugInfo)
-					{
-						foreach (var CurType in Enum.GetValues(typeof(RulesFileType)))
-						{
-							var RulesFiles = FolderRulesFileCache.RulesFilePaths[(int)CurType];
-							if (RulesFiles != null)
-							{
-								Log.TraceVerbose("Found {0} rules files for folder {1} of type {2}", RulesFiles.Count, Folder, CurType.ToString());
-							}
-						}
-					}
+				// Get the list of files of the type we're looking for
+				List<FileReference> SourceFilesForFolder;
+				if(RulesFileType == RulesCompiler.RulesFileType.Module)
+				{
+					SourceFilesForFolder = FolderRulesFileCache.ModuleRules;
+				}
+				else if(RulesFileType == RulesCompiler.RulesFileType.Target)
+				{
+					SourceFilesForFolder = FolderRulesFileCache.TargetRules;
+				}
+				else if(RulesFileType == RulesCompiler.RulesFileType.AutomationModule)
+				{
+					SourceFilesForFolder = FolderRulesFileCache.AutomationModules;
+				}
+				else
+				{
+					throw new BuildException("Unhandled rules type: {0}", RulesFileType);
 				}
 
-				var RulesFilePathsForType = FolderRulesFileCache.RulesFilePaths[(int)RulesFileType];
-				if (RulesFilePathsForType != null)
+				// Add them to the output list
+				foreach(FileReference SourceFile in SourceFilesForFolder)
 				{
-					foreach (FileReference RulesFilePath in RulesFilePathsForType)
+					if(UniqueSourceFiles.Add(SourceFile))
 					{
-						if (!SourceFiles.Contains(RulesFilePath))
-						{
-							SourceFiles.Add(RulesFilePath);
-						}
+						SourceFiles.Add(SourceFile);
 					}
 				}
 			}
 			return SourceFiles;
 		}
 
+		private static void FindAllRulesFilesRecursively(DirectoryReference Directory, RulesFileCache Cache)
+		{
+			// Scan all the files in this directory
+			bool bSearchSubFolders = true;
+			foreach(FileReference File in DirectoryLookupCache.EnumerateFiles(Directory))
+			{
+				if(File.HasExtension(".build.cs"))
+				{
+					Cache.ModuleRules.Add(File);
+					bSearchSubFolders = false;
+				}
+				else if(File.HasExtension(".target.cs"))
+				{
+					Cache.TargetRules.Add(File);
+				}
+				else if(File.HasExtension(".automation.csproj"))
+				{
+					Cache.AutomationModules.Add(File);
+					bSearchSubFolders = false;
+				}
+			}
+
+			// If we didn't find anything to stop the search, search all the subdirectories too
+			if(bSearchSubFolders)
+			{
+				foreach(DirectoryReference SubDirectory in DirectoryLookupCache.EnumerateDirectories(Directory))
+				{
+					FindAllRulesFilesRecursively(SubDirectory, Cache);
+				}
+			}
+		}
+
 		/// <summary>
 		/// The cached rules assembly for engine modules and targets.
 		/// </summary>
