Don't create receipts when generating project files; they aren't needed.

diff --git a/Engine/Source/Programs/UnrealBuildTool/Configuration/UEBuildTarget.cs b/Engine/Source/Programs/UnrealBuildTool/Configuration/UEBuildTarget.cs
--- a/Engine/Source/Programs/UnrealBuildTool/Configuration/UEBuildTarget.cs
+++ b/Engine/Source/Programs/UnrealBuildTool/Configuration/UEBuildTarget.cs
@@ -1918,7 +1918,10 @@ namespace UnrealBuildTool
 			}
 
 			// Create a receipt for the target
-			PrepareReceipts(TargetToolChain);
+			if(!ProjectFileGenerator.bGenerateProjectFiles)
+			{
+				PrepareReceipts(TargetToolChain);
+			}
 
 			// If we're only generating the manifest, return now
 			if (UEBuildConfiguration.bGenerateManifest || UEBuildConfiguration.bCleanProject)
