Cache the list of all valid UBT platforms rather than fetching it via reflection each time; reflection is very slow, and we call it a lot.

diff --git a/Engine/Source/Programs/UnrealBuildTool/System/UnrealBuildTool.cs b/Engine/Source/Programs/UnrealBuildTool/System/UnrealBuildTool.cs
--- a/Engine/Source/Programs/UnrealBuildTool/System/UnrealBuildTool.cs
+++ b/Engine/Source/Programs/UnrealBuildTool/System/UnrealBuildTool.cs
@@ -80,6 +80,9 @@ namespace UnrealBuildTool
 		/** True if this run of UBT should only invalidate Makefile. */
 		static public bool bIsInvalidatingMakefilesOnly = false;
 
+		/** Cached array of all the supported target platforms */
+		static public UnrealTargetPlatform[] AllPlatforms = (UnrealTargetPlatform[])Enum.GetValues(typeof(UnrealTargetPlatform));
+
         /** True if we should gather module dependencies for building in this run.  If this is false, then we'll expect to be loading information from disk about
             the target's modules before we'll be able to build anything.  One or both of IsGatheringBuild or IsAssemblingBuild must be true. */
         static public bool IsGatheringBuild
@@ -284,7 +287,7 @@ namespace UnrealBuildTool
         static private bool GetPlatforms(ref List<UnrealTargetPlatform> OutPlatforms, Func<UnrealTargetPlatform, bool> Predicate, bool bCheckValidity = true)
         {
             OutPlatforms.Clear();
-            foreach (UnrealTargetPlatform Platform in Enum.GetValues(typeof(UnrealTargetPlatform)))
+            foreach (UnrealTargetPlatform Platform in AllPlatforms)
             {
                 if (Predicate(Platform))
                 {
