Store the list of valid plugins on the rules assembly, and allow looking up the plugin info for a module without having to scan the whole list.

diff --git a/Engine/Source/Programs/UnrealBuildTool/Configuration/UEBuildModule.cs b/Engine/Source/Programs/UnrealBuildTool/Configuration/UEBuildModule.cs
--- a/Engine/Source/Programs/UnrealBuildTool/Configuration/UEBuildModule.cs
+++ b/Engine/Source/Programs/UnrealBuildTool/Configuration/UEBuildModule.cs
@@ -182,8 +182,7 @@ namespace UnrealBuildTool
 		private readonly bool? IsRedistributableOverride;
 
 		/** The name of the .Build.cs file this module was created from, if any */
-		private readonly string BuildCsFilenameField;
-		public string BuildCsFilename { get { return BuildCsFilenameField; } }
+		public FileReference RulesFile;
 
 		/**
 		 * Tells if this module can be redistributed.
@@ -245,14 +244,16 @@ namespace UnrealBuildTool
 			UEBuildModuleType InType,
 			DirectoryReference InModuleDirectory,
 			ModuleRules InRules,
-			string InBuildCsFilename
+			FileReference InRulesFile
 			)
 		{
 			Target = InTarget;
 			Name = InName;
 			Type = InType;
-			Rules = InRules;
 			ModuleDirectory = InModuleDirectory;
+			Rules = InRules;
+			RulesFile = InRulesFile;
+
 			PublicDefinitions = HashSetFromOptionalEnumerableStringParameter(InRules.Definitions);
 			PublicIncludePaths = HashSetFromOptionalEnumerableStringParameter(InRules.PublicIncludePaths);
 			PublicSystemIncludePaths = HashSetFromOptionalEnumerableStringParameter(InRules.PublicSystemIncludePaths);
@@ -268,9 +269,6 @@ namespace UnrealBuildTool
 			RuntimeDependencies = (InRules.RuntimeDependencies == null)? new List<RuntimeDependency>() : new List<RuntimeDependency>(InRules.RuntimeDependencies);
 			IsRedistributableOverride = InRules.IsRedistributableOverride;
 
-			Debug.Assert(InBuildCsFilename == null || InBuildCsFilename.EndsWith(".Build.cs", StringComparison.InvariantCultureIgnoreCase));
-			BuildCsFilenameField = InBuildCsFilename;
-
 			Target.RegisterModule(this);
 		}
 
@@ -621,11 +619,14 @@ namespace UnrealBuildTool
 			if(Modules == null)
 			{
 				Modules = new List<UEBuildModule>();
-				foreach(string ModuleName in ModuleNames.Distinct(StringComparer.InvariantCultureIgnoreCase))
+				foreach(string ModuleName in ModuleNames)
 				{
 					UEBuildModule Module = Target.FindOrCreateModuleByName(ModuleName);
-					Module.RecursivelyCreateModules();
-					Modules.Add(Module);
+					if(!Modules.Contains(Module))
+					{
+						Module.RecursivelyCreateModules();
+						Modules.Add(Module);
+					}
 				}
 			}
 		}
@@ -655,7 +656,7 @@ namespace UnrealBuildTool
 			string InName,
 			DirectoryReference InModuleDirectory,
 			ModuleRules InRules,
-			string InBuildCsFilename
+			FileReference InRulesFile
 			)
 		: base(
 			InTarget:										InTarget,
@@ -663,7 +664,7 @@ namespace UnrealBuildTool
 			InName:											InName,
 			InModuleDirectory:								InModuleDirectory,
 			InRules:										InRules,
-			InBuildCsFilename:								InBuildCsFilename
+			InRulesFile:									InRulesFile
 			)
 		{
 			bIncludedInTarget = true;
@@ -843,14 +844,14 @@ namespace UnrealBuildTool
 			IEnumerable<FileItem> InSourceFiles,
 			ModuleRules InRules,
 			bool bInBuildSourceFiles,
-			string InBuildCsFilename
+			FileReference InRulesFile
 			)
 			: base(	InTarget,
 					InName, 
 					InType,
 					InModuleDirectory,
 					InRules,
-					InBuildCsFilename
+					InRulesFile
 				)
 		{
 			GeneratedCodeDirectory = InGeneratedCodeDirectory;
@@ -1763,11 +1764,11 @@ namespace UnrealBuildTool
 			IEnumerable<FileItem> InSourceFiles,
 			ModuleRules InRules,
 			bool bInBuildSourceFiles,
-			string InBuildCsFilename
+			FileReference InRulesFile
 			)
 			: base(InTarget,InName,InType,InModuleDirectory,InGeneratedCodeDirectory,InIntelliSenseGatherer,
 			InSourceFiles,InRules,
-			bInBuildSourceFiles, InBuildCsFilename)
+			bInBuildSourceFiles, InRulesFile)
 		{
 			PrivateAssemblyReferences = HashSetFromOptionalEnumerableStringParameter(InRules.PrivateAssemblyReferences);
 		}
diff --git a/Engine/Source/Programs/UnrealBuildTool/Configuration/UEBuildTarget.cs b/Engine/Source/Programs/UnrealBuildTool/Configuration/UEBuildTarget.cs
--- a/Engine/Source/Programs/UnrealBuildTool/Configuration/UEBuildTarget.cs
+++ b/Engine/Source/Programs/UnrealBuildTool/Configuration/UEBuildTarget.cs
@@ -2507,7 +2507,7 @@ namespace UnrealBuildTool
 				InSourceFiles: SourceFiles.ToList(),
 				InRules: Rules,
 				bInBuildSourceFiles: true,
-				InBuildCsFilename: null);
+				InRulesFile: null);
 		}
 
 
@@ -2767,7 +2767,11 @@ namespace UnrealBuildTool
 		private UEBuildBinaryCPP AddBinaryForModule(UEBuildModule Module, UEBuildBinaryType BinaryType, bool bAllowCompilation, bool bIsCrossTarget)
 		{
 			// Get the plugin info for this module
-			PluginInfo Plugin = FindPluginForModule(Module.Name);
+			PluginInfo Plugin = null;
+			if(Module.RulesFile != null)
+			{
+				RulesAssembly.TryGetPluginForModule(Module.RulesFile, out Plugin);
+			}
 
 			// Get the root output directory and base name (target name/app name) for this binary
 			DirectoryReference BaseOutputDirectory;
@@ -2825,16 +2829,6 @@ namespace UnrealBuildTool
 			return Binary;
 		}
 
-		/// <summary>
-		/// Find the plugin which contains a given module
-		/// </summary>
-		/// <param name="ModuleName">Name of the module</param>
-		/// <returns>Matching plugin, or null if not found</returns>
-		private PluginInfo FindPluginForModule(string ModuleName)
-		{
-			return ValidPlugins.FirstOrDefault(ValidPlugin => ValidPlugin.Descriptor.Modules != null && ValidPlugin.Descriptor.Modules.Any(Module => Module.Name == ModuleName));
-		}
-
 		/**
 		 * @return true if debug information is created, false otherwise.
 		 */
@@ -2945,7 +2939,7 @@ namespace UnrealBuildTool
 		protected virtual void SetupPlugins()
 		{
 			// Filter the plugins list by the current project
-			ValidPlugins = Plugins.ReadAvailablePlugins(UnrealBuildTool.EngineDirectory, UnrealBuildTool.GetUProjectFile());
+			ValidPlugins = new List<PluginInfo>(RulesAssembly.EnumeratePlugins());
 
 			// Remove any plugins for platforms we don't have
 			List<string> ExcludeFolders = new List<string>();
@@ -3005,13 +2999,7 @@ namespace UnrealBuildTool
 				foreach(FileReference ForeignPlugin in ForeignPlugins)
 				{
 					PluginInfo ForeignPluginInfo = ValidPlugins.FirstOrDefault(x => x.File == ForeignPlugin);
-					if(ForeignPluginInfo == null)
-					{
-						ForeignPluginInfo = new PluginInfo(ForeignPlugin, PluginLoadedFrom.GameProject);
-						ValidPlugins.Add(ForeignPluginInfo);
-						BuildPlugins.Add(ForeignPluginInfo);
-					}
-					else if(!BuildPlugins.Contains(ForeignPluginInfo))
+					if(!BuildPlugins.Contains(ForeignPluginInfo))
 					{
 						BuildPlugins.Add(ForeignPluginInfo);
 					}
@@ -3339,7 +3327,7 @@ namespace UnrealBuildTool
 		{
 			Debug.Assert(Module.Target == this);
 			Modules.Add(Module.Name, Module);
-			FlatModuleCsData.Add(Module.Name, new FlatModuleCsDataType(Module.BuildCsFilename));
+			FlatModuleCsData.Add(Module.Name, new FlatModuleCsDataType((Module.RulesFile == null)? null : Module.RulesFile.FullName));
 		}
 
 		/** Finds a module given its name.  Throws an exception if the module couldn't be found. */
@@ -3357,7 +3345,8 @@ namespace UnrealBuildTool
 				DirectoryReference ModuleDirectory = ModuleFileName.Directory;
 				
 				// Get the plugin for this module
-				PluginInfo Plugin = FindPluginForModule(ModuleName);
+				PluginInfo Plugin;
+				RulesAssembly.TryGetPluginForModule(ModuleFileName, out Plugin);
 
 				// Get the type of module we're creating
 				var ModuleType = UEBuildModuleType.Unknown;
@@ -3435,7 +3424,7 @@ namespace UnrealBuildTool
 				if (RulesObject.Type != ModuleRules.ModuleType.External && ModuleName != "Core")
 				{
 					// Add the default include paths to the module rules, if they exist.
-					AddDefaultIncludePathsToModuleRules(ModuleFileName, IsGameModule, RulesObject);
+					AddDefaultIncludePathsToModuleRules(ModuleFileName, IsGameModule, Plugin, RulesObject);
 
 					// Add the path to the generated headers 
 					if(GeneratedCodeDirectory != null)
@@ -3507,7 +3496,7 @@ namespace UnrealBuildTool
 			IntelliSenseGatherer IntelliSenseGatherer,
 			List<FileItem>       ModuleSourceFiles,
 			bool                 bBuildSourceFiles,
-			FileReference         InBuildCsFile)
+			FileReference         InRulesFile)
 		{
 			switch (RulesObject.Type)
 			{
@@ -3522,7 +3511,7 @@ namespace UnrealBuildTool
 							InSourceFiles: ModuleSourceFiles,
 							InRules: RulesObject,
 							bInBuildSourceFiles: bBuildSourceFiles,
-							InBuildCsFilename: InBuildCsFile.FullName
+							InRulesFile: InRulesFile
 						);
 
 				case ModuleRules.ModuleType.CPlusPlusCLR:
@@ -3536,7 +3525,7 @@ namespace UnrealBuildTool
 							InSourceFiles: ModuleSourceFiles,
 							InRules: RulesObject,
 							bInBuildSourceFiles : bBuildSourceFiles,
-							InBuildCsFilename: InBuildCsFile.FullName
+							InRulesFile: InRulesFile
 						);
 
 				case ModuleRules.ModuleType.External:
@@ -3546,7 +3535,7 @@ namespace UnrealBuildTool
 							InType: ModuleType,
 							InModuleDirectory: ModuleDirectory,
 							InRules: RulesObject,
-							InBuildCsFilename: InBuildCsFile.FullName
+							InRulesFile: InRulesFile
 						);
 
 				default:
@@ -3560,14 +3549,11 @@ namespace UnrealBuildTool
 		/// <param name="ModuleFile">The filename to the module rules file (Build.cs)</param>
 		/// <param name="IsGameModule">true if it is a game module, false if not</param>
 		/// <param name="RulesObject">The module rules object itself</param>
-		public void AddDefaultIncludePathsToModuleRules(FileReference ModuleFile, bool IsGameModule, ModuleRules RulesObject)
+		public void AddDefaultIncludePathsToModuleRules(FileReference ModuleFile, bool IsGameModule, PluginInfo Plugin, ModuleRules RulesObject)
 		{
 			// Get the base source directory for this module. This may be the project source directory, engine source directory, or plugin source directory.
 			if(!ModuleFile.IsUnderDirectory(UnrealBuildTool.EngineSourceDirectory))
 			{
-				// Check if it belongs to a plugin. We check against the directory here, which may give different results than checking the plugin modules list.
-				PluginInfo Plugin = ValidPlugins.FirstOrDefault(ValidPlugin => ModuleFile.IsUnderDirectory(ValidPlugin.Directory));
-
 				// Add the module source directory 
 				DirectoryReference BaseSourceDirectory;
 				if(Plugin != null)
diff --git a/Engine/Source/Programs/UnrealBuildTool/System/Plugins.cs b/Engine/Source/Programs/UnrealBuildTool/System/Plugins.cs
--- a/Engine/Source/Programs/UnrealBuildTool/System/Plugins.cs
+++ b/Engine/Source/Programs/UnrealBuildTool/System/Plugins.cs
@@ -74,25 +74,47 @@ namespace UnrealBuildTool
 
 			// Read all the engine plugins
 			DirectoryReference EnginePluginsDirectoryName = DirectoryReference.Combine(EngineDirectoryName, "Plugins");
-			Plugins.AddRange(ReadAvailablePluginsInternal(EnginePluginsDirectoryName, PluginLoadedFrom.Engine));
+			Plugins.AddRange(ReadPluginsFromDirectory(EnginePluginsDirectoryName, PluginLoadedFrom.Engine));
 
 			// Read all the project plugins
 			if(ProjectFileName != null)
 			{
 				DirectoryReference ProjectPluginsDir = DirectoryReference.Combine(ProjectFileName.Directory, "Plugins");
-				Plugins.AddRange(ReadAvailablePluginsInternal(ProjectPluginsDir, PluginLoadedFrom.GameProject));
+				Plugins.AddRange(ReadPluginsFromDirectory(ProjectPluginsDir, PluginLoadedFrom.GameProject));
 			}
 
 			return Plugins;
 		}
 
+		/// <summary>
+		/// Read all the plugin descriptors under the given engine directory
+		/// </summary>
+		/// <param name="EngineDirectory">The parent directory to look in.</param>
+		/// <returns>Sequence of the found PluginInfo object.</returns>
+		public static IReadOnlyList<PluginInfo> ReadEnginePlugins(DirectoryReference EngineDirectory)
+		{
+			DirectoryReference EnginePluginsDirectory = DirectoryReference.Combine(EngineDirectory, "Plugins");
+			return ReadPluginsFromDirectory(EnginePluginsDirectory, PluginLoadedFrom.Engine);
+		}
+
+		/// <summary>
+		/// Read all the plugin descriptors under the given project directory
+		/// </summary>
+		/// <param name="EngineDirectory">The parent directory to look in.</param>
+		/// <returns>Sequence of the found PluginInfo object.</returns>
+		public static IReadOnlyList<PluginInfo> ReadProjectPlugins(DirectoryReference ProjectDirectory)
+		{
+			DirectoryReference ProjectPluginsDirectory = DirectoryReference.Combine(ProjectDirectory, "Plugins");
+			return ReadPluginsFromDirectory(ProjectPluginsDirectory, PluginLoadedFrom.GameProject);
+		}
+
 		/// <summary>
 		/// Read all the plugin descriptors under the given directory
 		/// </summary>
 		/// <param name="ParentDirectory">The parent directory to look in.</param>
 		/// <param name="LoadedFrom">The directory type</param>
 		/// <returns>Sequence of the found PluginInfo object.</returns>
-		static IEnumerable<PluginInfo> ReadAvailablePluginsInternal(DirectoryReference ParentDirectory, PluginLoadedFrom LoadedFrom)
+		public static IReadOnlyList<PluginInfo> ReadPluginsFromDirectory(DirectoryReference ParentDirectory, PluginLoadedFrom LoadedFrom)
 		{
 			List<PluginInfo> Plugins;
 			if(!PluginInfoCache.TryGetValue(ParentDirectory, out Plugins))
diff --git a/Engine/Source/Programs/UnrealBuildTool/System/RulesCompiler.cs b/Engine/Source/Programs/UnrealBuildTool/System/RulesCompiler.cs
--- a/Engine/Source/Programs/UnrealBuildTool/System/RulesCompiler.cs
+++ b/Engine/Source/Programs/UnrealBuildTool/System/RulesCompiler.cs
@@ -1082,30 +1082,45 @@ namespace UnrealBuildTool
 		/// </summary>
 		private Assembly CompiledAssembly;
 
+		/// <summary>
+		/// All the plugins included in this assembly
+		/// </summary>
+		private IReadOnlyList<PluginInfo> Plugins;
+
 		/// <summary>
 		/// Maps module names to their actual xxx.Module.cs file on disk
 		/// </summary>
-		private Dictionary<string, FileReference> ModuleNameToFileName = new Dictionary<string, FileReference>(StringComparer.InvariantCultureIgnoreCase);
+		private Dictionary<string, FileReference> ModuleNameToModuleFile = new Dictionary<string, FileReference>(StringComparer.InvariantCultureIgnoreCase);
 
 		/// <summary>
 		/// Maps target names to their actual xxx.Target.cs file on disk
 		/// </summary>
-		private Dictionary<string, FileReference> TargetNameToFileName = new Dictionary<string, FileReference>(StringComparer.InvariantCultureIgnoreCase);
+		private Dictionary<string, FileReference> TargetNameToTargetFile = new Dictionary<string, FileReference>(StringComparer.InvariantCultureIgnoreCase);
+
+		/// <summary>
+		/// Mapping from module file to its plugin info.
+		/// </summary>
+		private Dictionary<FileReference, PluginInfo> ModuleFileToPluginInfo;
 
 		/// <summary>
 		/// Constructor. Compiles a rules assembly from the given source files.
 		/// </summary>
-		/// <param name="ModuleFileNames">List of module files to compile</param>
-		/// <param name="TargetFileNames">List of target files to compile</param>
+		/// <param name="Plugins">All the plugins included in this assembly</param>
+		/// <param name="ModuleFiles">List of module files to compile</param>
+		/// <param name="TargetFiles">List of target files to compile</param>
+		/// <param name="ModuleFileToPluginInfo">Mapping of module file to the plugin that contains it</param>
 		/// <param name="AssemblyFileName">The output path for the compiled assembly</param>
-		public RulesAssembly(List<FileReference> ModuleFileNames, List<FileReference> TargetFileNames, FileReference AssemblyFileName, RulesAssembly InParent)
+		/// <param name="Parent">The parent rules assembly</param>
+		public RulesAssembly(IReadOnlyList<PluginInfo> Plugins, List<FileReference> ModuleFiles, List<FileReference> TargetFiles, Dictionary<FileReference, PluginInfo> ModuleFileToPluginInfo, FileReference AssemblyFileName, RulesAssembly Parent)
 		{
-			Parent = InParent;
+			this.Plugins = Plugins;
+			this.ModuleFileToPluginInfo = ModuleFileToPluginInfo;
+			this.Parent = Parent;
 
 			// Find all the source files
 			List<FileReference> AssemblySourceFiles = new List<FileReference>();
-			AssemblySourceFiles.AddRange( ModuleFileNames );
-			AssemblySourceFiles.AddRange( TargetFileNames );
+			AssemblySourceFiles.AddRange(ModuleFiles);
+			AssemblySourceFiles.AddRange(TargetFiles);
 
 			// Compile the assembly
 			if(AssemblySourceFiles.Count > 0)
@@ -1114,22 +1129,22 @@ namespace UnrealBuildTool
 			}
 
 			// Setup the module map
-			foreach( FileReference CurModuleFileName in ModuleFileNames )
+			foreach(FileReference ModuleFile in ModuleFiles)
 			{
-				string ModuleName = Path.GetFileNameWithoutExtension( Path.GetFileNameWithoutExtension( CurModuleFileName.FullName ) );	// Strip both extensions
-				if( !ModuleNameToFileName.ContainsKey( ModuleName ) )
+				string ModuleName = ModuleFile.GetFileNameWithoutAnyExtensions();
+				if(!ModuleNameToModuleFile.ContainsKey(ModuleName))
 				{
-					ModuleNameToFileName.Add( ModuleName, CurModuleFileName );
+					ModuleNameToModuleFile.Add(ModuleName, ModuleFile);
 				}
 			}
 
 			// Setup the target map
-			foreach( FileReference CurTargetFileName in TargetFileNames )
+			foreach(FileReference TargetFile in TargetFiles)
 			{
-				string TargetName = Path.GetFileNameWithoutExtension( Path.GetFileNameWithoutExtension( CurTargetFileName.FullName ) );	// Strip both extensions
-				if( !TargetNameToFileName.ContainsKey( TargetName ) )
+				string TargetName = TargetFile.GetFileNameWithoutAnyExtensions();
+				if(!TargetNameToTargetFile.ContainsKey(TargetName))
 				{
-					TargetNameToFileName.Add( TargetName, CurTargetFileName );
+					TargetNameToTargetFile.Add(TargetName, TargetFile);
 				}
 			}
 		}
@@ -1140,29 +1155,29 @@ namespace UnrealBuildTool
 		/// <param name="ExistingType"></param>
 		/// <param name="FileName"></param>
 		/// <returns>True if the type was found, false otherwise</returns>
-		public bool TryGetFileNameFromType(Type ExistingType, out FileReference FileName)
+		public bool TryGetFileNameFromType(Type ExistingType, out FileReference File)
 		{
 			if(ExistingType.Assembly == CompiledAssembly)
 			{
 				string Name = ExistingType.Name;
-				if(ModuleNameToFileName.TryGetValue(Name, out FileName))
+				if(ModuleNameToModuleFile.TryGetValue(Name, out File))
 				{
 					return true;
 				}
-				if(TargetNameToFileName.TryGetValue(Name, out FileName))
+				if(TargetNameToTargetFile.TryGetValue(Name, out File))
 				{
 					return true;
 				}
 			}
 			else
 			{
-				if(Parent != null && Parent.TryGetFileNameFromType(ExistingType, out FileName))
+				if(Parent != null && Parent.TryGetFileNameFromType(ExistingType, out File))
 				{
 					return true;
 				}
 			}
 
-			FileName = null;
+			File = null;
 			return false;
 		}
 
@@ -1173,10 +1188,10 @@ namespace UnrealBuildTool
 		/// <returns>The filename containing rules for this module, or an empty string if not found</returns>
 		public FileReference GetModuleFileName(string ModuleName)
 		{
-			FileReference FileName;
-			if(ModuleNameToFileName.TryGetValue(ModuleName, out FileName))
+			FileReference ModuleFile;
+			if(ModuleNameToModuleFile.TryGetValue(ModuleName, out ModuleFile))
 			{
-				return FileName;
+				return ModuleFile;
 			}
 			else
 			{
@@ -1191,10 +1206,10 @@ namespace UnrealBuildTool
 		/// <returns>The filename containing rules for this target, or an empty string if not found</returns>
 		public FileReference GetTargetFileName(string TargetName)
 		{
-			FileReference FileName;
-			if(TargetNameToFileName.TryGetValue(TargetName, out FileName))
+			FileReference TargetFile;
+			if(TargetNameToTargetFile.TryGetValue(TargetName, out TargetFile))
 			{
-				return FileName;
+				return TargetFile;
 			}
 			else
 			{
@@ -1227,7 +1242,7 @@ namespace UnrealBuildTool
 			var ModuleTypeName = ModuleName;
 
 			// Make sure the module file is known to us
-			if (!ModuleNameToFileName.TryGetValue(ModuleName, out ModuleFileName))
+			if (!ModuleNameToModuleFile.TryGetValue(ModuleName, out ModuleFileName))
 			{
 				if(Parent == null)
 				{
@@ -1420,7 +1435,7 @@ namespace UnrealBuildTool
 		public TargetRules CreateTargetRules(string TargetName, TargetInfo Target, bool bInEditorRecompile, out FileReference TargetFileName)
 		{
 			// Make sure the target file is known to us
-			bool bFoundTargetName = TargetNameToFileName.ContainsKey(TargetName);
+			bool bFoundTargetName = TargetNameToTargetFile.ContainsKey(TargetName);
 			if (bFoundTargetName == false)
 			{
 				if (UnrealBuildTool.RunningRocket())
@@ -1428,7 +1443,7 @@ namespace UnrealBuildTool
 					//@todo Rocket: Remove this when full game support is implemented
 					// If we are Rocket, they will currently only have an editor target.
 					// See if that exists
-					bFoundTargetName = TargetNameToFileName.ContainsKey(TargetName + "Editor");
+					bFoundTargetName = TargetNameToTargetFile.ContainsKey(TargetName + "Editor");
 					if (bFoundTargetName)
 					{
 						TargetName += "Editor";
@@ -1450,7 +1465,7 @@ namespace UnrealBuildTool
 					ExceptionMessage += "Location: " + CompiledAssembly.Location + Environment.NewLine;
 
 					ExceptionMessage += "Target rules found:" + Environment.NewLine;
-					foreach (KeyValuePair<string, FileReference> entry in TargetNameToFileName)
+					foreach (KeyValuePair<string, FileReference> entry in TargetNameToTargetFile)
 					{
 						ExceptionMessage += "\t" + entry.Key + " - " + entry.Value + Environment.NewLine;
 					}
@@ -1464,7 +1479,7 @@ namespace UnrealBuildTool
 			}
 
 			// Return the target file name to the caller
-			TargetFileName = TargetNameToFileName[TargetName];
+			TargetFileName = TargetNameToTargetFile[TargetName];
 
 			// Currently, we expect the user's rules object type name to be the same as the module name + 'Target'
 			string TargetTypeName = TargetName + "Target";
@@ -1496,7 +1511,7 @@ namespace UnrealBuildTool
 						if (SourceFolderIndex != -1)
 						{
 							DirectoryReference TargetSourceFolder = new DirectoryReference(TargetSourceFolderString.Substring(0, SourceFolderIndex + 8));
-							foreach (KeyValuePair<string, FileReference> CheckEntry in TargetNameToFileName)
+							foreach (KeyValuePair<string, FileReference> CheckEntry in TargetNameToTargetFile)
 							{
 								if (CheckEntry.Value.IsUnderDirectory(TargetSourceFolder))
 								{
@@ -1528,6 +1543,40 @@ namespace UnrealBuildTool
 
 			return RulesObject;
 		}
+
+		/// <summary>
+		/// Enumerates all the plugins that are available
+		/// </summary>
+		/// <returns></returns>
+		public IEnumerable<PluginInfo> EnumeratePlugins()
+		{
+			if(Parent == null)
+			{
+				return Plugins;
+			}
+			else
+			{
+				return Plugins.Concat(Parent.EnumeratePlugins());
+			}
+		}
+
+		/// <summary>
+		/// Tries to find the PluginInfo associated with a given module file
+		/// </summary>
+		/// <param name="ModuleFile">The module to search for</param>
+		/// <param name="Plugin">The matching plugin info, or null.</param>
+		/// <returns>True if the module belongs to a plugin</returns>
+		public bool TryGetPluginForModule(FileReference ModuleFile, out PluginInfo Plugin)
+		{
+			if(ModuleFileToPluginInfo.TryGetValue(ModuleFile, out Plugin))
+			{
+				return true;
+			}
+			else
+			{
+				return (Parent == null)? false : Parent.TryGetPluginForModule(ModuleFile, out Plugin);
+			}
+		}
 	}
 
 	public class RulesCompiler
@@ -1635,35 +1684,7 @@ namespace UnrealBuildTool
 			HashSet<FileReference> UniqueSourceFiles = new HashSet<FileReference>();
 			foreach( DirectoryReference Folder in Folders )
 			{
-				// Check to see if we've already cached source files for this folder
-				RulesFileCache FolderRulesFileCache;
-				if (!RootFolderToRulesFileCache.TryGetValue(Folder, out FolderRulesFileCache))
-				{
-					FolderRulesFileCache = new RulesFileCache();
-					FindAllRulesFilesRecursively(Folder, FolderRulesFileCache);
-					RootFolderToRulesFileCache[Folder] = FolderRulesFileCache;
-				}
-
-				// Get the list of files of the type we're looking for
-				List<FileReference> SourceFilesForFolder;
-				if(RulesFileType == RulesCompiler.RulesFileType.Module)
-				{
-					SourceFilesForFolder = FolderRulesFileCache.ModuleRules;
-				}
-				else if(RulesFileType == RulesCompiler.RulesFileType.Target)
-				{
-					SourceFilesForFolder = FolderRulesFileCache.TargetRules;
-				}
-				else if(RulesFileType == RulesCompiler.RulesFileType.AutomationModule)
-				{
-					SourceFilesForFolder = FolderRulesFileCache.AutomationModules;
-				}
-				else
-				{
-					throw new BuildException("Unhandled rules type: {0}", RulesFileType);
-				}
-
-				// Add them to the output list
+				IReadOnlyList<FileReference> SourceFilesForFolder = FindAllRulesFiles(Folder, RulesFileType);
 				foreach(FileReference SourceFile in SourceFilesForFolder)
 				{
 					if(UniqueSourceFiles.Add(SourceFile))
@@ -1675,6 +1696,36 @@ namespace UnrealBuildTool
 			return SourceFiles;
 		}
 
+		private static IReadOnlyList<FileReference> FindAllRulesFiles(DirectoryReference Directory, RulesFileType Type)
+		{
+			// Check to see if we've already cached source files for this folder
+			RulesFileCache Cache;
+			if (!RootFolderToRulesFileCache.TryGetValue(Directory, out Cache))
+			{
+				Cache = new RulesFileCache();
+				FindAllRulesFilesRecursively(Directory, Cache);
+				RootFolderToRulesFileCache[Directory] = Cache;
+			}
+
+			// Get the list of files of the type we're looking for
+			if(Type == RulesCompiler.RulesFileType.Module)
+			{
+				return Cache.ModuleRules;
+			}
+			else if(Type == RulesCompiler.RulesFileType.Target)
+			{
+				return Cache.TargetRules;
+			}
+			else if(Type == RulesCompiler.RulesFileType.AutomationModule)
+			{
+				return Cache.AutomationModules;
+			}
+			else
+			{
+				throw new BuildException("Unhandled rules type: {0}", Type);
+			}
+		}
+
 		private static void FindAllRulesFilesRecursively(DirectoryReference Directory, RulesFileCache Cache)
 		{
 			// Scan all the files in this directory
@@ -1726,12 +1777,17 @@ namespace UnrealBuildTool
 			if(EngineRulesAssembly == null)
 			{
 				// Find all the rules files
-				List<FileReference> ModuleFileNames = FindAllRulesSourceFiles(RulesFileType.Module, null, null, null, true);
-				List<FileReference> TargetFileNames = FindAllRulesSourceFiles(RulesFileType.Target, null, null, null, true);
+				List<FileReference> ModuleFiles = new List<FileReference>(FindAllRulesFiles(UnrealBuildTool.EngineSourceDirectory, RulesFileType.Module));
+				List<FileReference> TargetFiles = new List<FileReference>(FindAllRulesFiles(UnrealBuildTool.EngineSourceDirectory, RulesFileType.Target));
+
+				// Add all the plugin modules too
+				IReadOnlyList<PluginInfo> EnginePlugins = Plugins.ReadEnginePlugins(UnrealBuildTool.EngineDirectory);
+				Dictionary<FileReference, PluginInfo> ModuleFileToPluginInfo = new Dictionary<FileReference,PluginInfo>();
+				FindModuleRulesForPlugins(EnginePlugins, ModuleFiles, ModuleFileToPluginInfo);
 
 				// Create a path to the assembly that we'll either load or compile
 				FileReference AssemblyFileName = FileReference.Combine(UnrealBuildTool.EngineDirectory, BuildConfiguration.BaseIntermediateFolder, "BuildRules", "UE4Rules.dll");
-				EngineRulesAssembly = new RulesAssembly(ModuleFileNames, TargetFileNames, AssemblyFileName, null);
+				EngineRulesAssembly = new RulesAssembly(EnginePlugins, ModuleFiles, TargetFiles, ModuleFileToPluginInfo, AssemblyFileName, null);
 			}
 			return EngineRulesAssembly;
 		}
@@ -1747,30 +1803,30 @@ namespace UnrealBuildTool
 			RulesAssembly ProjectRulesAssembly;
 			if(!LoadedAssemblyMap.TryGetValue(ProjectFileName, out ProjectRulesAssembly))
 			{
-				DirectoryReference ProjectDirectoryName = ProjectFileName.Directory;
-
 				// Create the engine rules assembly
 				RulesAssembly Parent = CreateEngineRulesAssembly();
 
-				// Get the folders to search through
-				List<DirectoryReference> GameFolders = new List<DirectoryReference>();
-				GameFolders.Add(ProjectDirectoryName);
+				// Find all the rules under the project source directory
+				DirectoryReference ProjectDirectory = ProjectFileName.Directory;
+				DirectoryReference ProjectSourceDirectory = DirectoryReference.Combine(ProjectDirectory, "Source");
+				List<FileReference> ModuleFiles = new List<FileReference>(FindAllRulesFiles(ProjectSourceDirectory, RulesFileType.Module));
+				List<FileReference> TargetFiles = new List<FileReference>(FindAllRulesFiles(ProjectSourceDirectory, RulesFileType.Target));
+
+				// Find all the project plugins
+				IReadOnlyList<PluginInfo> ProjectPlugins = Plugins.ReadProjectPlugins(ProjectFileName.Directory);
+				Dictionary<FileReference, PluginInfo> ModuleFileToPluginInfo = new Dictionary<FileReference,PluginInfo>();
+				FindModuleRulesForPlugins(ProjectPlugins, ModuleFiles, ModuleFileToPluginInfo);
 
 				// Add the games project's intermediate source folder
-				List<DirectoryReference> AdditionalSearchPaths = new List<DirectoryReference>();
-				DirectoryReference ProjectIntermediateSourceDirectory = DirectoryReference.Combine(ProjectDirectoryName, "Intermediate", "Source");
+				DirectoryReference ProjectIntermediateSourceDirectory = DirectoryReference.Combine(ProjectDirectory, "Intermediate", "Source");
 				if (ProjectIntermediateSourceDirectory.Exists())
 				{
-					AdditionalSearchPaths.Add(ProjectIntermediateSourceDirectory);
+					TargetFiles.AddRange(FindAllRulesFiles(ProjectIntermediateSourceDirectory, RulesFileType.Target));
 				}
 
-				// Find all the rules source files
-				List<FileReference> ModuleFileNames = FindAllRulesSourceFiles(RulesFileType.Module, GameFolders, null, AdditionalSearchPaths, false);
-				List<FileReference> TargetFileNames = FindAllRulesSourceFiles(RulesFileType.Target, GameFolders, null, AdditionalSearchPaths, false);
-
 				// Compile the assembly
-				FileReference AssemblyFileName = FileReference.Combine(ProjectDirectoryName, BuildConfiguration.BaseIntermediateFolder, "BuildRules", ProjectFileName.GetFileNameWithoutExtension() + "ModuleRules.dll");
-				ProjectRulesAssembly = new RulesAssembly(ModuleFileNames, TargetFileNames, AssemblyFileName, Parent);
+				FileReference AssemblyFileName = FileReference.Combine(ProjectDirectory, BuildConfiguration.BaseIntermediateFolder, "BuildRules", ProjectFileName.GetFileNameWithoutExtension() + "ModuleRules.dll");
+				ProjectRulesAssembly = new RulesAssembly(ProjectPlugins, ModuleFiles, TargetFiles, ModuleFileToPluginInfo, AssemblyFileName, Parent);
 				LoadedAssemblyMap.Add(ProjectFileName, ProjectRulesAssembly);
 			}
 			return ProjectRulesAssembly;
@@ -1788,17 +1844,47 @@ namespace UnrealBuildTool
 			if(!LoadedAssemblyMap.TryGetValue(PluginFileName, out PluginRulesAssembly))
 			{
 				// Find all the rules source files
-				List<FileReference> ModuleFileNames = FindAllRulesSourceFiles(RulesFileType.Module, null, new List<FileReference>{ PluginFileName }, null, false);
-				List<FileReference> TargetFileNames = FindAllRulesSourceFiles(RulesFileType.Target, null, new List<FileReference>{ PluginFileName }, null, false);
+				List<FileReference> ModuleFiles = new List<FileReference>();
+				List<FileReference> TargetFiles = new List<FileReference>();
+
+				// Create a list of plugins for this assembly. If it already exists in the parent assembly, just create an empty assembly.
+				List<PluginInfo> ForeignPlugins = new List<PluginInfo>();
+				if(Parent == null || !Parent.EnumeratePlugins().Any(x => x.File == PluginFileName))
+				{
+					ForeignPlugins.Add(new PluginInfo(PluginFileName, PluginLoadedFrom.GameProject));
+				}
+
+				// Find all the modules
+				Dictionary<FileReference, PluginInfo> ModuleFileToPluginInfo = new Dictionary<FileReference,PluginInfo>();
+				FindModuleRulesForPlugins(ForeignPlugins, ModuleFiles, ModuleFileToPluginInfo);
 
 				// Compile the assembly
 				FileReference AssemblyFileName = FileReference.Combine(PluginFileName.Directory, BuildConfiguration.BaseIntermediateFolder, "BuildRules", Path.GetFileNameWithoutExtension(PluginFileName.FullName) + "ModuleRules.dll");
-				PluginRulesAssembly = new RulesAssembly(ModuleFileNames, TargetFileNames, AssemblyFileName, Parent);
+				PluginRulesAssembly = new RulesAssembly(ForeignPlugins, ModuleFiles, TargetFiles, ModuleFileToPluginInfo, AssemblyFileName, Parent);
 				LoadedAssemblyMap.Add(PluginFileName, PluginRulesAssembly);
 			}
 			return PluginRulesAssembly;
 		}
 
+		/// <summary>
+		/// Finds all the module rules for plugins under the given directory.
+		/// </summary>
+		/// <param name="PluginsDirectory">The directory to search</param>
+		/// <param name="ModuleFiles">List of module files to be populated</param>
+		/// <param name="ModuleFileToPluginFile">Dictionary which is filled with mappings from the module file to its corresponding plugin file</param>
+		private static void FindModuleRulesForPlugins(IReadOnlyList<PluginInfo> Plugins, List<FileReference> ModuleFiles, Dictionary<FileReference, PluginInfo> ModuleFileToPluginInfo)
+		{
+			foreach(PluginInfo Plugin in Plugins)
+			{
+				IReadOnlyList<FileReference> PluginModuleFiles = FindAllRulesFiles(DirectoryReference.Combine(Plugin.Directory, "Source"), RulesFileType.Module);
+				foreach(FileReference ModuleFile in PluginModuleFiles)
+				{
+					ModuleFiles.Add(ModuleFile);
+					ModuleFileToPluginInfo[ModuleFile] = Plugin;
+				}
+			}
+		}
+
 		/// <summary>
 		/// Gets the filename that declares the given type.
 		/// </summary>
