Cache the list of libraries to link dependencies with rather than recalculating every time.

diff --git a/Engine/Source/Programs/UnrealBuildTool/Configuration/UEBuildBinary.cs b/Engine/Source/Programs/UnrealBuildTool/Configuration/UEBuildBinary.cs
--- a/Engine/Source/Programs/UnrealBuildTool/Configuration/UEBuildBinary.cs
+++ b/Engine/Source/Programs/UnrealBuildTool/Configuration/UEBuildBinary.cs
@@ -412,6 +412,7 @@ namespace UnrealBuildTool
 		public readonly List<UEBuildModule> Modules = new List<UEBuildModule>();
 		private bool bCreateImportLibrarySeparately;
 		private bool bIncludeDependentLibrariesInLibrary;
+		private List<string> DependentLinkLibraries;
 
 		/// <summary>
 		/// Create an instance initialized to the given configuration
@@ -570,22 +571,26 @@ namespace UnrealBuildTool
 		/// <param name="DependentLinkEnvironment">The link environment of the dependency</param>
 		public override void SetupDependentLinkEnvironment(LinkEnvironment DependentLinkEnvironment)
 		{
-			foreach (FileReference OutputFilePath in Config.OutputFilePaths)
+			// Cache the list of libraries in the dependent link environment between calls. We typically run this code path many times for each module.
+			if(DependentLinkLibraries == null)
 			{
-				FileReference LibraryFileName;
-				if (Config.Type == UEBuildBinaryType.StaticLibrary
-					|| DependentLinkEnvironment.Config.Target.Platform == CPPTargetPlatform.Mac
-					|| DependentLinkEnvironment.Config.Target.Platform == CPPTargetPlatform.Linux)
-				{
-					LibraryFileName = OutputFilePath;
-				}
-				else
+				DependentLinkLibraries = new List<string>();
+				foreach (FileReference OutputFilePath in Config.OutputFilePaths)
 				{
-					LibraryFileName = FileReference.Combine(Config.IntermediateDirectory, OutputFilePath.GetFileNameWithoutExtension() + ".lib");
+					FileReference LibraryFileName;
+					if (Config.Type == UEBuildBinaryType.StaticLibrary || DependentLinkEnvironment.Config.Target.Platform == CPPTargetPlatform.Mac || DependentLinkEnvironment.Config.Target.Platform == CPPTargetPlatform.Linux)
+					{
+						LibraryFileName = OutputFilePath;
+					}
+					else
+					{
+						LibraryFileName = FileReference.Combine(Config.IntermediateDirectory, OutputFilePath.GetFileNameWithoutExtension() + ".lib");
+					}
+					DependentLinkLibraries.Add(LibraryFileName.FullName);
 				}
-				DependentLinkEnvironment.Config.AdditionalLibraries.Add(LibraryFileName.FullName);
 			}
-			
+			DependentLinkEnvironment.Config.AdditionalLibraries.AddRange(DependentLinkLibraries);
+
 			// If we're linking against static library containing the launch module on windows, we need to add the compiled resource separately. We can't link it through the static library.
 			if(Config.Type == UEBuildBinaryType.StaticLibrary && Modules.Any(x => x.Name == "Launch") && (Target.Platform == UnrealTargetPlatform.Win32 || Target.Platform == UnrealTargetPlatform.Win64))
 			{
