UE-18280 - fix for not being able to build Mac from PC using SSH
#ubt

diff --git a/Engine/Source/Programs/UnrealBuildTool/Mac/MacToolChain.cs b/Engine/Source/Programs/UnrealBuildTool/Mac/MacToolChain.cs
--- a/Engine/Source/Programs/UnrealBuildTool/Mac/MacToolChain.cs
+++ b/Engine/Source/Programs/UnrealBuildTool/Mac/MacToolChain.cs
@@ -1155,7 +1155,7 @@ namespace UnrealBuildTool
 				DebugOutputMap.Add (GenDebugAction, OutputFile.AbsolutePath);
 				GenDebugAction.ActionHandler = new Action.BlockingActionHandler(MacToolChain.RPCDebugInfoActionHandler);
 			}
-			GenDebugAction.WorkingDirectory = UnrealBuildTool.EngineSourceDirectory.FullName;
+			GenDebugAction.WorkingDirectory = GetMacDevSrcRoot();
 			GenDebugAction.CommandPath = "sh";
 
 			// Deletes ay existing file on the building machine,
@@ -1186,7 +1186,7 @@ namespace UnrealBuildTool
 
 			// Make the compile action
 			Action FinalizeAppBundleAction = new Action(ActionType.CreateAppBundle);
-			FinalizeAppBundleAction.WorkingDirectory = UnrealBuildTool.EngineSourceDirectory.FullName;
+			FinalizeAppBundleAction.WorkingDirectory = GetMacDevSrcRoot(); // Path.GetFullPath(".");
 			FinalizeAppBundleAction.CommandPath = "/bin/sh";
 			FinalizeAppBundleAction.CommandDescription = "";
 
@@ -1211,12 +1211,12 @@ namespace UnrealBuildTool
 		FileItem CopyBundleResource(UEBuildBundleResource Resource, FileItem Executable)
 		{
 			Action CopyAction = new Action(ActionType.CreateAppBundle);
-			CopyAction.WorkingDirectory = UnrealBuildTool.EngineSourceDirectory.FullName;
+			CopyAction.WorkingDirectory = GetMacDevSrcRoot(); // Path.GetFullPath(".");
 			CopyAction.CommandPath = "/bin/sh";
 			CopyAction.CommandDescription = "";
 
 			string BundlePath = Executable.AbsolutePath.Substring(0, Executable.AbsolutePath.IndexOf(".app") + 4);
-			string SourcePath = Path.Combine(CopyAction.WorkingDirectory, Resource.ResourcePath);
+			string SourcePath = Path.Combine(Path.GetFullPath("."), Resource.ResourcePath);
 			string TargetPath = Path.Combine(BundlePath, "Contents", Resource.BundleContentsSubdir, Path.GetFileName(Resource.ResourcePath));
 
 			FileItem TargetItem;
diff --git a/Engine/Source/Programs/UnrealBuildTool/Mac/UEDeployMac.cs b/Engine/Source/Programs/UnrealBuildTool/Mac/UEDeployMac.cs
--- a/Engine/Source/Programs/UnrealBuildTool/Mac/UEDeployMac.cs
+++ b/Engine/Source/Programs/UnrealBuildTool/Mac/UEDeployMac.cs
@@ -26,10 +26,14 @@ namespace UnrealBuildTool.IOS
 		{
 			Log.TraceInformation("Deploying now!");
 
-			string IntermediateDirectory = InTarget.AppName + "/Intermediate/Build/Mac/" + InTarget.AppName + "/" + InTarget.Configuration;
-			if (!Directory.Exists("../../" + IntermediateDirectory))
+			string IntermediateDirectory = InTarget.EngineIntermediateDirectory.FullName;
+			if (!Directory.Exists(IntermediateDirectory))
 			{
-				IntermediateDirectory = "Engine/Intermediate/Build/Mac/" + InTarget.AppName + "/" + InTarget.Configuration;
+				IntermediateDirectory = Path.GetFullPath("../../" + InTarget.AppName + "/Intermediate/Build/Mac/" + InTarget.AppName + "/" + InTarget.Configuration);
+				if (!Directory.Exists(IntermediateDirectory))
+				{
+					IntermediateDirectory = Path.GetFullPath("../../Engine/Intermediate/Build/Mac/" + InTarget.AppName + "/" + InTarget.Configuration);
+				}
 			}
 
 			MacToolChain Toolchain = UEToolChain.GetPlatformToolChain(CPPTargetPlatform.Mac) as MacToolChain;
@@ -48,13 +52,13 @@ namespace UnrealBuildTool.IOS
 					// Copy the command scripts to the intermediate on the target Mac.
 					string RemoteFixDylibDepsScript = Toolchain.ConvertPath(Path.GetFullPath(FixDylibDepsScript));
 					RemoteFixDylibDepsScript = RemoteFixDylibDepsScript.Replace("../../../../", "../../");
-					RPCUtilHelper.CopyFile("../../" + FixDylibDepsScript, RemoteFixDylibDepsScript, true);
+					RPCUtilHelper.CopyFile(Path.GetFullPath(FixDylibDepsScript), RemoteFixDylibDepsScript, true);
 
 					if (!InTarget.GlobalLinkEnvironment.Config.bIsBuildingConsoleApplication)
 					{
 						string RemoteFinalizeAppBundleScript = Toolchain.ConvertPath(Path.GetFullPath(FinalizeAppBundleScript));
 						RemoteFinalizeAppBundleScript = RemoteFinalizeAppBundleScript.Replace("../../../../", "../../");
-						RPCUtilHelper.CopyFile("../../" + FinalizeAppBundleScript, RemoteFinalizeAppBundleScript, true);
+						RPCUtilHelper.CopyFile(Path.GetFullPath(FinalizeAppBundleScript), RemoteFinalizeAppBundleScript, true);
 					}
 
 
@@ -106,7 +110,7 @@ namespace UnrealBuildTool.IOS
 						
 						// Get the app bundle's name
 						string AppFullName = InTarget.AppName;
-						if (InTarget.Configuration != UnrealTargetConfiguration.Development)
+						if (InTarget.Configuration != InTarget.Rules.UndecoratedConfiguration)
 						{
 							AppFullName += "-" + InTarget.Platform.ToString();
 							AppFullName += "-" + InTarget.Configuration.ToString();
diff --git a/Engine/Source/Programs/UnrealBuildTool/ToolChain/RemoteToolChain.cs b/Engine/Source/Programs/UnrealBuildTool/ToolChain/RemoteToolChain.cs
--- a/Engine/Source/Programs/UnrealBuildTool/ToolChain/RemoteToolChain.cs
+++ b/Engine/Source/Programs/UnrealBuildTool/ToolChain/RemoteToolChain.cs
@@ -646,11 +646,14 @@ namespace UnrealBuildTool
 		}
 
 		private static Dictionary<Object,StringBuilder> SSHOutputMap = new Dictionary<object,StringBuilder>();
+		private static System.Threading.Mutex DictionaryLock = new System.Threading.Mutex();
 		static public void OutputReceivedForSSH(Object Sender, DataReceivedEventArgs Line)
 		{
 			if ((Line != null) && (Line.Data != null) && (Line.Data != ""))
 			{
+				DictionaryLock.WaitOne();
 				StringBuilder SSHOutput = SSHOutputMap[Sender];
+				DictionaryLock.ReleaseMutex();
 				if (SSHOutput.Length != 0)
 				{
 					SSHOutput.Append(Environment.NewLine);
@@ -892,7 +895,9 @@ namespace UnrealBuildTool
 			Hashtable Return = new Hashtable();
 
 			// add this process to the map
+			DictionaryLock.WaitOne();
 			SSHOutputMap[SSHProcess] = new StringBuilder("");
+			DictionaryLock.ReleaseMutex();
 			SSHProcess.OutputDataReceived += new DataReceivedEventHandler(OutputReceivedForSSH);
 			SSHProcess.ErrorDataReceived += new DataReceivedEventHandler(OutputReceivedForSSH);
 
@@ -901,10 +906,12 @@ namespace UnrealBuildTool
 			Console.WriteLine("Execute took {0}", (DateTime.Now - Start).ToString());
 
 			// now we have enough to fill out the HashTable
+			DictionaryLock.WaitOne();
 			Return["CommandOutput"] = SSHOutputMap[SSHProcess].ToString();
 			Return["ExitCode"] = (object)ExitCode;
 
 			SSHOutputMap.Remove(SSHProcess);
+			DictionaryLock.ReleaseMutex();
 
 			return Return;
 		}
