Cache the name of the API define for each module and the relative path to its base directory, to avoid calculating them each time.

diff --git a/Engine/Source/Programs/UnrealBuildTool/Configuration/UEBuildBinary.cs b/Engine/Source/Programs/UnrealBuildTool/Configuration/UEBuildBinary.cs
--- a/Engine/Source/Programs/UnrealBuildTool/Configuration/UEBuildBinary.cs
+++ b/Engine/Source/Programs/UnrealBuildTool/Configuration/UEBuildBinary.cs
@@ -663,7 +663,7 @@ namespace UnrealBuildTool
 		private LinkEnvironment SetupBinaryLinkEnvironment(LinkEnvironment LinkEnvironment, CPPEnvironment CompileEnvironment)
 		{
 			var BinaryLinkEnvironment = LinkEnvironment.DeepCopy();
-			var LinkEnvironmentVisitedModules = new Dictionary<UEBuildModule, bool>();
+			var LinkEnvironmentVisitedModules = new HashSet<UEBuildModule>();
 			var BinaryDependencies = new List<UEBuildBinary>();
 			CompileEnvironment.Config.bIsBuildingDLL = IsBuildingDll(Config.Type);
 			CompileEnvironment.Config.bIsBuildingLibrary = IsBuildingLibrary(Config.Type);
diff --git a/Engine/Source/Programs/UnrealBuildTool/Configuration/UEBuildModule.cs b/Engine/Source/Programs/UnrealBuildTool/Configuration/UEBuildModule.cs
--- a/Engine/Source/Programs/UnrealBuildTool/Configuration/UEBuildModule.cs
+++ b/Engine/Source/Programs/UnrealBuildTool/Configuration/UEBuildModule.cs
@@ -202,6 +202,12 @@ namespace UnrealBuildTool
 		/** Whether this module is included in the current target.  Only set after UEBuildBinary.BindModules is called. */
 		public bool bIncludedInTarget = false;
 
+		/** Include path for this module's base directory, relative to the Engine/Source directory */
+		protected string NormalizedModuleIncludePath;
+
+		/** The name of the _API define for this module */
+		protected readonly string ModuleApiDefine;
+
 		protected readonly HashSet<string> PublicDefinitions;
 		protected readonly HashSet<string> PublicIncludePaths;
 		protected readonly HashSet<string> PrivateIncludePaths;
@@ -254,6 +260,9 @@ namespace UnrealBuildTool
 			Rules = InRules;
 			RulesFile = InRulesFile;
 
+			NormalizedModuleIncludePath = Utils.CleanDirectorySeparators(ModuleDirectory.MakeRelativeTo(UnrealBuildTool.EngineSourceDirectory), '/');
+			ModuleApiDefine = Name.ToUpperInvariant() + "_API";
+
 			PublicDefinitions = HashSetFromOptionalEnumerableStringParameter(InRules.Definitions);
 			PublicIncludePaths = HashSetFromOptionalEnumerableStringParameter(InRules.PublicIncludePaths);
 			PublicSystemIncludePaths = HashSetFromOptionalEnumerableStringParameter(InRules.PublicSystemIncludePaths);
@@ -286,14 +295,12 @@ namespace UnrealBuildTool
 			HashSet<string> SystemIncludePaths,
 			List<string> Definitions,
 			List<UEBuildFramework> AdditionalFrameworks,
-			Dictionary<UEBuildModule, bool> VisitedModules
+			HashSet<UEBuildModule> VisitedModules
 			)
 		{
 			// There may be circular dependencies in compile dependencies, so we need to avoid reentrance.
-			if(!VisitedModules.ContainsKey(this))
+			if(VisitedModules.Add(this))
 			{
-				VisitedModules.Add(this,true);
-
 				// Add this module's public include paths and definitions.
 				AddIncludePathsWithChecks(IncludePaths, PublicIncludePaths);
 				AddIncludePathsWithChecks(SystemIncludePaths, PublicSystemIncludePaths);
@@ -308,7 +315,7 @@ namespace UnrealBuildTool
 					if(bIncludePathsOnly)
 					{
 						Log.TraceVerbose( "{0}: Include paths only for {1} (no binary)", SourceBinary.Config.OutputFilePaths[0].GetFileNameWithoutExtension(), Name );
-						Definitions.Add( Name.ToUpperInvariant() + "_API=" );
+						Definitions.Add( ModuleApiDefine + "=" );
 					}
 				}
 				else
@@ -320,19 +327,19 @@ namespace UnrealBuildTool
 					{
 						// When generating IntelliSense files, never add dllimport/dllexport specifiers as it
 						// simply confuses the compiler
-						Definitions.Add( Name.ToUpperInvariant() + "_API=" );
+						Definitions.Add( ModuleApiDefine + "=" );
 					}
 					else if( Binary == SourceBinary )
 					{
 						if( Binary.Config.bAllowExports )
 						{
 							Log.TraceVerbose( "{0}: Exporting {1} from {2}", SourceBinaryPath.GetFileNameWithoutExtension(), Name, BinaryPath.GetFileNameWithoutExtension() );
-							Definitions.Add( Name.ToUpperInvariant() + "_API=DLLEXPORT" );
+							Definitions.Add( ModuleApiDefine + "=DLLEXPORT" );
 						}
 						else
 						{
 							Log.TraceVerbose( "{0}: Not importing/exporting {1} (binary: {2})", SourceBinaryPath.GetFileNameWithoutExtension(), Name, BinaryPath.GetFileNameWithoutExtension() );
-							Definitions.Add( Name.ToUpperInvariant() + "_API=" );
+							Definitions.Add( ModuleApiDefine + "=" );
 						}
 					}
 					else
@@ -343,17 +350,17 @@ namespace UnrealBuildTool
 						if( bIncludePathsOnly )
 						{
 							Log.TraceVerbose( "{0}: Include paths only for {1} (binary: {2})", SourceBinaryPath.GetFileNameWithoutExtension(), Name, BinaryPath.GetFileNameWithoutExtension() );
-							Definitions.Add( Name.ToUpperInvariant() + "_API=" );
+							Definitions.Add( ModuleApiDefine + "=" );
 						}
 						else if (Binary.Config.bAllowExports)
 						{
 							Log.TraceVerbose( "{0}: Importing {1} from {2}", SourceBinaryPath.GetFileNameWithoutExtension(), Name, BinaryPath.GetFileNameWithoutExtension() );
-							Definitions.Add( Name.ToUpperInvariant() + "_API=DLLIMPORT" );
+							Definitions.Add( ModuleApiDefine + "=DLLIMPORT" );
 						}
 						else
 						{
 							Log.TraceVerbose( "{0}: Not importing/exporting {1} (binary: {2})", SourceBinaryPath.GetFileNameWithoutExtension(), Name, BinaryPath.GetFileNameWithoutExtension());
-							Definitions.Add( Name.ToUpperInvariant() + "_API=" );
+							Definitions.Add( ModuleApiDefine + "=" );
 						}
 					}
 				}
@@ -375,7 +382,7 @@ namespace UnrealBuildTool
 				}
 
 				// Add the module's directory to the include path, so we can root #includes to it
-				IncludePaths.Add(Utils.CleanDirectorySeparators(ModuleDirectory.MakeRelativeTo(UnrealBuildTool.EngineSourceDirectory), '/'));
+				IncludePaths.Add(NormalizedModuleIncludePath);
 
 				// Add the additional frameworks so that the compiler can know about their #include paths
 				AdditionalFrameworks.AddRange(PublicAdditionalFrameworks);
@@ -427,7 +434,7 @@ namespace UnrealBuildTool
 			List<UEBuildFramework> AdditionalFrameworks
 			)
 		{
-			Dictionary<UEBuildModule, bool> VisitedModules = new Dictionary<UEBuildModule, bool>();
+			HashSet<UEBuildModule> VisitedModules = new HashSet<UEBuildModule>();
 
 			if (this.Type == UEBuildModuleType.Game)
 			{
@@ -467,14 +474,12 @@ namespace UnrealBuildTool
 			List<UEBuildBundleResource> AdditionalBundleResources,
 			List<string> DelayLoadDLLs,
 			List<UEBuildBinary> BinaryDependencies,
-			Dictionary<UEBuildModule, bool> VisitedModules
+			HashSet<UEBuildModule> VisitedModules
 			)
 		{
 			// There may be circular dependencies in compile dependencies, so we need to avoid reentrance.
-			if(!VisitedModules.ContainsKey(this))
+			if(VisitedModules.Add(this))
 			{
-				VisitedModules.Add(this,true);
-
 				// Only process modules that are included in the current target.
 				if(bIncludedInTarget)
 				{
@@ -533,7 +538,7 @@ namespace UnrealBuildTool
 			UEBuildBinary SourceBinary,
 			LinkEnvironment LinkEnvironment,
 			List<UEBuildBinary> BinaryDependencies,
-			Dictionary<UEBuildModule, bool> VisitedModules
+			HashSet<UEBuildModule> VisitedModules
 			)
 		{
 			// Allow the module's public dependencies to add library paths and additional libraries to the link environment.
@@ -1224,7 +1229,7 @@ namespace UnrealBuildTool
 									SharedPCHCompileEnvironment.Config.CPPIncludeInfo.SystemIncludePaths, 
 									SharedPCHCompileEnvironment.Config.Definitions, 
 									SharedPCHCompileEnvironment.Config.AdditionalFrameworks, 
-									new Dictionary<UEBuildModule,bool>());
+									new HashSet<UEBuildModule>());
 
 								PCHOutput = PrecompileHeaderEnvironment.GeneratePCHCreationAction( 
 									Target,
@@ -1795,7 +1800,7 @@ namespace UnrealBuildTool
 			UEBuildBinary SourceBinary,
 			LinkEnvironment LinkEnvironment,
 			List<UEBuildBinary> BinaryDependencies,
-			Dictionary<UEBuildModule, bool> VisitedModules
+			HashSet<UEBuildModule> VisitedModules
 			)
 		{
 			base.SetupPrivateLinkEnvironment(SourceBinary, LinkEnvironment, BinaryDependencies, VisitedModules);
diff --git a/Engine/Source/Programs/UnrealBuildTool/Mac/MacToolChain.cs b/Engine/Source/Programs/UnrealBuildTool/Mac/MacToolChain.cs
--- a/Engine/Source/Programs/UnrealBuildTool/Mac/MacToolChain.cs
+++ b/Engine/Source/Programs/UnrealBuildTool/Mac/MacToolChain.cs
@@ -1338,7 +1338,7 @@ namespace UnrealBuildTool
 				var Modules = Binary.GetAllDependencyModules(bIncludeDynamicallyLoaded: false, bForceCircular: false);
 				var BinaryLinkEnvironment = Binary.Target.GlobalLinkEnvironment.DeepCopy();
 				var BinaryDependencies = new List<UEBuildBinary>();
-				var LinkEnvironmentVisitedModules = new Dictionary<UEBuildModule, bool>();
+				var LinkEnvironmentVisitedModules = new HashSet<UEBuildModule>();
 				foreach (var Module in Modules)
 				{
 					Module.SetupPrivateLinkEnvironment(Binary, BinaryLinkEnvironment, BinaryDependencies, LinkEnvironmentVisitedModules);
