Only expand $(PluginDir) variables once we know they're needed; the code path is so frequently called that it shows up in the profiler.

diff --git a/Engine/Source/Programs/UnrealBuildTool/System/RulesCompiler.cs b/Engine/Source/Programs/UnrealBuildTool/System/RulesCompiler.cs
--- a/Engine/Source/Programs/UnrealBuildTool/System/RulesCompiler.cs
+++ b/Engine/Source/Programs/UnrealBuildTool/System/RulesCompiler.cs
@@ -1268,7 +1268,6 @@ namespace UnrealBuildTool
 				RulesObjectType = CompiledAssembly.GetType("UnrealBuildTool.Rules." + ModuleName);
 			}
 
-			var AssemblyFileName = Path.GetFileNameWithoutExtension(CompiledAssembly.Location);
 			if (RulesObjectType == null)
 			{
 				throw new BuildException("Expecting to find a type to be declared in a module rules named '{0}' in {1}.  This type must derive from the 'ModuleRules' type defined by Unreal Build Tool.", ModuleTypeName, CompiledAssembly.FullName);
@@ -1282,7 +1281,7 @@ namespace UnrealBuildTool
 			}
 			catch (Exception Ex)
 			{
-				throw new BuildException(Ex, "Unable to instantiate instance of '{0}' object type from compiled assembly '{1}'.  Unreal Build Tool creates an instance of your module's 'Rules' object in order to find out about your module's requirements.  The CLR exception details may provide more information:  {2}", ModuleTypeName, AssemblyFileName, Ex.ToString());
+				throw new BuildException(Ex, "Unable to instantiate instance of '{0}' object type from compiled assembly '{1}'.  Unreal Build Tool creates an instance of your module's 'Rules' object in order to find out about your module's requirements.  The CLR exception details may provide more information:  {2}", ModuleTypeName, CompiledAssembly.FullName, Ex.ToString());
 			}
 
 			// Have to do absolute here as this could be a project that is under the root
@@ -1298,7 +1297,7 @@ namespace UnrealBuildTool
 				{
 					if (RulesObject.PrivateAssemblyReferences.Count > 0)
 					{
-						throw new BuildException("Module rules for '{0}' may not specify PrivateAssemblyReferences unless it is a CPlusPlusCLR module type.", AssemblyFileName);
+						throw new BuildException("Module rules for '{0}' may not specify PrivateAssemblyReferences unless it is a CPlusPlusCLR module type.", CompiledAssembly.FullName);
 					}
 
 					var InvalidDependencies = RulesObject.DynamicallyLoadedModuleNames.Intersect(RulesObject.PublicDependencyModuleNames.Concat(RulesObject.PrivateDependencyModuleNames)).ToList();
@@ -1344,10 +1343,18 @@ namespace UnrealBuildTool
 					}
 				}
 
-				// update the run-time dependencies path to remove $(PluginDir) and replace with $(ProjectDir) or $(EngineDir) equivalent
+				// Update the run-time dependencies path to remove $(PluginDir) and replace with a full path. When the receipt is saved it'll be converted to a $(ProjectDir) or $(EngineDir) equivalent.
 				foreach (var Dependency in RulesObject.RuntimeDependencies)
 				{
-					Dependency.ExpandPathVariables(RulesObject.ModuleDirectory);
+					const string PluginDirVariable = "$(PluginDir)";
+					if(Dependency.Path.StartsWith(PluginDirVariable, StringComparison.InvariantCultureIgnoreCase))
+					{
+						PluginInfo Plugin;
+						if(ModuleFileToPluginInfo.TryGetValue(ModuleFileName, out Plugin))
+						{
+							Dependency.Path = Plugin.Directory + Dependency.Path.Substring(PluginDirVariable.Length);
+						}
+					}
 				}
 			}
 
diff --git a/Engine/Source/Programs/UnrealBuildTool/System/TargetReceipt.cs b/Engine/Source/Programs/UnrealBuildTool/System/TargetReceipt.cs
--- a/Engine/Source/Programs/UnrealBuildTool/System/TargetReceipt.cs
+++ b/Engine/Source/Programs/UnrealBuildTool/System/TargetReceipt.cs
@@ -89,19 +89,6 @@ namespace UnrealBuildTool
 		{
 			return (StagePath == null)? Path : String.Format("{0} -> {1}", Path, StagePath);
 		}
-
-		public void ExpandPathVariables(string ModuleDir)
-		{
-			// expand the $(PluginDir) variable
-			string PluginDir = ModuleDir.Substring(0, ModuleDir.LastIndexOf(System.IO.Path.DirectorySeparatorChar + "Source" + System.IO.Path.DirectorySeparatorChar));
-			Dictionary<string, string> Variables = new Dictionary<string, string>();
-			Variables["PluginDir"] = System.IO.Path.GetFullPath(PluginDir).TrimEnd(System.IO.Path.DirectorySeparatorChar);
-			Path = Utils.ExpandVariables(Path, Variables);
-			if (StagePath != null)
-			{
-				StagePath = Utils.ExpandVariables(StagePath, Variables);
-			}
-		}
 	}
 
 	/// <summary>
