Make sure we don't get any trailing directory separators in file system reference objects. Path.GetFullPath() doesn't seem to be removing them.

diff --git a/Engine/Source/Programs/UnrealBuildTool/System/FileSystemReference.cs b/Engine/Source/Programs/UnrealBuildTool/System/FileSystemReference.cs
--- a/Engine/Source/Programs/UnrealBuildTool/System/FileSystemReference.cs
+++ b/Engine/Source/Programs/UnrealBuildTool/System/FileSystemReference.cs
@@ -28,7 +28,7 @@ namespace UnrealBuildTool
 		/// </summary>
 		public FileSystemReference(string InPath)
 		{
-			FullName = Path.GetFullPath(InPath);
+			FullName = Path.GetFullPath(InPath).TrimEnd(Path.DirectorySeparatorChar);
 			CanonicalName = FullName.ToLowerInvariant();
 		}
 
