Cache checks for whether a module has source files.

diff --git a/Engine/Source/Programs/UnrealBuildTool/Configuration/UEBuildTarget.cs b/Engine/Source/Programs/UnrealBuildTool/Configuration/UEBuildTarget.cs
--- a/Engine/Source/Programs/UnrealBuildTool/Configuration/UEBuildTarget.cs
+++ b/Engine/Source/Programs/UnrealBuildTool/Configuration/UEBuildTarget.cs
@@ -2623,8 +2623,7 @@ namespace UnrealBuildTool
 						UEBuildModule ModuleInstance = FindOrCreateModuleByName(Module.Name);
 
 						// Add the corresponding binary for it
-						FileReference ModuleFileName = RulesAssembly.GetModuleFileName(Module.Name);
-						bool bHasSource = (ModuleFileName != null && Directory.EnumerateFiles(ModuleFileName.Directory.FullName, "*.cpp", SearchOption.AllDirectories).Any());
+						bool bHasSource = RulesAssembly.DoesModuleHaveSource(Module.Name);
 						AddBinaryForModule(ModuleInstance, BinaryType, bAllowCompilation: bHasSource, bIsCrossTarget: false);
 
 						// Add it to the binary if we're compiling monolithic (and it's enabled)
diff --git a/Engine/Source/Programs/UnrealBuildTool/System/RulesCompiler.cs b/Engine/Source/Programs/UnrealBuildTool/System/RulesCompiler.cs
--- a/Engine/Source/Programs/UnrealBuildTool/System/RulesCompiler.cs
+++ b/Engine/Source/Programs/UnrealBuildTool/System/RulesCompiler.cs
@@ -1102,6 +1102,11 @@ namespace UnrealBuildTool
 		/// </summary>
 		private Dictionary<FileReference, PluginInfo> ModuleFileToPluginInfo;
 
+		/// <summary>
+		/// Cache for whether a module has source code
+		/// </summary>
+		private Dictionary<FileReference, bool> ModuleHasSource = new Dictionary<FileReference, bool>();
+
 		/// <summary>
 		/// Constructor. Compiles a rules assembly from the given source files.
 		/// </summary>
@@ -1584,6 +1589,31 @@ namespace UnrealBuildTool
 				return (Parent == null)? false : Parent.TryGetPluginForModule(ModuleFile, out Plugin);
 			}
 		}
+
+		/// <summary>
+		/// Determines if a module in this rules assembly has source code.
+		/// </summary>
+		/// <param name="ModuleName">Name of the module to check</param>
+		/// <returns>True if the module has source files, false if the module was not found, or does not have source files.</returns>
+		public bool DoesModuleHaveSource(string ModuleName)
+		{
+			FileReference ModuleFile;
+			if(ModuleNameToModuleFile.TryGetValue(ModuleName, out ModuleFile))
+			{
+				bool HasSource;
+				if(!ModuleHasSource.TryGetValue(ModuleFile, out HasSource))
+				{
+					foreach(string FileName in Directory.EnumerateFiles(ModuleFile.Directory.FullName, "*.cpp", SearchOption.AllDirectories))
+					{
+						HasSource = true;
+						break;
+					}
+					ModuleHasSource.Add(ModuleFile, HasSource);
+				}
+				return HasSource;
+			}
+			return (Parent == null)? false : Parent.DoesModuleHaveSource(ModuleName);
+		}
 	}
 
 	public class RulesCompiler
