Only add shader format and target platform modules as dependencies of the TargetPlatform module. Adding them for every single module in a target has a huge effect on the number of modules we need to resolve (UEBuildTarget.FindOrCreateModuleByName() was being called 200,000 times when generating project files, thanks to SCW and UFE triggering these code paths).

diff --git a/Engine/Source/Programs/UnrealBuildTool/Android/UEBuildAndroid.cs b/Engine/Source/Programs/UnrealBuildTool/Android/UEBuildAndroid.cs
--- a/Engine/Source/Programs/UnrealBuildTool/Android/UEBuildAndroid.cs
+++ b/Engine/Source/Programs/UnrealBuildTool/Android/UEBuildAndroid.cs
@@ -331,20 +331,23 @@ namespace UnrealBuildTool
 				}
 
 				// allow standalone tools to use targetplatform modules, without needing Engine
-				if (UEBuildConfiguration.bForceBuildTargetPlatforms)
+				if(ModuleName == "TargetPlatform")
 				{
-					Rules.PlatformSpecificDynamicallyLoadedModuleNames.Add("AndroidTargetPlatform");
-					Rules.PlatformSpecificDynamicallyLoadedModuleNames.Add("Android_PVRTCTargetPlatform");
-					Rules.PlatformSpecificDynamicallyLoadedModuleNames.Add("Android_ATCTargetPlatform");
-					Rules.PlatformSpecificDynamicallyLoadedModuleNames.Add("Android_DXTTargetPlatform");
-					Rules.PlatformSpecificDynamicallyLoadedModuleNames.Add("Android_ETC1TargetPlatform");
-					Rules.PlatformSpecificDynamicallyLoadedModuleNames.Add("Android_ETC2TargetPlatform");
-					Rules.PlatformSpecificDynamicallyLoadedModuleNames.Add("Android_ASTCTargetPlatform");
-                }
-
-				if (bBuildShaderFormats)
-				{
-					//Rules.DynamicallyLoadedModuleNames.Add("ShaderFormatAndroid");		//@todo android: ShaderFormatAndroid
+				    if (UEBuildConfiguration.bForceBuildTargetPlatforms)
+				    {
+					    Rules.PlatformSpecificDynamicallyLoadedModuleNames.Add("AndroidTargetPlatform");
+					    Rules.PlatformSpecificDynamicallyLoadedModuleNames.Add("Android_PVRTCTargetPlatform");
+					    Rules.PlatformSpecificDynamicallyLoadedModuleNames.Add("Android_ATCTargetPlatform");
+					    Rules.PlatformSpecificDynamicallyLoadedModuleNames.Add("Android_DXTTargetPlatform");
+					    Rules.PlatformSpecificDynamicallyLoadedModuleNames.Add("Android_ETC1TargetPlatform");
+					    Rules.PlatformSpecificDynamicallyLoadedModuleNames.Add("Android_ETC2TargetPlatform");
+					    Rules.PlatformSpecificDynamicallyLoadedModuleNames.Add("Android_ASTCTargetPlatform");
+                    }
+    
+				    if (bBuildShaderFormats)
+				    {
+					    //Rules.DynamicallyLoadedModuleNames.Add("ShaderFormatAndroid");		//@todo android: ShaderFormatAndroid
+				    }
 				}
 			}
 		}
diff --git a/Engine/Source/Programs/UnrealBuildTool/HTML5/UEBuildHTML5.cs b/Engine/Source/Programs/UnrealBuildTool/HTML5/UEBuildHTML5.cs
--- a/Engine/Source/Programs/UnrealBuildTool/HTML5/UEBuildHTML5.cs
+++ b/Engine/Source/Programs/UnrealBuildTool/HTML5/UEBuildHTML5.cs
@@ -316,7 +316,7 @@ namespace UnrealBuildTool
 				if ((!UEBuildConfiguration.bBuildRequiresCookedData
                     && ModuleName == "Engine"
                     && UEBuildConfiguration.bBuildDeveloperTools)
-                    || UEBuildConfiguration.bForceBuildTargetPlatforms)
+                    || (UEBuildConfiguration.bForceBuildTargetPlatforms && ModuleName == "TargetPlatform"))
                 {
                     Rules.PlatformSpecificDynamicallyLoadedModuleNames.Add("HTML5TargetPlatform");
                 }
diff --git a/Engine/Source/Programs/UnrealBuildTool/IOS/UEBuildIOS.cs b/Engine/Source/Programs/UnrealBuildTool/IOS/UEBuildIOS.cs
--- a/Engine/Source/Programs/UnrealBuildTool/IOS/UEBuildIOS.cs
+++ b/Engine/Source/Programs/UnrealBuildTool/IOS/UEBuildIOS.cs
@@ -442,15 +442,18 @@ namespace UnrealBuildTool
                 }
 
                 // allow standalone tools to use targetplatform modules, without needing Engine
-                if (UEBuildConfiguration.bForceBuildTargetPlatforms)
-                {
-                    Rules.PlatformSpecificDynamicallyLoadedModuleNames.Add("IOSTargetPlatform");
-                }
-
-                if (bBuildShaderFormats)
-                {
-                    Rules.PlatformSpecificDynamicallyLoadedModuleNames.Add("MetalShaderFormat");
-                }
+				if(ModuleName == "TargetPlatform")
+				{
+                    if (UEBuildConfiguration.bForceBuildTargetPlatforms)
+                    {
+                        Rules.PlatformSpecificDynamicallyLoadedModuleNames.Add("IOSTargetPlatform");
+                    }
+    
+                    if (bBuildShaderFormats)
+                    {
+                        Rules.PlatformSpecificDynamicallyLoadedModuleNames.Add("MetalShaderFormat");
+                    }
+				}
             }
         }
     }
diff --git a/Engine/Source/Programs/UnrealBuildTool/Linux/UEBuildLinux.cs b/Engine/Source/Programs/UnrealBuildTool/Linux/UEBuildLinux.cs
--- a/Engine/Source/Programs/UnrealBuildTool/Linux/UEBuildLinux.cs
+++ b/Engine/Source/Programs/UnrealBuildTool/Linux/UEBuildLinux.cs
@@ -303,7 +303,7 @@ namespace UnrealBuildTool
                 }
 
                 // allow standalone tools to use targetplatform modules, without needing Engine
-                if (UEBuildConfiguration.bForceBuildTargetPlatforms)
+                if (UEBuildConfiguration.bForceBuildTargetPlatforms && ModuleName == "TargetPlatform")
                 {
                     Rules.PlatformSpecificDynamicallyLoadedModuleNames.Add("LinuxTargetPlatform");
                     Rules.PlatformSpecificDynamicallyLoadedModuleNames.Add("LinuxNoEditorTargetPlatform");
@@ -323,19 +323,22 @@ namespace UnrealBuildTool
                 }
 
                 // allow standalone tools to use target platform modules, without needing Engine
-                if (UEBuildConfiguration.bForceBuildTargetPlatforms)
-                {
-                    Rules.DynamicallyLoadedModuleNames.Add("LinuxTargetPlatform");
-                    Rules.DynamicallyLoadedModuleNames.Add("LinuxNoEditorTargetPlatform");
-                    Rules.DynamicallyLoadedModuleNames.Add("LinuxServerTargetPlatform");
-					Rules.DynamicallyLoadedModuleNames.Add("AllDesktopTargetPlatform");
-                }
-
-                if (bBuildShaderFormats)
-                {
-					// Rules.DynamicallyLoadedModuleNames.Add("ShaderFormatD3D");
-                    Rules.DynamicallyLoadedModuleNames.Add("ShaderFormatOpenGL");
-                }
+				if(ModuleName == "TargetPlatform")
+				{
+                    if (UEBuildConfiguration.bForceBuildTargetPlatforms)
+                    {
+                        Rules.DynamicallyLoadedModuleNames.Add("LinuxTargetPlatform");
+                        Rules.DynamicallyLoadedModuleNames.Add("LinuxNoEditorTargetPlatform");
+                        Rules.DynamicallyLoadedModuleNames.Add("LinuxServerTargetPlatform");
+					    Rules.DynamicallyLoadedModuleNames.Add("AllDesktopTargetPlatform");
+                    }
+    
+                    if (bBuildShaderFormats)
+                    {
+					    // Rules.DynamicallyLoadedModuleNames.Add("ShaderFormatD3D");
+                        Rules.DynamicallyLoadedModuleNames.Add("ShaderFormatOpenGL");
+                    }
+				}
             }
         }
 
diff --git a/Engine/Source/Programs/UnrealBuildTool/Mac/UEBuildMac.cs b/Engine/Source/Programs/UnrealBuildTool/Mac/UEBuildMac.cs
--- a/Engine/Source/Programs/UnrealBuildTool/Mac/UEBuildMac.cs
+++ b/Engine/Source/Programs/UnrealBuildTool/Mac/UEBuildMac.cs
@@ -107,20 +107,23 @@ namespace UnrealBuildTool
                 }
 
 				// allow standalone tools to use target platform modules, without needing Engine
-				if (UEBuildConfiguration.bForceBuildTargetPlatforms)
+				if(ModuleName == "TargetPlatform")
 				{
-					Rules.DynamicallyLoadedModuleNames.Add("MacTargetPlatform");
-					Rules.DynamicallyLoadedModuleNames.Add("MacNoEditorTargetPlatform");
-					Rules.DynamicallyLoadedModuleNames.Add("MacClientTargetPlatform");
-					Rules.DynamicallyLoadedModuleNames.Add("MacServerTargetPlatform");
-					Rules.DynamicallyLoadedModuleNames.Add("AllDesktopTargetPlatform");
+				    if (UEBuildConfiguration.bForceBuildTargetPlatforms)
+				    {
+					    Rules.DynamicallyLoadedModuleNames.Add("MacTargetPlatform");
+					    Rules.DynamicallyLoadedModuleNames.Add("MacNoEditorTargetPlatform");
+					    Rules.DynamicallyLoadedModuleNames.Add("MacClientTargetPlatform");
+					    Rules.DynamicallyLoadedModuleNames.Add("MacServerTargetPlatform");
+					    Rules.DynamicallyLoadedModuleNames.Add("AllDesktopTargetPlatform");
+				    }
+    
+                    if (bBuildShaderFormats)
+                    {
+					    // Rules.DynamicallyLoadedModuleNames.Add("ShaderFormatD3D");
+                        Rules.DynamicallyLoadedModuleNames.Add("ShaderFormatOpenGL");
+                    }
 				}
-
-                if (bBuildShaderFormats)
-                {
-					// Rules.DynamicallyLoadedModuleNames.Add("ShaderFormatD3D");
-                    Rules.DynamicallyLoadedModuleNames.Add("ShaderFormatOpenGL");
-                }
 			}
 		}
 		
diff --git a/Engine/Source/Programs/UnrealBuildTool/Windows/UEBuildWindows.cs b/Engine/Source/Programs/UnrealBuildTool/Windows/UEBuildWindows.cs
--- a/Engine/Source/Programs/UnrealBuildTool/Windows/UEBuildWindows.cs
+++ b/Engine/Source/Programs/UnrealBuildTool/Windows/UEBuildWindows.cs
@@ -480,9 +480,6 @@ namespace UnrealBuildTool
             UEBuildConfiguration.bCompileICU = true;
         }
 
-//#todo-rco: Remove when public
-		public static string VulkanSDKPath = Environment.GetEnvironmentVariable("VK_SDK_PATH");
-
 		public override void ModifyModuleRules(string ModuleName, ModuleRules Rules, TargetInfo Target)
         {
             if ((Target.Platform == UnrealTargetPlatform.Win32) || (Target.Platform == UnrealTargetPlatform.Win64))
@@ -498,29 +495,33 @@ namespace UnrealBuildTool
                 }
 
                 // allow standalone tools to use target platform modules, without needing Engine
-                if (UEBuildConfiguration.bForceBuildTargetPlatforms)
-                {
-                    Rules.DynamicallyLoadedModuleNames.Add("WindowsTargetPlatform");
-                    Rules.DynamicallyLoadedModuleNames.Add("WindowsNoEditorTargetPlatform");
-                    Rules.DynamicallyLoadedModuleNames.Add("WindowsServerTargetPlatform");
-                    Rules.DynamicallyLoadedModuleNames.Add("WindowsClientTargetPlatform");
-					Rules.DynamicallyLoadedModuleNames.Add("AllDesktopTargetPlatform");
-                }
-
-                if (bBuildShaderFormats)
-                {
-                    Rules.DynamicallyLoadedModuleNames.Add("ShaderFormatD3D");
-                    Rules.DynamicallyLoadedModuleNames.Add("ShaderFormatOpenGL");
-
-//#todo-rco: Remove when public
-					{
-						if (!String.IsNullOrEmpty(VulkanSDKPath))
-						{
-							Rules.DynamicallyLoadedModuleNames.Add("VulkanShaderFormat");
-						}
-					}
+				if(ModuleName == "TargetPlatform")
+				{
+                    if (UEBuildConfiguration.bForceBuildTargetPlatforms)
+                    {
+                        Rules.DynamicallyLoadedModuleNames.Add("WindowsTargetPlatform");
+                        Rules.DynamicallyLoadedModuleNames.Add("WindowsNoEditorTargetPlatform");
+                        Rules.DynamicallyLoadedModuleNames.Add("WindowsServerTargetPlatform");
+                        Rules.DynamicallyLoadedModuleNames.Add("WindowsClientTargetPlatform");
+					    Rules.DynamicallyLoadedModuleNames.Add("AllDesktopTargetPlatform");
+                    }
+    
+                    if (bBuildShaderFormats)
+                    {
+                        Rules.DynamicallyLoadedModuleNames.Add("ShaderFormatD3D");
+                        Rules.DynamicallyLoadedModuleNames.Add("ShaderFormatOpenGL");
+    
+    //#todo-rco: Remove when public
+						string VulkanSDKPath = Environment.GetEnvironmentVariable("VK_SDK_PATH");
+					    {
+						    if (!String.IsNullOrEmpty(VulkanSDKPath))
+						    {
+							    Rules.DynamicallyLoadedModuleNames.Add("VulkanShaderFormat");
+						    }
+					    }
+				    }
 				}
-
+				
                 if (ModuleName == "D3D11RHI")
                 {
                     // To enable platform specific D3D11 RHI Types
