[INTEGRATE] Use 32-bit Clang compiler on Windows if the 64-bit version isn't installed

diff --git a/Engine/Source/Programs/UnrealBuildTool/Windows/VCEnvironment.cs b/Engine/Source/Programs/UnrealBuildTool/Windows/VCEnvironment.cs
--- a/Engine/Source/Programs/UnrealBuildTool/Windows/VCEnvironment.cs
+++ b/Engine/Source/Programs/UnrealBuildTool/Windows/VCEnvironment.cs
@@ -304,12 +304,20 @@ namespace UnrealBuildTool
 				if( WindowsPlatform.bUseVCCompilerArgs )
 				{ 
 					// Use regular Clang compiler on Windows
-					Result = Path.Combine( Environment.GetFolderPath( Environment.SpecialFolder.ProgramFilesX86 ), "LLVM", "msbuild-bin", "cl.exe" );
+					Result = Path.Combine( Environment.GetFolderPath( Environment.SpecialFolder.ProgramFiles ), "LLVM", "msbuild-bin", "cl.exe" );
+					if( !File.Exists( Result ) )
+					{
+ 	 					Result = Path.Combine( Environment.GetFolderPath( Environment.SpecialFolder.ProgramFilesX86 ), "LLVM", "msbuild-bin", "cl.exe" );
+					}
 				}
 				else
 				{
 					// Use 'clang-cl', a wrapper around Clang that supports Visual C++ compiler command-line arguments
-					Result = Path.Combine( Environment.GetFolderPath( Environment.SpecialFolder.ProgramFilesX86 ), "LLVM", "bin", "clang.exe" );
+					Result = Path.Combine( Environment.GetFolderPath( Environment.SpecialFolder.ProgramFiles ), "LLVM", "bin", "clang.exe" );
+					if( !File.Exists( Result ) )
+					{
+						Result = Path.Combine( Environment.GetFolderPath( Environment.SpecialFolder.ProgramFilesX86 ), "LLVM", "bin", "clang.exe" );
+					}
 				}
 				if (!File.Exists(Result))
 				{
@@ -340,7 +348,11 @@ namespace UnrealBuildTool
 			// If we were asked to use Clang, then we'll redirect the path to the compiler to the LLVM installation directory
 			if( WindowsPlatform.bCompileWithClang && WindowsPlatform.bAllowClangLinker )
 			{
-				var Result = Path.Combine( Environment.GetFolderPath( Environment.SpecialFolder.ProgramFilesX86 ), "LLVM", "bin", "lld.exe" );
+				var Result = Path.Combine( Environment.GetFolderPath( Environment.SpecialFolder.ProgramFiles ), "LLVM", "bin", "lld.exe" );
+				if( !File.Exists( Result ) )
+				{
+					Result = Path.Combine( Environment.GetFolderPath( Environment.SpecialFolder.ProgramFilesX86 ), "LLVM", "bin", "lld.exe" );
+				}
 				if( !File.Exists( Result ) )
 				{
 					throw new BuildException( "Clang was selected as the Windows compiler, but LLVM/Clang does not appear to be installed.  Could not find: " + Result );
