[INTEGRATE] Unreal Build Tool: Output total number of actions to execute when building
- Also: Minor formatting fixes for verbose debug output with files being complied on Windows

diff --git a/Engine/Source/Programs/UnrealBuildTool/System/ActionGraph.cs b/Engine/Source/Programs/UnrealBuildTool/System/ActionGraph.cs
--- a/Engine/Source/Programs/UnrealBuildTool/System/ActionGraph.cs
+++ b/Engine/Source/Programs/UnrealBuildTool/System/ActionGraph.cs
@@ -324,6 +324,8 @@ namespace UnrealBuildTool
 			ExecutorName = "";
 			if (ActionsToExecute.Count > 0)
 			{
+				Log.TraceInformation("Executing {0} actions", ActionsToExecute.Count );
+
 				if (BuildConfiguration.bAllowXGE || BuildConfiguration.bXGEExport)
 				{
 					XGE.ExecutionResult XGEResult = XGE.ExecutionResult.TasksSucceeded;
@@ -425,7 +427,7 @@ namespace UnrealBuildTool
 			else
 			{
 				ExecutorName = "NoActionsToExecute";
-				Log.TraceInformation("Target is up to date.");
+				Log.TraceInformation("Target is up to date");
 			}
 
 			return Result;
diff --git a/Engine/Source/Programs/UnrealBuildTool/Windows/VCToolChain.cs b/Engine/Source/Programs/UnrealBuildTool/Windows/VCToolChain.cs
--- a/Engine/Source/Programs/UnrealBuildTool/Windows/VCToolChain.cs
+++ b/Engine/Source/Programs/UnrealBuildTool/Windows/VCToolChain.cs
@@ -1334,13 +1334,11 @@ namespace UnrealBuildTool
 
 				if( CompileEnvironment.Config.PrecompiledHeaderAction == PrecompiledHeaderAction.Create )
 				{
-					Log.TraceVerbose("Creating PCH: " + CompileEnvironment.Config.PrecompiledHeaderIncludeFilename);
-					Log.TraceVerbose("     Command: " + CompileAction.CommandArguments);
+					Log.TraceVerbose("Creating PCH " + CompileEnvironment.Config.PrecompiledHeaderIncludeFilename + ": \"" + CompileAction.CommandPath + "\"" + CompileAction.CommandArguments);
 				}
 				else
 				{
-					Log.TraceVerbose("   Compiling: " + CompileAction.StatusDescription);
-					Log.TraceVerbose("     Command: " + CompileAction.CommandArguments);
+					Log.TraceVerbose("   Compiling " + CompileAction.StatusDescription + ": \"" + CompileAction.CommandPath + "\"" + CompileAction.CommandArguments);
 				}
 
 				if( WindowsPlatform.bCompileWithClang || BuildConfiguration.bRunUnrealCodeAnalyzer)
