[INTEGRATE] Unreal Build Tool: Output total number of actions to execute when building (part 2)
- Changed output to be XGE-specific and cleaned up text (the Local Executor already outputs action counts)

diff --git a/Engine/Source/Programs/UnrealBuildTool/System/ActionGraph.cs b/Engine/Source/Programs/UnrealBuildTool/System/ActionGraph.cs
--- a/Engine/Source/Programs/UnrealBuildTool/System/ActionGraph.cs
+++ b/Engine/Source/Programs/UnrealBuildTool/System/ActionGraph.cs
@@ -324,10 +324,13 @@ namespace UnrealBuildTool
 			ExecutorName = "";
 			if (ActionsToExecute.Count > 0)
 			{
-				Log.TraceInformation("Executing {0} actions", ActionsToExecute.Count );
-
 				if (BuildConfiguration.bAllowXGE || BuildConfiguration.bXGEExport)
 				{
+					Log.TraceInformation("{0} {1} action{2} to XGE", 
+						BuildConfiguration.bXGEExport ? "Exporting" : "Distributing", 
+						ActionsToExecute.Count, 
+						ActionsToExecute.Count == 1 ? "" : "s" );
+
 					XGE.ExecutionResult XGEResult = XGE.ExecutionResult.TasksSucceeded;
 
 					// Batch up XGE execution by actions with the same output event handler.
