[INTEGRATE] Unreal Build Tool: Improved adaptive unity build
- When bUseAdaptiveUnityBuild is enabled (experimental!), your first build after your working file set changes is now much faster
- UBT no longer invalidates other existing compiled unity blobs from the same module when ejecting source files

diff --git a/Engine/Source/Programs/UnrealBuildTool/System/Unity.cs b/Engine/Source/Programs/UnrealBuildTool/System/Unity.cs
--- a/Engine/Source/Programs/UnrealBuildTool/System/Unity.cs
+++ b/Engine/Source/Programs/UnrealBuildTool/System/Unity.cs
@@ -24,16 +24,36 @@ namespace UnrealBuildTool
 			public List<FileItem> Files       { get; private set; }
 			public long           TotalLength { get; private set; }
 
+			/// The length of this file collection, plus any additional virtual space needed for bUseAdapativeUnityBuild.
+			/// See the comment above AddVirtualFile() below for more information.
+			public long           VirtualLength { get; private set; }
+
 			public FileCollection()
 			{
 				Files       = new List<FileItem>();
 				TotalLength = 0;
+				VirtualLength = 0;
 			}
 
 			public void AddFile(FileItem File)
 			{
 				Files.Add(File);
-				TotalLength += File.Info.Length;
+
+				long FileLength = File.Info.Length;
+				TotalLength += FileLength;
+				VirtualLength += FileLength;
+			}
+
+			/// <summary>
+			/// Doesn't actually add a file, but instead reserves space.  This is used with "bUseAdaptiveUnityBuild", to prevent
+			/// other compiled unity blobs in the module's numbered set from having to be recompiled after we eject source files
+			/// one of that module's unity blobs.  Basically, it can prevent dozens of files from being recompiled after the first
+			/// time building after your working set of source files changes
+			/// </summary>
+			/// <param name="VirtualFileLength">Length of the virtual file to add to this file collection</param>
+			public void AddVirtualFile(long VirtualFileLength)
+			{
+				VirtualLength += VirtualFileLength;
 			}
 		}
 
@@ -66,12 +86,27 @@ namespace UnrealBuildTool
 			public void AddFile(FileItem File)
 			{
 				CurrentCollection.AddFile(File);
-				if (SplitLength != -1 && CurrentCollection.TotalLength > SplitLength)
+				if (SplitLength != -1 && CurrentCollection.VirtualLength > SplitLength)
 				{
 					EndCurrentUnityFile();
 				}
 			}
 
+			/// <summary>
+			/// Doesn't actually add a file, but instead reserves space, then splits the unity blob normally as if it
+			/// was a real file that was added.  See the comment above FileCollection.AddVirtualFile() for more info.
+			/// </summary>
+			/// <param name="File">The file to add virtually.  Only the size of the file is tracked.</param>
+			public void AddVirtualFile(FileItem File)
+			{
+				CurrentCollection.AddVirtualFile( File.Info.Length );
+				if (SplitLength != -1 && CurrentCollection.VirtualLength > SplitLength)
+				{
+					EndCurrentUnityFile();
+				}
+			}
+
+
 			/// <summary>
 			/// Starts a new unity file.  If the current unity file contains no files, this function has no effect, i.e. you will not get an empty unity file.
 			/// </summary>
@@ -204,6 +239,12 @@ namespace UnrealBuildTool
 						// Just compile this file normally, not as part of the unity blob
 						NewCPPFiles.Add( CPPFile );
 
+						// Let the unity file builder know about the file, so that we can retain the existing size of the unity blobs.
+						// This won't actually make the source file part of the unity blob, but it will keep track of how big the
+						// file is so that other existing unity blobs from the same module won't be invalidated.  This prevents much
+						// longer compile times the first time you build after your working file set changes.
+						CPPUnityFileBuilder.AddVirtualFile( CPPFile );
+
 						var CPPFileName = Path.GetFileName( CPPFile.AbsolutePath );
 						if( AdaptiveUnityBuildInfoString.Length == 0 )
 						{
