UBT Telemetry re-enabled. Commented out most of the legacy events.

diff --git a/Engine/Source/Programs/UnrealBuildTool/System/DependencyCache.cs b/Engine/Source/Programs/UnrealBuildTool/System/DependencyCache.cs
--- a/Engine/Source/Programs/UnrealBuildTool/System/DependencyCache.cs
+++ b/Engine/Source/Programs/UnrealBuildTool/System/DependencyCache.cs
@@ -95,8 +95,7 @@ namespace UnrealBuildTool
 					{
 						Log.TraceInformation("Loading IncludeFileCache took " + TimerDuration.TotalSeconds + "s");
 					}
-					Telemetry.SendEvent("LoadIncludeDependencyCacheStats.2",
-						"TotalDuration", TimerDuration.TotalSeconds.ToString("0.00"));
+					//Telemetry.SendEvent("LoadIncludeDependencyCacheStats.2", "TotalDuration", TimerDuration.TotalSeconds.ToString("0.00"));
 					return Result;
 				}
 			}
diff --git a/Engine/Source/Programs/UnrealBuildTool/System/LocalExecutor.cs b/Engine/Source/Programs/UnrealBuildTool/System/LocalExecutor.cs
--- a/Engine/Source/Programs/UnrealBuildTool/System/LocalExecutor.cs
+++ b/Engine/Source/Programs/UnrealBuildTool/System/LocalExecutor.cs
@@ -282,6 +282,7 @@ namespace UnrealBuildTool
 			// track how long it took
 			Action.EndTime = DateTimeOffset.Now;
 
+/*
 			// send telemetry
 			{
 				// See if the action produced a PCH file (infer from the extension as we no longer have the compile environments used to generate the actions.
@@ -299,6 +300,7 @@ namespace UnrealBuildTool
 						"Duration", Action.Duration.TotalSeconds.ToString("0.00"));
 				}
 			}
+*/
 
 			if (!Utils.IsRunningOnMono)
 			{
diff --git a/Engine/Source/Programs/UnrealBuildTool/System/Telemetry.cs b/Engine/Source/Programs/UnrealBuildTool/System/Telemetry.cs
--- a/Engine/Source/Programs/UnrealBuildTool/System/Telemetry.cs
+++ b/Engine/Source/Programs/UnrealBuildTool/System/Telemetry.cs
@@ -2,6 +2,7 @@
 
 using System;
 using System.Collections.Generic;
+using System.IO;
 using System.Linq;
 using System.Reflection;
 using System.Text;
@@ -36,6 +37,10 @@ namespace UnrealBuildTool
 		[AttributeUsage(AttributeTargets.Class, AllowMultiple = false)]
 		public class ProviderAttribute : Attribute
 		{
+            /// <summary>
+            /// If true, this provider is only created for Epic internal builds.
+            /// </summary>
+            public bool bEpicInternalOnly;
 		}
 
 		/// <summary>
@@ -59,21 +64,26 @@ namespace UnrealBuildTool
 		public static void Initialize()
 		{
 			var TelemetryInitStartTime = DateTime.UtcNow;
+            
+            // only collect telemetry on Epic builds
+            bool bIsEpicInternal = File.Exists(@"..\Build\NotForLicensees\EpicInternal.txt");
 
-			// reflect over the assembly and find the first eligible provider class.
-			Provider = 
-				(from reflectionAssembly in AppDomain.CurrentDomain.GetReferencingAssemblies(Assembly.GetExecutingAssembly())
-				from T in reflectionAssembly.GetTypes()
-				// type must implement interface, contain the attribute, be concrete and default constructible.
-				where T.IsDefined(typeof(ProviderAttribute), true) && T.GetInterfaces().Contains(typeof(IProvider)) && T.IsClass && !T.IsAbstract && T.GetConstructor(Type.EmptyTypes) != null
-				// create the instance
-				select (IProvider)Activator.CreateInstance(T)).FirstOrDefault();
-			
-			if( BuildConfiguration.bPrintPerformanceInfo )
-			{ 
-				var TelemetryInitTime = (DateTime.UtcNow - TelemetryInitStartTime).TotalSeconds;
-				Log.TraceInformation( "Telemetry initialization took " + TelemetryInitTime + "s" );
-			}
+            // reflect over the assembly and find the first eligible provider class.
+            Provider =
+                (from reflectionAssembly in AppDomain.CurrentDomain.GetReferencingAssemblies(Assembly.GetExecutingAssembly())
+                    from T in reflectionAssembly.GetTypes()
+                    // type must implement interface, contain the attribute, be concrete and default constructible.
+                    where T.IsDefined(typeof(ProviderAttribute), true) && T.GetInterfaces().Contains(typeof(IProvider)) && T.IsClass && !T.IsAbstract && T.GetConstructor(Type.EmptyTypes) != null
+                    // if it's an epic internal build, it's ok to use epic internal providers.
+                    where bIsEpicInternal || !T.GetCustomAttribute<ProviderAttribute>().bEpicInternalOnly
+                    // create the instance
+                    select (IProvider)Activator.CreateInstance(T)).FirstOrDefault();
+
+            if (BuildConfiguration.bPrintPerformanceInfo)
+            {
+                var TelemetryInitTime = (DateTime.UtcNow - TelemetryInitStartTime).TotalSeconds;
+                Log.TraceInformation("Telemetry initialization took " + TelemetryInitTime + "s");
+            }
 		}
 
 		/// <summary>
diff --git a/Engine/Source/Programs/UnrealBuildTool/System/UnrealBuildTool.cs b/Engine/Source/Programs/UnrealBuildTool/System/UnrealBuildTool.cs
--- a/Engine/Source/Programs/UnrealBuildTool/System/UnrealBuildTool.cs
+++ b/Engine/Source/Programs/UnrealBuildTool/System/UnrealBuildTool.cs
@@ -1225,7 +1225,7 @@ namespace UnrealBuildTool
                         "ComputerName", Environment.MachineName,
                         "User", Environment.UserName,
                         "Domain", Environment.UserDomainName,
-                        "CommandLine", string.Join(" ", Arguments),
+                        "CommandLine", string.Join("|", Arguments),
                         "UBT Action", bGenerateVCProjectFiles 
                             ? "GenerateVCProjectFiles" 
                             : bGenerateXcodeProjectFiles 
@@ -1247,7 +1247,8 @@ namespace UnrealBuildTool
                             : "Build",
                         "Platform", CheckPlatform.ToString(),
                         "Configuration", CheckConfiguration.ToString(),
-                        "EngineVersion", (Version == null)? "0" : Version.Changelist.ToString()
+                        "EngineVersion", (Version == null)? "0" : Version.Changelist.ToString(),
+                        "Branch", (Version == null)? "" : Version.BranchName
                         );
 
 
@@ -1394,17 +1395,17 @@ namespace UnrealBuildTool
                         Log.TraceInformation("Execution time: {0}s", BuildDuration);
                     }
 
-                    Telemetry.SendEvent("PerformanceInfo.2",
-                        "TotalExecutionTimeSec", BuildDuration.ToString("0.00"),
-                        "TotalTimeSpentGettingIncludesSec", CPPEnvironment.TotalTimeSpentGettingIncludes.ToString("0.00"),
-                        "TotalIncludesRequested", CPPEnvironment.TotalIncludesRequested.ToString(),
-                        "DirectIncludeCacheMissesTotalTimeSec", CPPEnvironment.DirectIncludeCacheMissesTotalTime.ToString("0.00"),
-                        "TotalDirectIncludeCacheMisses", CPPEnvironment.TotalDirectIncludeCacheMisses.ToString(),
-                        "TotalFindIncludedFileCalls", CPPEnvironment.TotalFindIncludedFileCalls.ToString(),
-                        "IncludePathSearchAttempts", CPPEnvironment.IncludePathSearchAttempts.ToString(),
-                        "TotalFileItemCount", FileItem.TotalFileItemCount.ToString(),
-                        "MissingFileItemCount", FileItem.MissingFileItemCount.ToString()
-                        );
+                    //Telemetry.SendEvent("PerformanceInfo.2",
+                    //    "TotalExecutionTimeSec", BuildDuration.ToString("0.00"),
+                    //    "TotalTimeSpentGettingIncludesSec", CPPEnvironment.TotalTimeSpentGettingIncludes.ToString("0.00"),
+                    //    "TotalIncludesRequested", CPPEnvironment.TotalIncludesRequested.ToString(),
+                    //    "DirectIncludeCacheMissesTotalTimeSec", CPPEnvironment.DirectIncludeCacheMissesTotalTime.ToString("0.00"),
+                    //    "TotalDirectIncludeCacheMisses", CPPEnvironment.TotalDirectIncludeCacheMisses.ToString(),
+                    //    "TotalFindIncludedFileCalls", CPPEnvironment.TotalFindIncludedFileCalls.ToString(),
+                    //    "IncludePathSearchAttempts", CPPEnvironment.IncludePathSearchAttempts.ToString(),
+                    //    "TotalFileItemCount", FileItem.TotalFileItemCount.ToString(),
+                    //    "MissingFileItemCount", FileItem.MissingFileItemCount.ToString()
+                    //    );
                 }
                 catch (Exception Exception)
                 {
@@ -2103,16 +2104,16 @@ namespace UnrealBuildTool
                     TotalLinkTime,
                     TotalOtherActionsTime
                 );
-                Telemetry.SendEvent("BuildStatsTotal.2",
-                    "ExecutorName", ExecutorName,
-                    "TotalUBTWallClockTimeSec", BuildDuration.ToString("0.00"),
-                    "TotalBuildProjectThreadTimeSec", TotalBuildProjectTime.ToString("0.00"),
-                    "TotalCompileThreadTimeSec", TotalCompileTime.ToString("0.00"),
-                    "TotalCreateAppBundleThreadTimeSec", TotalCreateAppBundleTime.ToString("0.00"),
-                    "TotalGenerateDebugInfoThreadTimeSec", TotalGenerateDebugInfoTime.ToString("0.00"),
-                    "TotalLinkThreadTimeSec", TotalLinkTime.ToString("0.00"),
-                    "TotalOtherActionsThreadTimeSec", TotalOtherActionsTime.ToString("0.00")
-                    );
+                //Telemetry.SendEvent("BuildStatsTotal.2",
+                //    "ExecutorName", ExecutorName,
+                //    "TotalUBTWallClockTimeSec", BuildDuration.ToString("0.00"),
+                //    "TotalBuildProjectThreadTimeSec", TotalBuildProjectTime.ToString("0.00"),
+                //    "TotalCompileThreadTimeSec", TotalCompileTime.ToString("0.00"),
+                //    "TotalCreateAppBundleThreadTimeSec", TotalCreateAppBundleTime.ToString("0.00"),
+                //    "TotalGenerateDebugInfoThreadTimeSec", TotalGenerateDebugInfoTime.ToString("0.00"),
+                //    "TotalLinkThreadTimeSec", TotalLinkTime.ToString("0.00"),
+                //    "TotalOtherActionsThreadTimeSec", TotalOtherActionsTime.ToString("0.00")
+                //    );
 
                 Log.TraceInformation("Total build time: {0:0.00} seconds", BuildDuration);
 
@@ -2130,10 +2131,10 @@ namespace UnrealBuildTool
                 {
                     Log.TraceInformation("XGE execution time: {0:0.00} seconds", BuildDuration);
                 }
-                Telemetry.SendEvent("BuildStatsTotal.2",
-                    "ExecutorName", ExecutorName,
-                    "TotalUBTWallClockTimeSec", BuildDuration.ToString("0.00")
-                    );
+                //Telemetry.SendEvent("BuildStatsTotal.2",
+                //    "ExecutorName", ExecutorName,
+                //    "TotalUBTWallClockTimeSec", BuildDuration.ToString("0.00")
+                //    );
             }
 
             return BuildResult;
diff --git a/Engine/Source/Programs/UnrealBuildTool/System/XGE.cs b/Engine/Source/Programs/UnrealBuildTool/System/XGE.cs
--- a/Engine/Source/Programs/UnrealBuildTool/System/XGE.cs
+++ b/Engine/Source/Programs/UnrealBuildTool/System/XGE.cs
@@ -60,6 +60,7 @@ namespace UnrealBuildTool
 								// sometimes the args comes in as null
 								var match = XGEDurationRegex.Match(args.Data ?? "");
 
+/*
 								// This is considered fragile and risky parsing of undocumented XGE output, so protect this code from taking down UBT
 								try
 								{
@@ -109,6 +110,7 @@ namespace UnrealBuildTool
 										"OutputLine", args.Data,
 										"Exception", ex.ToString());
 								}
+*/
 
 								// XGE outputs the duration info in a format that makes VC think it's an error file/line notation if the full filename is used.
 								// So munge it a bit so we don't confuse VC.
