Speed up reading and accessing ini files from UAT/UBT
#uat
#ubt

diff --git a/Engine/Source/Programs/AutomationTool/AutomationUtils/ProjectParams.cs b/Engine/Source/Programs/AutomationTool/AutomationUtils/ProjectParams.cs
--- a/Engine/Source/Programs/AutomationTool/AutomationUtils/ProjectParams.cs
+++ b/Engine/Source/Programs/AutomationTool/AutomationUtils/ProjectParams.cs
@@ -1727,7 +1727,7 @@ namespace AutomationTool
 			// the full solution for per-platform packaging settings.
 			if (!Manifests)
 			{				
-				ConfigCacheIni GameIni = new ConfigCacheIni("Game", Path.GetDirectoryName(RawProjectPath));
+				ConfigCacheIni GameIni = ConfigCacheIni.CreateConfigCacheIni(UnrealTargetPlatform.Unknown, "Game", new DirectoryReference(Path.GetDirectoryName(RawProjectPath)));
 				String IniPath = "/Script/UnrealEd.ProjectPackagingSettings";
 				bool bSetting = false;
 				if (!GameIni.GetBool(IniPath, "bGenerateChunks", out bSetting))
diff --git a/Engine/Source/Programs/AutomationTool/AutomationUtils/ProjectUtils.cs b/Engine/Source/Programs/AutomationTool/AutomationUtils/ProjectUtils.cs
--- a/Engine/Source/Programs/AutomationTool/AutomationUtils/ProjectUtils.cs
+++ b/Engine/Source/Programs/AutomationTool/AutomationUtils/ProjectUtils.cs
@@ -325,7 +325,7 @@ namespace AutomationTool
 				{
 					if (TargetPlatformType != UnrealTargetPlatform.Unknown)
 					{
-						var Config = new ConfigCacheIni(TargetPlatformType, "Engine", Path.GetDirectoryName(RawProjectPath), EngineDirectory);
+						var Config = ConfigCacheIni.CreateConfigCacheIni(TargetPlatformType, "Engine", new DirectoryReference(Path.GetDirectoryName(RawProjectPath)), new DirectoryReference(EngineDirectory));
 						Properties.EngineConfigs.Add(TargetPlatformType, Config);
 					}
 				}
@@ -334,7 +334,7 @@ namespace AutomationTool
 				{
 					if (TargetPlatformType != UnrealTargetPlatform.Unknown)
 					{
-						var Config = new ConfigCacheIni(TargetPlatformType, "Game", Path.GetDirectoryName(RawProjectPath));
+						var Config = ConfigCacheIni.CreateConfigCacheIni(TargetPlatformType, "Game", new DirectoryReference(Path.GetDirectoryName(RawProjectPath)));
 						Properties.GameConfigs.Add(TargetPlatformType, Config);
 					}
 				}
diff --git a/Engine/Source/Programs/AutomationTool/GUBP/LegacyBranchSetup.cs b/Engine/Source/Programs/AutomationTool/GUBP/LegacyBranchSetup.cs
--- a/Engine/Source/Programs/AutomationTool/GUBP/LegacyBranchSetup.cs
+++ b/Engine/Source/Programs/AutomationTool/GUBP/LegacyBranchSetup.cs
@@ -357,6 +357,8 @@ partial class GUBP
 
     void AddNodesForBranch(List<UnrealTargetPlatform> InitialHostPlatforms, JobInfo JobInfo, GUBPBranchHacker.BranchOptions BranchOptions, out List<BuildNodeTemplate> AllNodeDefinitions, out List<AggregateNodeTemplate> AllAggregateDefinitions, ref int TimeQuantum)
 	{
+		DateTime StartTime = DateTime.UtcNow;
+
 		if (BranchOptions.QuantumOverride != 0)
 		{
 			TimeQuantum = BranchOptions.QuantumOverride;
@@ -1200,6 +1202,9 @@ partial class GUBP
 
 		// Get the email list for each node
 		FindEmailsForNodes(BranchConfig.BranchName, AllNodeDefinitions);
+
+		TimeSpan Span = DateTime.UtcNow - StartTime;
+		Log("Time Spent: " + Span.TotalSeconds);
 	}
 
 	private static Dictionary<string, int> ApplyFrequencyBarriers(Dictionary<string, GUBPNode> GUBPNodes, Dictionary<string, GUBPAggregateNode> GUBPAggregates, Dictionary<string, sbyte> FrequencyBarriers)
diff --git a/Engine/Source/Programs/AutomationTool/HTML5/HTML5Platform.Automation.cs b/Engine/Source/Programs/AutomationTool/HTML5/HTML5Platform.Automation.cs
--- a/Engine/Source/Programs/AutomationTool/HTML5/HTML5Platform.Automation.cs
+++ b/Engine/Source/Programs/AutomationTool/HTML5/HTML5Platform.Automation.cs
@@ -523,7 +523,7 @@ public class HTML5Platform : Platform
 #region AMAZON S3 
 	public void UploadToS3(DeploymentContext SC)
 	{
-		ConfigCacheIni Ini = new ConfigCacheIni(SC.StageTargetPlatform.PlatformType, "Engine", Path.GetDirectoryName(SC.RawProjectPath));
+		ConfigCacheIni Ini = ConfigCacheIni.CreateConfigCacheIni(SC.StageTargetPlatform.PlatformType, "Engine", new DirectoryReference(Path.GetDirectoryName(SC.RawProjectPath)));
 		bool Upload = false;
 
 		string KeyId = "";
diff --git a/Engine/Source/Programs/AutomationTool/Scripts/CopyBuildToStagingDirectory.Automation.cs b/Engine/Source/Programs/AutomationTool/Scripts/CopyBuildToStagingDirectory.Automation.cs
--- a/Engine/Source/Programs/AutomationTool/Scripts/CopyBuildToStagingDirectory.Automation.cs
+++ b/Engine/Source/Programs/AutomationTool/Scripts/CopyBuildToStagingDirectory.Automation.cs
@@ -308,7 +308,7 @@ public partial class Project : CommandUtils
 		}
 		SC.StageFiles(StagedFileType.NonUFS, GetIntermediateCommandlineDir(SC), CommandLineFile, false, null, "", true, false);
 
-        ConfigCacheIni PlatformGameConfig = new ConfigCacheIni(SC.StageTargetPlatform.PlatformType, "Game", CommandUtils.GetDirectoryName(Params.RawProjectPath));
+        ConfigCacheIni PlatformGameConfig = ConfigCacheIni.CreateConfigCacheIni(SC.StageTargetPlatform.PlatformType, "Game", new DirectoryReference(CommandUtils.GetDirectoryName(Params.RawProjectPath)));
         var ProjectContentRoot = CombinePaths(SC.ProjectRoot, "Content");
         var StageContentRoot = CombinePaths(SC.RelativeProjectRootForStage, "Content");
         
diff --git a/Engine/Source/Programs/UnrealBuildTool/Android/AndroidToolChain.cs b/Engine/Source/Programs/UnrealBuildTool/Android/AndroidToolChain.cs
--- a/Engine/Source/Programs/UnrealBuildTool/Android/AndroidToolChain.cs
+++ b/Engine/Source/Programs/UnrealBuildTool/Android/AndroidToolChain.cs
@@ -46,7 +46,7 @@ namespace UnrealBuildTool
 		public static void ParseArchitectures()
 		{
 			// look in ini settings for what platforms to compile for
-			ConfigCacheIni Ini = new ConfigCacheIni(UnrealTargetPlatform.Android, "Engine", UnrealBuildTool.GetUProjectPath());
+			ConfigCacheIni Ini = ConfigCacheIni.CreateConfigCacheIni(UnrealTargetPlatform.Android, "Engine", UnrealBuildTool.GetUProjectPath());
             Arches = new List<string>();
 			bool bBuild = true;
 			if (Ini.GetBool("/Script/AndroidRuntimeSettings.AndroidRuntimeSettings", "bBuildForArmV7", out bBuild) && bBuild)
@@ -162,7 +162,7 @@ namespace UnrealBuildTool
 		static public string GetNdkApiLevel()
 		{
 			// ask the .ini system for what version to use
-			ConfigCacheIni Ini = new ConfigCacheIni(UnrealTargetPlatform.Android, "Engine", UnrealBuildTool.GetUProjectPath());
+			ConfigCacheIni Ini = ConfigCacheIni.CreateConfigCacheIni(UnrealTargetPlatform.Android, "Engine", UnrealBuildTool.GetUProjectPath());
 			string NDKLevel;
 			Ini.GetString("/Script/AndroidPlatformEditor.AndroidSDKSettings", "NDKAPILevel", out NDKLevel);
 
diff --git a/Engine/Source/Programs/UnrealBuildTool/Android/UEBuildAndroid.cs b/Engine/Source/Programs/UnrealBuildTool/Android/UEBuildAndroid.cs
--- a/Engine/Source/Programs/UnrealBuildTool/Android/UEBuildAndroid.cs
+++ b/Engine/Source/Programs/UnrealBuildTool/Android/UEBuildAndroid.cs
@@ -58,7 +58,7 @@ namespace UnrealBuildTool
         {
             string NDKPath = Environment.GetEnvironmentVariable("NDKROOT");
             {
-                var configCacheIni = new ConfigCacheIni("Engine", (DirectoryReference)null);
+                var configCacheIni = ConfigCacheIni.CreateConfigCacheIni(UnrealTargetPlatform.Unknown, "Engine", (DirectoryReference)null);
                 var AndroidEnv = new Dictionary<string, string>();
 
                 Dictionary<string, string> EnvVarNames = new Dictionary<string,string> { 
diff --git a/Engine/Source/Programs/UnrealBuildTool/Android/UEDeployAndroid.cs b/Engine/Source/Programs/UnrealBuildTool/Android/UEDeployAndroid.cs
--- a/Engine/Source/Programs/UnrealBuildTool/Android/UEDeployAndroid.cs
+++ b/Engine/Source/Programs/UnrealBuildTool/Android/UEDeployAndroid.cs
@@ -66,7 +66,7 @@ namespace UnrealBuildTool.Android
             ConfigCacheIni config = null;
             if(!ConfigCache.TryGetValue(baseIniName, out config))
             {
-                config = new ConfigCacheIni(UnrealTargetPlatform.Android, "Engine", UnrealBuildTool.GetUProjectPath());
+                config = ConfigCacheIni.CreateConfigCacheIni(UnrealTargetPlatform.Android, "Engine", UnrealBuildTool.GetUProjectPath());
                 ConfigCache.Add(baseIniName, config);
             }
 
diff --git a/Engine/Source/Programs/UnrealBuildTool/Configuration/EngineConfiguration.cs b/Engine/Source/Programs/UnrealBuildTool/Configuration/EngineConfiguration.cs
--- a/Engine/Source/Programs/UnrealBuildTool/Configuration/EngineConfiguration.cs
+++ b/Engine/Source/Programs/UnrealBuildTool/Configuration/EngineConfiguration.cs
@@ -26,6 +26,57 @@ namespace UnrealBuildTool
 	/// </summary>
 	public class ConfigCacheIni
 	{
+
+		// command class for being able to create config caches over and over without needing to read the ini files
+		public class Command
+		{
+			public string TrimmedLine;
+		}
+
+		class SectionCommand : Command
+		{
+			public FileReference Filename;
+			public int LineIndex;
+		}
+
+		class KeyValueCommand : Command
+		{
+			public string Key;
+			public string Value;
+			public ParseAction LastAction;
+		}
+
+		// cached ini files
+		static Dictionary<string, List<Command>> FileCache = new Dictionary<string,List<Command>>();
+		static Dictionary<string, ConfigCacheIni> IniCache = new Dictionary<string, ConfigCacheIni>();
+		static Dictionary<string, ConfigCacheIni> BaseIniCache = new Dictionary<string, ConfigCacheIni>();
+		static List<string> RequiredSections = new List<string>(){"AppxManifest", "/Script/AndroidRuntimeSettings.AndroidRuntimeSettings", "/Script/AndroidPlatformEditor.AndroidSDKSettings",
+																	"/Script/BuildSettings.BuildSettings", "/Script/IOSRuntimeSettings.IOSRuntimeSettings", "/Script/WindowsTargetPlatform.WindowsTargetSettings"};
+
+		// static creation functions for ini files
+		public static ConfigCacheIni CreateConfigCacheIni(UnrealTargetPlatform Platform, string BaseIniName, DirectoryReference ProjectDirectory, DirectoryReference EngineDirectory = null)
+		{
+			if (EngineDirectory == null)
+			{
+				EngineDirectory = UnrealBuildTool.EngineDirectory;
+			}
+
+			// cache base ini for use as the seed for the rest
+			if (!BaseIniCache.ContainsKey(BaseIniName))
+			{
+				BaseIniCache.Add(BaseIniName, new ConfigCacheIni(UnrealTargetPlatform.Unknown, BaseIniName, null, EngineDirectory, EngineOnly:true));
+			}
+
+			// build the new ini and cache it for later re-use
+			ConfigCacheIni BaseCache = BaseIniCache[BaseIniName];
+			string Key = GetIniPlatformName(Platform) + BaseIniName + EngineDirectory.FullName + (ProjectDirectory != null ? ProjectDirectory.FullName : "");
+			if (!IniCache.ContainsKey(Key))
+			{
+				IniCache.Add(Key, new ConfigCacheIni(Platform, BaseIniName, ProjectDirectory, EngineDirectory, BaseCache:BaseCache));
+			}
+			return IniCache[Key];
+		}
+
 		/// <summary>
 		/// List of values (or a single value)
 		/// </summary>
@@ -95,12 +146,12 @@ namespace UnrealBuildTool
 		/// <param name="ProjectDirectory">Project path</param>
 		/// <param name="Platform">Target platform</param>
 		/// <param name="BaseIniName">Ini name (Engine, Editor, etc)</param>
-		public ConfigCacheIni(UnrealTargetPlatform Platform, string BaseIniName, DirectoryReference ProjectDirectory, DirectoryReference EngineDirectory = null)
+		public ConfigCacheIni(UnrealTargetPlatform Platform, string BaseIniName, DirectoryReference ProjectDirectory, DirectoryReference EngineDirectory = null, bool EngineOnly = false, ConfigCacheIni BaseCache = null)
 		{
-			Init(Platform, BaseIniName, ProjectDirectory, EngineDirectory);
+			Init(Platform, BaseIniName, ProjectDirectory, EngineDirectory, EngineOnly, BaseCache);
 		}
 
-		private void Init(UnrealTargetPlatform Platform, string BaseIniName, DirectoryReference ProjectDirectory, DirectoryReference EngineDirectory)
+		private void Init(UnrealTargetPlatform Platform, string BaseIniName, DirectoryReference ProjectDirectory, DirectoryReference EngineDirectory, bool EngineOnly = false, ConfigCacheIni BaseCache = null)
 		{
 			Sections = new Dictionary<string, IniSection>(StringComparer.InvariantCultureIgnoreCase);
 			if (EngineDirectory == null)
@@ -108,11 +159,31 @@ namespace UnrealBuildTool
 				EngineDirectory = UnrealBuildTool.EngineDirectory;
 			}
 
-			foreach (var IniFileName in EnumerateCrossPlatformIniFileNames(ProjectDirectory, EngineDirectory, Platform, BaseIniName))
+			if (BaseCache != null)
+			{
+				foreach(var Pair in BaseCache.Sections)
+				{
+					Sections.Add(Pair.Key, Pair.Value);
+				}
+			}
+			if (EngineOnly)
 			{
-				if (IniFileName.Exists())
+				foreach (var IniFileName in EnumerateEngineIniFileNames(EngineDirectory, BaseIniName))
 				{
-					ParseIniFile(IniFileName);
+					if (IniFileName.Exists())
+					{
+						ParseIniFile(IniFileName);
+					}
+				}
+			}
+			else
+			{
+				foreach (var IniFileName in EnumerateCrossPlatformIniFileNames(ProjectDirectory, EngineDirectory, Platform, BaseIniName, BaseCache != null))
+				{
+					if (IniFileName.Exists())
+					{
+						ParseIniFile(IniFileName);
+					}
 				}
 			}
 		}
@@ -401,13 +472,23 @@ namespace UnrealBuildTool
 		public void ParseIniFile(FileReference Filename)
 		{
 			String[] IniLines = null;
-			try
+			List<Command> Commands = null;
+			if (!FileCache.ContainsKey(Filename.FullName))
 			{
-				IniLines = File.ReadAllLines(Filename.FullName);
+				try
+				{
+					IniLines = File.ReadAllLines(Filename.FullName);
+					Commands = new List<Command>();
+					FileCache.Add(Filename.FullName, Commands);
+				}
+				catch (Exception ex)
+				{
+					Console.WriteLine("Error reading ini file: " + Filename + " Exception: " + ex.Message);
+				}
 			}
-			catch (Exception ex)
+			else
 			{
-				Console.WriteLine("Error reading ini file: " + Filename + " Exception: " + ex.Message);
+				Commands = FileCache[Filename.FullName];
 			}
 			if (IniLines != null)
 			{
@@ -437,6 +518,14 @@ namespace UnrealBuildTool
 						{
 							CurrentSection = FindOrAddSection(TrimmedLine, Filename, LineIndex);
 							LastAction = ParseAction.None;
+							if (CurrentSection != null)
+							{
+								SectionCommand Command = new SectionCommand();
+								Command.Filename = Filename;
+								Command.LineIndex = LineIndex;
+								Command.TrimmedLine = TrimmedLine;
+								Commands.Add(Command);
+							}
 						}
 						else
 						{
@@ -448,10 +537,10 @@ namespace UnrealBuildTool
 							LastAction = GetActionForLine(ref TrimmedLine);
 							if (LastAction != ParseAction.None)
 							{
-								if (CurrentSection == null)
+/*								if (CurrentSection == null)
 								{
 									throw new IniParsingException("Trying to parse key/value pair that doesn't belong to any section in {0}, line {1}: {2}", Filename, LineIndex, TrimmedLine);
-								}
+								}*/
 								ParseKeyValuePair(TrimmedLine, Filename, LineIndex, out Key, out SingleValue);
 							}
 						}
@@ -460,16 +549,43 @@ namespace UnrealBuildTool
 					{
 						SingleValue += TrimmedLine;
 					}
-					if (!bMultiLine && LastAction != ParseAction.None)
+					if (!bMultiLine && LastAction != ParseAction.None && CurrentSection != null)
 					{
 						ProcessKeyValuePair(CurrentSection, Key, SingleValue, LastAction);
+						KeyValueCommand Command = new KeyValueCommand();
+						Command.Key = Key;
+						Command.Value = SingleValue;
+						Command.LastAction = LastAction;
+						Commands.Add(Command);
 						LastAction = ParseAction.None;
 						SingleValue = "";
 						Key = "";
 					}
+					else if (CurrentSection == null)
+					{
+						LastAction = ParseAction.None;
+					}
 					LineIndex++;
 				}
 			}
+			else if (Commands != null)
+			{
+				IniSection CurrentSection = null;
+
+				// run each command
+				for(int Idx = 0; Idx < Commands.Count; ++Idx)
+				{
+					var Command = Commands[Idx];
+					if (Command is SectionCommand)
+					{
+						CurrentSection = FindOrAddSection((Command as SectionCommand).TrimmedLine, (Command as SectionCommand).Filename, (Command as SectionCommand).LineIndex);
+					}
+					else if (Command is KeyValueCommand)
+					{
+						ProcessKeyValuePair(CurrentSection, (Command as KeyValueCommand).Key, (Command as KeyValueCommand).Value, (Command as KeyValueCommand).LastAction);
+					}
+				}
+			}
 		}
 
 		/// <summary>
@@ -562,19 +678,23 @@ namespace UnrealBuildTool
 			{
 				throw new IniParsingException("Empty section name when parsing {0}, line {1}: {2}", Filename, LineIndex, TrimmedLine);
 			}
-			IniSection CurrentSection;
-			if (Sections.TryGetValue(SectionName, out CurrentSection) == false)
+			else if (RequiredSections.Contains(SectionName))
 			{
-				CurrentSection = new IniSection();
-				Sections.Add(SectionName, CurrentSection);
-			}			
-			return CurrentSection;
+				IniSection CurrentSection;
+				if (Sections.TryGetValue(SectionName, out CurrentSection) == false)
+				{
+					CurrentSection = new IniSection();
+					Sections.Add(SectionName, CurrentSection);
+				}
+				return CurrentSection;
+			}
+			return null;
 		}
 
 		/// <summary>
-		/// Returns a list of INI filenames for the given project
+		/// Returns a list of INI filenames for the engine
 		/// </summary>
-		private static IEnumerable<FileReference> EnumerateCrossPlatformIniFileNames(DirectoryReference ProjectDirectory, DirectoryReference EngineDirectory, UnrealTargetPlatform Platform, string BaseIniName)
+		private static IEnumerable<FileReference> EnumerateEngineIniFileNames(DirectoryReference EngineDirectory, string BaseIniName)
 		{
 			// Engine/Config/Base.ini (included in every ini type, required)
 			yield return FileReference.Combine(EngineDirectory, "Config", "Base.ini");
@@ -584,10 +704,29 @@ namespace UnrealBuildTool
 
 			// Engine/Config/NotForLicensees/Base* ini
 			yield return FileReference.Combine(EngineDirectory, "Config", "NotForLicensees", "Base" + BaseIniName + ".ini");
+		}
+
 
-			// NOTE: 4.7: See comment in GetSourceIniHierarchyFilenames()
-			// Engine/Config/NoRedist/Base* ini
-			// yield return Path.Combine(EngineDirectory, "Config", "NoRedist", "Base" + BaseIniName + ".ini");
+		/// <summary>
+		/// Returns a list of INI filenames for the given project
+		/// </summary>
+		private static IEnumerable<FileReference> EnumerateCrossPlatformIniFileNames(DirectoryReference ProjectDirectory, DirectoryReference EngineDirectory, UnrealTargetPlatform Platform, string BaseIniName, bool SkipEngine)
+		{
+			if (!SkipEngine)
+			{
+				// Engine/Config/Base.ini (included in every ini type, required)
+				yield return FileReference.Combine(EngineDirectory, "Config", "Base.ini");
+
+				// Engine/Config/Base* ini
+				yield return FileReference.Combine(EngineDirectory, "Config", "Base" + BaseIniName + ".ini");
+
+				// Engine/Config/NotForLicensees/Base* ini
+				yield return FileReference.Combine(EngineDirectory, "Config", "NotForLicensees", "Base" + BaseIniName + ".ini");
+
+				// NOTE: 4.7: See comment in GetSourceIniHierarchyFilenames()
+				// Engine/Config/NoRedist/Base* ini
+				// yield return Path.Combine(EngineDirectory, "Config", "NoRedist", "Base" + BaseIniName + ".ini");
+			}
 
 			if(ProjectDirectory != null)
 			{
diff --git a/Engine/Source/Programs/UnrealBuildTool/Configuration/UEBuildPlatform.cs b/Engine/Source/Programs/UnrealBuildTool/Configuration/UEBuildPlatform.cs
--- a/Engine/Source/Programs/UnrealBuildTool/Configuration/UEBuildPlatform.cs
+++ b/Engine/Source/Programs/UnrealBuildTool/Configuration/UEBuildPlatform.cs
@@ -757,7 +757,7 @@ namespace UnrealBuildTool
 		{
 			if (!bInitializedProject)
 			{
-				ConfigCacheIni Ini = new ConfigCacheIni(InPlatform, "Engine", UnrealBuildTool.GetUProjectPath());
+				ConfigCacheIni Ini = ConfigCacheIni.CreateConfigCacheIni(InPlatform, "Engine", UnrealBuildTool.GetUProjectPath());
 				bool bValue = UEBuildConfiguration.bCompileAPEX;
 				if (Ini.GetBool("/Script/BuildSettings.BuildSettings", "bCompileApex", out bValue))
 				{
@@ -946,10 +946,12 @@ namespace UnrealBuildTool
 		{
 		}
 
+		static Dictionary<UnrealTargetPlatform, ConfigCacheIni> EngineIniCache = new Dictionary<UnrealTargetPlatform, ConfigCacheIni>();
+
 		protected static bool DoProjectSettingsMatchDefault(UnrealTargetPlatform Platform, DirectoryReference ProjectDirectoryName, string Section, string[] BoolKeys, string[] IntKeys, string[] StringKeys)
 		{
-			ConfigCacheIni ProjIni = new ConfigCacheIni(Platform, "Engine", ProjectDirectoryName);
-			ConfigCacheIni DefaultIni = new ConfigCacheIni(Platform, "Engine", (DirectoryReference)null);
+			ConfigCacheIni ProjIni = ConfigCacheIni.CreateConfigCacheIni(Platform, "Engine", ProjectDirectoryName);
+			ConfigCacheIni DefaultIni = ConfigCacheIni.CreateConfigCacheIni(Platform, "Engine", (DirectoryReference)null);
 
 			// look at all bool values
 			if (BoolKeys != null) foreach (string Key in BoolKeys)
diff --git a/Engine/Source/Programs/UnrealBuildTool/IOS/UEBuildIOS.cs b/Engine/Source/Programs/UnrealBuildTool/IOS/UEBuildIOS.cs
--- a/Engine/Source/Programs/UnrealBuildTool/IOS/UEBuildIOS.cs
+++ b/Engine/Source/Programs/UnrealBuildTool/IOS/UEBuildIOS.cs
@@ -167,7 +167,7 @@ namespace UnrealBuildTool
 
 				// update the configuration based on the project file
 				// look in ini settings for what platforms to compile for
-				ConfigCacheIni Ini = new ConfigCacheIni(InPlatform, "Engine", UnrealBuildTool.GetUProjectPath());
+				ConfigCacheIni Ini = ConfigCacheIni.CreateConfigCacheIni(InPlatform, "Engine", UnrealBuildTool.GetUProjectPath());
 				string MinVersion = "IOS_6";
 				if (Ini.GetString("/Script/IOSRuntimeSettings.IOSRuntimeSettings", "MinimumiOSVersion", out MinVersion))
 				{
diff --git a/Engine/Source/Programs/UnrealBuildTool/IOS/UEDeployIOS.cs b/Engine/Source/Programs/UnrealBuildTool/IOS/UEDeployIOS.cs
--- a/Engine/Source/Programs/UnrealBuildTool/IOS/UEDeployIOS.cs
+++ b/Engine/Source/Programs/UnrealBuildTool/IOS/UEDeployIOS.cs
@@ -181,7 +181,7 @@ namespace UnrealBuildTool.IOS
 
 			// get the settings from the ini file
 			// plist replacements
-			ConfigCacheIni Ini = new ConfigCacheIni(UnrealTargetPlatform.IOS, "Engine", UnrealBuildTool.GetUProjectPath());
+			ConfigCacheIni Ini = ConfigCacheIni.CreateConfigCacheIni(UnrealTargetPlatform.IOS, "Engine", UnrealBuildTool.GetUProjectPath());
 
 			// orientations
 			string SupportedOrientations = "";
@@ -748,7 +748,7 @@ namespace UnrealBuildTool.IOS
 		{
 			// get the settings from the ini file
 			// plist replacements
-			ConfigCacheIni Ini = new ConfigCacheIni(UnrealTargetPlatform.IOS, "Engine", UnrealBuildTool.GetUProjectPath());
+			ConfigCacheIni Ini = ConfigCacheIni.CreateConfigCacheIni(UnrealTargetPlatform.IOS, "Engine", UnrealBuildTool.GetUProjectPath());
 			bool bSupported = false;
 			Ini.GetBool("/Script/IOSRuntimeSettings.IOSRuntimeSettings", "bEnableCloudKitSupport", out bSupported);
 
diff --git a/Engine/Source/Programs/UnrealBuildTool/ToolChain/RemoteToolChain.cs b/Engine/Source/Programs/UnrealBuildTool/ToolChain/RemoteToolChain.cs
--- a/Engine/Source/Programs/UnrealBuildTool/ToolChain/RemoteToolChain.cs
+++ b/Engine/Source/Programs/UnrealBuildTool/ToolChain/RemoteToolChain.cs
@@ -211,7 +211,7 @@ namespace UnrealBuildTool
 		{
 			base.ParseProjectSettings();
 
-			ConfigCacheIni Ini = new ConfigCacheIni(UnrealTargetPlatform.IOS, "Engine", UnrealBuildTool.GetUProjectPath());
+			ConfigCacheIni Ini = ConfigCacheIni.CreateConfigCacheIni(UnrealTargetPlatform.IOS, "Engine", UnrealBuildTool.GetUProjectPath());
 			string ServerName = RemoteServerName;
 			if (Ini.GetString("/Script/IOSRuntimeSettings.IOSRuntimeSettings", "RemoteServerName", out ServerName) && !String.IsNullOrEmpty(ServerName))
 			{
diff --git a/Engine/Source/Programs/UnrealBuildTool/Windows/UEBuildWindows.cs b/Engine/Source/Programs/UnrealBuildTool/Windows/UEBuildWindows.cs
--- a/Engine/Source/Programs/UnrealBuildTool/Windows/UEBuildWindows.cs
+++ b/Engine/Source/Programs/UnrealBuildTool/Windows/UEBuildWindows.cs
@@ -430,7 +430,7 @@ namespace UnrealBuildTool
 			// ...check if it was supported from a config.
 			if (!SupportWindowsXPIfAvailable)
 			{
-				ConfigCacheIni Ini = new ConfigCacheIni(UnrealTargetPlatform.Win64, "Engine", UnrealBuildTool.GetUProjectPath());
+				ConfigCacheIni Ini = ConfigCacheIni.CreateConfigCacheIni(UnrealTargetPlatform.Win64, "Engine", UnrealBuildTool.GetUProjectPath());
 				string MinimumOS;
 				if (Ini.GetString("/Script/WindowsTargetPlatform.WindowsTargetSettings", "MinimumOSVersion", out MinimumOS))
 				{
