jira: UE-20737
-Change the plugin editor template to prevent new plugin from crashing when doing a cook
-Fix UBT switch "-editorrecompile" to work as expected

diff --git a/Engine/Plugins/Editor/PluginBrowser/Templates/Advanced/Source/PLUGIN_NAME/Private/PLUGIN_NAMEStyle.cpp b/Engine/Plugins/Editor/PluginBrowser/Templates/Advanced/Source/PLUGIN_NAME/Private/PLUGIN_NAMEStyle.cpp
--- a/Engine/Plugins/Editor/PluginBrowser/Templates/Advanced/Source/PLUGIN_NAME/Private/PLUGIN_NAMEStyle.cpp
+++ b/Engine/Plugins/Editor/PluginBrowser/Templates/Advanced/Source/PLUGIN_NAME/Private/PLUGIN_NAMEStyle.cpp
@@ -57,7 +57,10 @@ TSharedRef< FSlateStyleSet > FPLUGIN_NAMEStyle::Create()
 
 void FPLUGIN_NAMEStyle::ReloadTextures()
 {
-	FSlateApplication::Get().GetRenderer()->ReloadTextureResources();
+	if (FSlateApplication::IsInitialized())
+	{
+		FSlateApplication::Get().GetRenderer()->ReloadTextureResources();
+	}
 }
 
 const ISlateStyle& FPLUGIN_NAMEStyle::Get()
diff --git a/Engine/Plugins/Editor/PluginBrowser/Templates/Basic/Source/PLUGIN_NAME/Private/PLUGIN_NAMEStyle.cpp b/Engine/Plugins/Editor/PluginBrowser/Templates/Basic/Source/PLUGIN_NAME/Private/PLUGIN_NAMEStyle.cpp
--- a/Engine/Plugins/Editor/PluginBrowser/Templates/Basic/Source/PLUGIN_NAME/Private/PLUGIN_NAMEStyle.cpp
+++ b/Engine/Plugins/Editor/PluginBrowser/Templates/Basic/Source/PLUGIN_NAME/Private/PLUGIN_NAMEStyle.cpp
@@ -58,7 +58,10 @@ TSharedRef< FSlateStyleSet > FPLUGIN_NAMEStyle::Create()
 
 void FPLUGIN_NAMEStyle::ReloadTextures()
 {
-	FSlateApplication::Get().GetRenderer()->ReloadTextureResources();
+	if (FSlateApplication::IsInitialized())
+	{
+		FSlateApplication::Get().GetRenderer()->ReloadTextureResources();
+	}
 }
 
 const ISlateStyle& FPLUGIN_NAMEStyle::Get()
diff --git a/Engine/Source/Developer/DesktopPlatform/Private/DesktopPlatformBase.cpp b/Engine/Source/Developer/DesktopPlatform/Private/DesktopPlatformBase.cpp
--- a/Engine/Source/Developer/DesktopPlatform/Private/DesktopPlatformBase.cpp
+++ b/Engine/Source/Developer/DesktopPlatform/Private/DesktopPlatformBase.cpp
@@ -593,7 +593,7 @@ bool FDesktopPlatformBase::CompileGameProject(const FString& RootDir, const FStr
 	}
 
 	// Append any other options
-	Arguments += " -editorrecompile -progress -noubtmakefiles";
+	Arguments += " -editorrecompile -progress -noubtmakefiles -NoHotReloadFromIDE";
 
 	// Run UBT
 	return RunUnrealBuildTool(LOCTEXT("CompilingProject", "Compiling project..."), RootDir, Arguments, Warn);
diff --git a/Engine/Source/Programs/UnrealBuildTool/System/RulesCompiler.cs b/Engine/Source/Programs/UnrealBuildTool/System/RulesCompiler.cs
--- a/Engine/Source/Programs/UnrealBuildTool/System/RulesCompiler.cs
+++ b/Engine/Source/Programs/UnrealBuildTool/System/RulesCompiler.cs
@@ -1522,7 +1522,7 @@ namespace UnrealBuildTool
 						}
 						if (SourceFolderIndex != -1)
 						{
-							DirectoryReference TargetSourceFolder = new DirectoryReference(TargetSourceFolderString.Substring(0, SourceFolderIndex + 8));
+							DirectoryReference TargetSourceFolder = new DirectoryReference(TargetSourceFolderString.Substring(0, SourceFolderIndex + 7));
 							foreach (KeyValuePair<string, FileReference> CheckEntry in TargetNameToTargetFile)
 							{
 								if (CheckEntry.Value.IsUnderDirectory(TargetSourceFolder))
