missed some ini sections that need to be loaded

diff --git a/Engine/Source/Programs/UnrealBuildTool/Configuration/EngineConfiguration.cs b/Engine/Source/Programs/UnrealBuildTool/Configuration/EngineConfiguration.cs
--- a/Engine/Source/Programs/UnrealBuildTool/Configuration/EngineConfiguration.cs
+++ b/Engine/Source/Programs/UnrealBuildTool/Configuration/EngineConfiguration.cs
@@ -51,7 +51,9 @@ namespace UnrealBuildTool
 		static Dictionary<string, ConfigCacheIni> IniCache = new Dictionary<string, ConfigCacheIni>();
 		static Dictionary<string, ConfigCacheIni> BaseIniCache = new Dictionary<string, ConfigCacheIni>();
 		static List<string> RequiredSections = new List<string>(){"AppxManifest", "/Script/AndroidRuntimeSettings.AndroidRuntimeSettings", "/Script/AndroidPlatformEditor.AndroidSDKSettings",
-																	"/Script/BuildSettings.BuildSettings", "/Script/IOSRuntimeSettings.IOSRuntimeSettings", "/Script/WindowsTargetPlatform.WindowsTargetSettings"};
+																	"/Script/BuildSettings.BuildSettings", "/Script/IOSRuntimeSettings.IOSRuntimeSettings", "/Script/WindowsTargetPlatform.WindowsTargetSettings",
+																	"/Script/UnrealEd.ProjectPackagingSettings", "/Script/PS4PlatformEditor.PS4TargetSettings", "/Script/XboxOneTargetPlatform.XboxOneTargetSettings",
+																	"/Script/HTML5PlatformEditor.HTML5TargetSettings"};
 
 		// static creation functions for ini files
 		public static ConfigCacheIni CreateConfigCacheIni(UnrealTargetPlatform Platform, string BaseIniName, DirectoryReference ProjectDirectory, DirectoryReference EngineDirectory = null)
