Fix modules being incompatible after being compiled on editor startup; build ID was being set to a new value, but the editor had already read it.

diff --git a/Engine/Source/Programs/UnrealBuildTool/Configuration/UEBuildTarget.cs b/Engine/Source/Programs/UnrealBuildTool/Configuration/UEBuildTarget.cs
--- a/Engine/Source/Programs/UnrealBuildTool/Configuration/UEBuildTarget.cs
+++ b/Engine/Source/Programs/UnrealBuildTool/Configuration/UEBuildTarget.cs
@@ -1713,7 +1713,7 @@ namespace UnrealBuildTool
 			// Create a unique identifier for this build, which can be used to identify modules when the changelist is constant. It's fine to share this between runs with the same makefile; 
 			// the output won't change. By default we leave it blank when compiling a subset of modules (for hot reload, etc...), otherwise it won't match anything else. When writing to a directory
 			// that already contains a manifest, we'll reuse the build id that's already in there (see below).
-			string BuildId = (OnlyModules.Count == 0)? Guid.NewGuid().ToString() : "";
+			string BuildId = (OnlyModules.Count == 0 && !bEditorRecompile)? Guid.NewGuid().ToString() : "";
 
 			// Find all the build products and modules from this binary
 			Receipt = new TargetReceipt(TargetName, Platform, Configuration, BuildId, Version);
