Adding GUID to .emscripten build files to allow easy re-gen of files in case of change or programmer error.
Fixes UE-20591.

diff --git a/Engine/Source/Programs/UnrealBuildTool/HTML5/HTML5SDKInfo.cs b/Engine/Source/Programs/UnrealBuildTool/HTML5/HTML5SDKInfo.cs
--- a/Engine/Source/Programs/UnrealBuildTool/HTML5/HTML5SDKInfo.cs
+++ b/Engine/Source/Programs/UnrealBuildTool/HTML5/HTML5SDKInfo.cs
@@ -12,6 +12,8 @@ namespace UnrealBuildTool
 	public class HTML5SDKInfo
 	{
 		static string SDKBase { get { return Path.GetFullPath(Path.Combine(new string[] { BuildConfiguration.RelativeEnginePath,  "Source" , "ThirdParty", "HTML5", "emsdk" })); } }
+		// A GUID as a string. Allows updates to flush the emscripten install without bumping the SDK version number. Useful if a programming-error causes a bogus install.
+		static string SDKInfoGUID = "49CA9678-2667-48BC-A6A9-25D6FB341F08";
 		static string SDKVersion = "1.30.0";
 		static string EMSCRIPTEN_ROOT { get { return Path.Combine(SDKBase, "emscripten", SDKVersion); } }
 
@@ -140,7 +142,7 @@ namespace UnrealBuildTool
 		{
 			string ConfigFile = DOT_EMSCRIPTEN;
 
-			if (!File.Exists(ConfigFile) || !File.ReadAllText(ConfigFile).Contains("GENERATEDBYUE4"))
+			if (!File.Exists(ConfigFile) || !File.ReadAllText(ConfigFile).Contains("GENERATEDBYUE4='" + SDKVersion + "+" + SDKInfoGUID + "'"))
 			{
 				var ConfigString = String.Join(
 								Environment.NewLine,
@@ -155,7 +157,7 @@ namespace UnrealBuildTool
 								"COMPILER_ENGINE = NODE_JS",
 								"JS_ENGINES = [NODE_JS]",
 								"V8_ENGINE = ''",
-								"GENERATEDBYUE4='"+SDKVersion +"'"
+								"GENERATEDBYUE4='"+SDKVersion +"+"+SDKInfoGUID+"'"
 								);
 				File.WriteAllText(ConfigFile, ConfigString.Replace("\\","/"));
 			}
