Add code to try manually determining the INCLUDE and LIB environment variables based on the registry, and output a warning if they are incorrect.

diff --git a/Engine/Source/Programs/UnrealBuildTool/Windows/VCEnvironment.cs b/Engine/Source/Programs/UnrealBuildTool/Windows/VCEnvironment.cs
--- a/Engine/Source/Programs/UnrealBuildTool/Windows/VCEnvironment.cs
+++ b/Engine/Source/Programs/UnrealBuildTool/Windows/VCEnvironment.cs
@@ -5,6 +5,7 @@ using System.Collections.Generic;
 using System.Text.RegularExpressions;
 using System.Diagnostics;
 using System.IO;
+using System.Linq;
 using Microsoft.Win32;
 using System.Text;
 
@@ -18,6 +19,7 @@ namespace UnrealBuildTool
 		public readonly string            VSToolPath64Bit;      // The path to the 64bit platform tool binaries.
 		public readonly string            WindowsSDKDir;        // Installation folder of the Windows SDK, e.g. C:\Program Files\Microsoft SDKs\Windows\v6.0A\
 		public readonly string            WindowsSDKExtensionDir;  // Installation folder of the Windows SDK Extensions, e.g. C:\Program Files (x86)\Windows SDKs\10
+		public readonly string            WindowsSDKLibVersion;
 		public readonly string            NetFxSDKExtensionDir;    // Installation folder of the NetFx SDK, since that is split out from platform SDKs >= v10
 		public readonly Version           WindowsSDKExtensionHeaderLibVersion;  // 10.0.9910.0 for instance...
 		public readonly string            CompilerPath;         // The path to the linker for linking executables
@@ -25,6 +27,9 @@ namespace UnrealBuildTool
 		public readonly string            LinkerPath;           // The path to the linker for linking executables
 		public readonly string            LibraryLinkerPath;    // The path to the linker for linking libraries
 		public readonly string            ResourceCompilerPath; // The path to the resource compiler
+		public readonly string            VisualCppDir;         // Installation folder for Visual C++, eg. C:\Program Files (x86)\Microsoft Visual Studio 12.0\VC.
+		public readonly string            UniversalCRTDir;      // For Visual Studio 2015; the path to the universal CRT.
+		public readonly string            UniversalCRTVersion;  // For Visual Studio 2015; the universal CRT version to use.
 
 		private string _MSBuildPath = null;
 		public string MSBuildPath // The path to MSBuild
@@ -66,9 +71,13 @@ namespace UnrealBuildTool
 			}
 
 			WindowsSDKDir        = FindWindowsSDKInstallationFolder(Platform);
+			WindowsSDKLibVersion = FindWindowsSDKLibVersion(WindowsSDKDir);
 			WindowsSDKExtensionDir = FindWindowsSDKExtensionInstallationFolder();
 			NetFxSDKExtensionDir = FindNetFxSDKExtensionInstallationFolder();
+			VisualCppDir         = FindVisualCppInstallationFolder(WindowsPlatform.Compiler);
 			WindowsSDKExtensionHeaderLibVersion = FindWindowsSDKExtensionLatestVersion(WindowsSDKExtensionDir);
+			UniversalCRTDir = FindUniversalCRTInstallationFolder();
+			UniversalCRTVersion = FindUniversalCRTVersion(UniversalCRTDir);
 
 			VSToolPath32Bit      = GetVSToolPath32Bit        (BaseVSToolPath);
 			VSToolPath64Bit      = GetVSToolPath64Bit        (BaseVSToolPath);
@@ -85,6 +94,10 @@ namespace UnrealBuildTool
 			LibraryLinkerPath    = GetLibraryLinkerToolPath   (LinkerVSToolPath);
 			ResourceCompilerPath = GetResourceCompilerToolPath(Platform);
 
+			// Manually determine the compile environment
+			List<string> IncludePaths = GetVisualCppIncludePaths(VisualCppDir, UniversalCRTDir, UniversalCRTVersion, NetFxSDKExtensionDir, WindowsSDKDir, WindowsSDKLibVersion);
+			List<string> LibraryPaths = GetVisualCppLibraryPaths(VisualCppDir, UniversalCRTDir, UniversalCRTVersion, NetFxSDKExtensionDir, WindowsSDKDir, WindowsSDKLibVersion, Platform);
+
 			// We ensure an extra trailing slash because of a user getting an odd error where the paths seemed to get concatenated wrongly:
 			//
 			// C:\Programme\Microsoft Visual Studio 12.0\Common7\Tools../../VC/bin/x86_amd64/vcvarsx86_amd64.bat
@@ -114,6 +127,34 @@ namespace UnrealBuildTool
 				Environment.SetEnvironmentVariable("LIB",     Utils.ResolveEnvironmentVariable(WindowsSDKDir + "lib" + ConfigSuffix + ";%LIB%"));
 				Environment.SetEnvironmentVariable("INCLUDE", Utils.ResolveEnvironmentVariable(WindowsSDKDir + "include;%INCLUDE%"));
 			}
+
+			// Check the environment matches up with the environment variables set by the batch files. For now, this is academic, but once we've established this as reliable, 
+			// start using them in preference. 
+			// NOTE: We skip this step for Windows XP, because the batch file actually initializes them incorrectly (and we override them above).
+			if(!WindowsPlatform.IsWindowsXPSupported())
+			{
+				string IncludePathsString = String.Join(";", IncludePaths) + ";";
+				CompareEnvironmentVariable("INCLUDE", IncludePathsString);
+
+				string LibraryPathsString = String.Join(";", LibraryPaths) + ";";
+				CompareEnvironmentVariable("LIB", LibraryPathsString);
+			}
+		}
+
+		/** Check that an environment variable is the same as one derived from keys in the registry, and send a telemetry event if it's not */
+		private void CompareEnvironmentVariable(string VariableName, string RegistryValue)
+		{
+			string EnvironmentValue = Environment.GetEnvironmentVariable(VariableName);
+			Telemetry.SendEvent("CompareEnvironmentVariable", 
+				"Name", VariableName,
+				"Match", (EnvironmentValue == RegistryValue).ToString(),
+				"EnvironmentValue", EnvironmentValue,
+				"RegistryValue", RegistryValue,
+				"Platform", Platform.ToString(),
+				"UWPBuildForStore", UWPPlatform.bBuildForStore.ToString(),
+				"WinXP", WindowsPlatform.IsWindowsXPSupported().ToString(),
+				"Compiler", WindowsPlatform.Compiler.ToString(),
+				"UseWindowsSDK10", WindowsPlatform.bUseWindowsSDK10.ToString());
 		}
 
 		/// <returns>The path to Windows SDK directory for the specified version.</returns>
@@ -173,6 +214,36 @@ namespace UnrealBuildTool
 			return FinalResult;
 		}
 
+		/** Gets the version of the Windows SDK libraries to use. As per VCVarsQueryRegistry.bat, this is the directory name that sorts last. */
+		static string FindWindowsSDKLibVersion(string WindowsSDKDir)
+		{
+			string WindowsSDKLibVersion;
+			if(WindowsPlatform.Compiler == WindowsCompiler.VisualStudio2012)
+			{
+				WindowsSDKLibVersion = "win8";
+			}
+			else if(WindowsPlatform.Compiler == WindowsCompiler.VisualStudio2015 && WindowsPlatform.bUseWindowsSDK10)
+			{
+				DirectoryInfo IncludeDir = new DirectoryInfo(Path.Combine(WindowsSDKDir, "include"));
+				if(!IncludeDir.Exists)
+				{
+					throw new BuildException("Couldn't find Windows 10 SDK include directory ({0})", IncludeDir.FullName);
+				}
+
+				DirectoryInfo LatestIncludeDir = IncludeDir.EnumerateDirectories().OrderBy(x => x.Name).LastOrDefault();
+				if(LatestIncludeDir == null)
+				{
+					throw new BuildException("No Windows 10 SDK versions found under ({0})", IncludeDir.FullName);
+				}
+
+				WindowsSDKLibVersion = LatestIncludeDir.Name;
+			}
+			else
+			{
+				WindowsSDKLibVersion = "winv6.3";
+			}
+			return WindowsSDKLibVersion;
+		}
 		
 		private static string FindNetFxSDKExtensionInstallationFolder()
 		{
@@ -180,14 +251,7 @@ namespace UnrealBuildTool
 			switch (WindowsPlatform.Compiler)
 			{
 				case WindowsCompiler.VisualStudio2015:
-					if( WindowsPlatform.bUseWindowsSDK10 )
-					{ 
-						Version = "4.6";
-					}
-					else
-					{ 
-						return string.Empty;
-					}
+					Version = "4.6";
 					break;
 
 				default:
@@ -416,6 +480,218 @@ namespace UnrealBuildTool
 			return Path.Combine(FrameworkDirectory, FrameworkVersion, "MSBuild.exe");
 		}
 
+		/** Gets the Visual C++ installation folder from the registry */
+		static string FindVisualCppInstallationFolder(WindowsCompiler Version)
+		{
+			// Get the version string
+			string VisualCppVersion;
+			switch (WindowsPlatform.Compiler)
+			{
+				case WindowsCompiler.VisualStudio2015:
+					VisualCppVersion = "14.0";
+					break;
+				case WindowsCompiler.VisualStudio2013:
+					VisualCppVersion = "12.0";
+					break;
+				case WindowsCompiler.VisualStudio2012:
+					VisualCppVersion = "11.0";
+					break;
+				default:
+					throw new BuildException("Unexpected compiler version when trying to determine Visual C++ installation folder");
+			}
+
+			// Read the registry value
+			object Value = 
+				Registry.GetValue("HKEY_LOCAL_MACHINE\\SOFTWARE\\Microsoft\\VisualStudio\\SxS\\VC7", VisualCppVersion, null) ??
+				Registry.GetValue("HKEY_CURRENT_USER\\SOFTWARE\\Microsoft\\VisualStudio\\SxS\\VC7", VisualCppVersion, null) ??
+				Registry.GetValue("HKEY_LOCAL_MACHINE\\SOFTWARE\\Wow6432Node\\Microsoft\\VisualStudio\\SxS\\VC7", VisualCppVersion, null) ??
+				Registry.GetValue("HKEY_CURRENT_USER\\SOFTWARE\\Wow6432Node\\Microsoft\\VisualStudio\\SxS\\VC7", VisualCppVersion, null);
+
+			string InstallDir = Value as string;
+			if(InstallDir == null)
+			{
+				throw new BuildException("Visual C++ must be installed to build this target.");
+			}
+			return InstallDir;
+		}
+
+		/** Finds the directory containing the Universal CRT installation. Returns null for Visual Studio versions before 2015 */
+		static string FindUniversalCRTInstallationFolder()
+		{
+			if(WindowsPlatform.Compiler != WindowsCompiler.VisualStudio2015)
+			{
+				return null;
+			}
+
+			object Value = 
+				Registry.GetValue("HKEY_LOCAL_MACHINE\\SOFTWARE\\Microsoft\\Windows Kits\\Installed Roots", "KitsRoot10", null) ??
+				Registry.GetValue("HKEY_CURRENT_USER\\SOFTWARE\\Microsoft\\Windows Kits\\Installed Roots", "KitsRoot10", null) ??
+				Registry.GetValue("HKEY_LOCAL_MACHINE\\SOFTWARE\\Wow6432Node\\Microsoft\\Windows Kits\\Installed Roots", "KitsRoot10", null) ??
+				Registry.GetValue("HKEY_CURRENT_USER\\SOFTWARE\\Wow6432Node\\Microsoft\\Windows Kits\\Installed Roots", "KitsRoot10", null);
+
+			return (Value as string);
+		}
+
+		/** Gets the version of the Universal CRT to use. As per VCVarsQueryRegistry.bat, this is the directory name that sorts last. */
+		static string FindUniversalCRTVersion(string UniversalCRTDir)
+		{
+			string UniversalCRTVersion = null;
+			if(UniversalCRTDir != null)
+			{
+				DirectoryInfo IncludeDir = new DirectoryInfo(Path.Combine(UniversalCRTDir, "include"));
+				if(IncludeDir.Exists)
+				{
+					DirectoryInfo LatestIncludeDir = IncludeDir.EnumerateDirectories().OrderBy(x => x.Name).LastOrDefault();
+					if(LatestIncludeDir != null)
+					{
+						UniversalCRTVersion = LatestIncludeDir.Name;
+					}
+				}
+			}
+			return UniversalCRTVersion;
+		}
+
+		/** Sets the Visual C++ INCLUDE environment variable */
+		static List<string> GetVisualCppIncludePaths(string VisualCppDir, string UniversalCRTDir, string UniversalCRTVersion, string NetFXSDKDir, string WindowsSDKDir, string WindowsSDKLibVersion)
+		{
+			List<string> IncludePaths = new List<string>();
+
+			// Add the standard Visual C++ include paths
+			string StdIncludeDir = Path.Combine(VisualCppDir, "INCLUDE");
+			if(Directory.Exists(StdIncludeDir))
+			{
+				IncludePaths.Add(StdIncludeDir);
+			}
+			string AtlMfcIncludeDir = Path.Combine(VisualCppDir, "ATLMFC", "INCLUDE");
+			if(Directory.Exists(AtlMfcIncludeDir))
+			{
+				IncludePaths.Add(AtlMfcIncludeDir);
+			}
+
+			// Add the universal CRT paths
+			if(!String.IsNullOrEmpty(UniversalCRTDir) && !String.IsNullOrEmpty(UniversalCRTVersion))
+			{
+				IncludePaths.Add(Path.Combine(UniversalCRTDir, "include", UniversalCRTVersion, "ucrt"));
+			}
+
+			// Add the NETFXSDK include path
+			if(!String.IsNullOrEmpty(NetFXSDKDir))
+			{
+				IncludePaths.Add(Path.Combine(NetFXSDKDir, "include", "um")); // 2015
+			}
+
+			// Add the Windows SDK paths
+			if(WindowsPlatform.Compiler == WindowsCompiler.VisualStudio2015 && WindowsPlatform.bUseWindowsSDK10)
+			{ 
+				IncludePaths.Add(Path.Combine(WindowsSDKDir, "include", WindowsSDKLibVersion, "shared"));
+				IncludePaths.Add(Path.Combine(WindowsSDKDir, "include", WindowsSDKLibVersion, "um"));
+				IncludePaths.Add(Path.Combine(WindowsSDKDir, "include", WindowsSDKLibVersion, "winrt"));
+			}
+			else
+			{
+				IncludePaths.Add(Path.Combine(WindowsSDKDir, "include", "shared"));
+				IncludePaths.Add(Path.Combine(WindowsSDKDir, "include", "um"));
+				IncludePaths.Add(Path.Combine(WindowsSDKDir, "include", "winrt"));
+			}
+
+			// Add the existing include paths
+			string ExistingIncludePaths = Environment.GetEnvironmentVariable("INCLUDE");
+			if(ExistingIncludePaths != null)
+			{
+				IncludePaths.AddRange(ExistingIncludePaths.Split(new char[]{ ';' }, StringSplitOptions.RemoveEmptyEntries));
+			}
+
+			// Set the environment variable
+			return IncludePaths;
+		}
+
+		/** Sets the Visual C++ LIB environment variable */
+		static List<string> GetVisualCppLibraryPaths(string VisualCppDir, string UniversalCRTDir, string UniversalCRTVersion, string NetFXSDKDir, string WindowsSDKDir, string WindowsSDKLibVersion, CPPTargetPlatform Platform)
+		{
+			List<string> LibraryPaths = new List<string>();
+
+			// Add the standard Visual C++ library paths
+			if(Platform == CPPTargetPlatform.UWP && UWPPlatform.bBuildForStore && WindowsPlatform.Compiler == WindowsCompiler.VisualStudio2015)
+			{
+				string StoreLibraryDir = Path.Combine(VisualCppDir, "LIB", "amd64", "store");
+				if(Directory.Exists(StoreLibraryDir))
+				{
+					LibraryPaths.Add(StoreLibraryDir);
+				}
+			}
+			else if(Platform == CPPTargetPlatform.Win32)
+			{
+				string StdLibraryDir = Path.Combine(VisualCppDir, "LIB");
+				if(Directory.Exists(StdLibraryDir))
+				{
+					LibraryPaths.Add(StdLibraryDir);
+				}
+				string AtlMfcLibraryDir = Path.Combine(VisualCppDir, "ATLMFC", "LIB");
+				if(Directory.Exists(AtlMfcLibraryDir))
+				{
+					LibraryPaths.Add(AtlMfcLibraryDir);
+				}
+			}
+			else
+			{
+				string StdLibraryDir = Path.Combine(VisualCppDir, "LIB", "amd64");
+				if(Directory.Exists(StdLibraryDir))
+				{
+					LibraryPaths.Add(StdLibraryDir);
+				}
+				string AtlMfcLibraryDir = Path.Combine(VisualCppDir, "ATLMFC", "LIB", "amd64");
+				if(Directory.Exists(AtlMfcLibraryDir))
+				{
+					LibraryPaths.Add(AtlMfcLibraryDir);
+				}
+			}
+			
+			// Add the Universal CRT
+			if(!String.IsNullOrEmpty(UniversalCRTDir) && !String.IsNullOrEmpty(UniversalCRTVersion))
+			{
+				if(Platform == CPPTargetPlatform.Win32)
+				{
+					LibraryPaths.Add(Path.Combine(UniversalCRTDir, "lib", UniversalCRTVersion, "ucrt", "x86"));
+				}
+				else
+				{
+					LibraryPaths.Add(Path.Combine(UniversalCRTDir, "lib", UniversalCRTVersion, "ucrt", "x64"));
+				}
+			}
+
+			// Add the NETFXSDK include path
+			if(!String.IsNullOrEmpty(NetFXSDKDir))
+			{
+				if(Platform == CPPTargetPlatform.Win32)
+				{
+					LibraryPaths.Add(Path.Combine(NetFXSDKDir, "lib", "um", "x86"));
+				}
+				else
+				{
+					LibraryPaths.Add(Path.Combine(NetFXSDKDir, "lib", "um", "x64"));
+				}
+			}
+
+			// Add the standard Windows SDK paths
+			if(Platform == CPPTargetPlatform.Win32)
+			{
+				LibraryPaths.Add(Path.Combine(WindowsSDKDir, "lib", WindowsSDKLibVersion, "um", "x86"));
+			}
+			else
+			{
+				LibraryPaths.Add(Path.Combine(WindowsSDKDir, "lib", WindowsSDKLibVersion, "um", "x64"));
+			}
+
+			// Add the existing library paths
+			string ExistingLibraryPaths = Environment.GetEnvironmentVariable("LIB");
+			if(ExistingLibraryPaths != null)
+			{
+				LibraryPaths.AddRange(ExistingLibraryPaths.Split(new char[]{ ';' }, StringSplitOptions.RemoveEmptyEntries));
+			}
+
+			return LibraryPaths;
+		}
+
 		static VCEnvironment EnvVars = null;
 	}
 }
