Removing a bunch of static variables for the project file in UBT.

diff --git a/Engine/Source/Programs/UnrealBuildTool/Configuration/UEBuildTarget.cs b/Engine/Source/Programs/UnrealBuildTool/Configuration/UEBuildTarget.cs
--- a/Engine/Source/Programs/UnrealBuildTool/Configuration/UEBuildTarget.cs
+++ b/Engine/Source/Programs/UnrealBuildTool/Configuration/UEBuildTarget.cs
@@ -175,6 +175,7 @@ namespace UnrealBuildTool
 	/// </summary>
 	public class TargetDescriptor
 	{
+		public FileReference ProjectFile;
 		public string TargetName;
 		public UnrealTargetPlatform Platform;
 		public UnrealTargetConfiguration Configuration;
@@ -226,7 +227,7 @@ namespace UnrealBuildTool
 		}
 
 
-		public static List<TargetDescriptor> ParseTargetCommandLine(string[] SourceArguments )
+		public static List<TargetDescriptor> ParseTargetCommandLine(string[] SourceArguments, FileReference ProjectFile)
 		{
 			var Targets = new List<TargetDescriptor>();
 
@@ -484,13 +485,13 @@ namespace UnrealBuildTool
 					TargetName = PossibleTargetName;
 
 					// If a project file was not specified see if we can find one
-					FileReference CheckProjectFile;
-					if(UProjectInfo.TryGetProjectForTarget(TargetName, out CheckProjectFile))
+					FileReference ProjectFileForTarget = ProjectFile;
+					if(ProjectFileForTarget == null && UProjectInfo.TryGetProjectForTarget(TargetName, out ProjectFileForTarget))
 					{
-						Log.TraceVerbose("Found project file for {0} - {1}", TargetName, CheckProjectFile);
+						Log.TraceVerbose("Found project file for {0} - {1}", TargetName, ProjectFileForTarget);
 						if (UnrealBuildTool.HasUProjectFile() == false)
 						{
-							UnrealBuildTool.SetProjectFile(CheckProjectFile);
+							UnrealBuildTool.SetProjectFile(ProjectFileForTarget);
 						}
 					}
 
@@ -505,6 +506,7 @@ namespace UnrealBuildTool
 
 					Targets.Add( new TargetDescriptor()
 						{
+							ProjectFile = ProjectFileForTarget,
 							TargetName = TargetName,
 							Platform = Platform,
 							Configuration = Configuration,
@@ -538,9 +540,9 @@ namespace UnrealBuildTool
 			var CreateTargetStartTime = DateTime.UtcNow;
 
 			RulesAssembly RulesAssembly;
-			if(UnrealBuildTool.HasUProjectFile())
+			if(Desc.ProjectFile != null)
 			{
-				RulesAssembly = RulesCompiler.CreateProjectRulesAssembly(UnrealBuildTool.GetUProjectFile());
+				RulesAssembly = RulesCompiler.CreateProjectRulesAssembly(Desc.ProjectFile);
 			}
 			else
 			{
@@ -728,6 +730,13 @@ namespace UnrealBuildTool
 		[NonSerialized]
 		public RulesAssembly RulesAssembly;
 
+		/** The project file for this target */
+		public FileReference ProjectFile;
+
+		/** The project descriptor for this target */
+		[NonSerialized]
+		public ProjectDescriptor ProjectDescriptor;
+
 		/** Type of target */
 		public TargetRules.TargetType TargetType;
 
@@ -913,6 +922,7 @@ namespace UnrealBuildTool
 		public UEBuildTarget(SerializationInfo Info, StreamingContext Context)
 		{
 			TargetType                   = (TargetRules.TargetType)Info.GetInt32("tt");
+			ProjectFile					 = (FileReference)Info.GetValue("pf", typeof(FileReference));
 			AppName                      = Info.GetString("an");
 			TargetName                   = Info.GetString("tn");
 			bUseSharedBuildEnvironment   = Info.GetBoolean("sb");
@@ -944,6 +954,7 @@ namespace UnrealBuildTool
 		public void GetObjectData(SerializationInfo Info, StreamingContext Context)
 		{
 			Info.AddValue("tt", (int)TargetType);
+			Info.AddValue("pf", ProjectFile);
 			Info.AddValue("an", AppName);
 			Info.AddValue("tn", TargetName);
 			Info.AddValue("sb", bUseSharedBuildEnvironment);
@@ -977,6 +988,7 @@ namespace UnrealBuildTool
 		/// <param name="InTargetCsFilename">The name of the target </param>
 		public UEBuildTarget(TargetDescriptor InDesc, TargetRules InRules, RulesAssembly InRulesAssembly, string InPossibleAppName, FileReference InTargetCsFilename)
 		{
+			ProjectFile = InDesc.ProjectFile;
 			AppName = InDesc.TargetName;
 			TargetName = InDesc.TargetName;
 			Platform = InDesc.Platform;
@@ -1020,9 +1032,9 @@ namespace UnrealBuildTool
 			}
 
 			// Figure out what the project directory is. If we have a uproject file, use that. Otherwise use the engine directory.
-			if (UnrealBuildTool.HasUProjectFile())
+			if (ProjectFile != null)
 			{
-				ProjectDirectory = UnrealBuildTool.GetUProjectPath();
+				ProjectDirectory = ProjectFile.Directory;
 			}
 			else
 			{
@@ -1049,6 +1061,12 @@ namespace UnrealBuildTool
 			// Get the receipt path for this target
 			ReceiptFileName = TargetReceipt.GetDefaultPath(ProjectDirectory.FullName, TargetName, Platform, Configuration, UEBuildPlatform.GetBuildPlatform(Platform).GetActiveArchitecture());
 
+			// Read the project descriptor
+			if(ProjectFile != null)
+			{
+				ProjectDescriptor = ProjectDescriptor.FromFile(ProjectFile.FullName);
+			}
+
 			RemoteRoot = InDesc.RemoteRoot;
 
 			OnlyModules = InDesc.OnlyModules;
@@ -1063,7 +1081,7 @@ namespace UnrealBuildTool
 			{
 				OutputDirectory = UnrealBuildTool.EngineDirectory;
 			}
-			OutputPaths = MakeExecutablePaths(OutputDirectory, bCompileMonolithic? TargetName : AppName, Platform, Configuration, Rules.UndecoratedConfiguration, bCompileMonolithic && UnrealBuildTool.HasUProjectFile(), Rules.ExeBinariesSubFolder);
+			OutputPaths = MakeExecutablePaths(OutputDirectory, bCompileMonolithic? TargetName : AppName, Platform, Configuration, Rules.UndecoratedConfiguration, bCompileMonolithic && ProjectFile != null, Rules.ExeBinariesSubFolder);
 
 			// handle some special case defines (so build system can pass -DEFINE as normal instead of needing
 			// to know about special parameters)
@@ -1549,7 +1567,7 @@ namespace UnrealBuildTool
 			{
 				// Create the module rules
 				FileReference ModuleRulesFileName;
-				ModuleRules Rules = RulesAssembly.CreateModuleRules(ModuleName, TargetInfo, out ModuleRulesFileName);
+				ModuleRules Rules = CreateModuleRulesAndSetDefaults(ModuleName, out ModuleRulesFileName);
 
 				// Add the rules file itself
 				FileNames.Add(ModuleRulesFileName.FullName);
@@ -1645,21 +1663,21 @@ namespace UnrealBuildTool
 		/** Generates a public manifest file for writing out */
         public void GenerateManifest()
 		{
-			string ManifestPath;
-			if (UnrealBuildTool.RunningRocket() && UnrealBuildTool.HasUProjectFile())
+			FileReference ManifestPath;
+			if (UnrealBuildTool.RunningRocket() && ProjectFile != null)
 			{
-				ManifestPath = Path.Combine(UnrealBuildTool.GetUProjectPath().FullName, BuildConfiguration.BaseIntermediateFolder, "Manifest.xml");
+				ManifestPath = FileReference.Combine(ProjectFile.Directory, BuildConfiguration.BaseIntermediateFolder, "Manifest.xml");
 			}
 			else
 			{
-				ManifestPath = "../Intermediate/Build/Manifest.xml";
+				ManifestPath = FileReference.Combine(UnrealBuildTool.EngineDirectory, "Intermediate", "Build", "Manifest.xml");
 			}
 
 			BuildManifest Manifest = new BuildManifest();
 			if (UEBuildConfiguration.bMergeManifests)
 			{
 				// Load in existing manifest (if any)
-				Manifest = Utils.ReadClass<BuildManifest>(ManifestPath);
+				Manifest = Utils.ReadClass<BuildManifest>(ManifestPath.FullName);
 			}
 
 			// Expand all the paths in the receipt; they'll currently use variables for the engine and project directories
@@ -1696,7 +1714,7 @@ namespace UnrealBuildTool
 			}
 			if (UEBuildConfiguration.bGenerateManifest)
 			{
-				Utils.WriteClass<BuildManifest>(Manifest, ManifestPath, "");
+				Utils.WriteClass<BuildManifest>(Manifest, ManifestPath.FullName, "");
 			}
 		}
 
@@ -1897,9 +1915,9 @@ namespace UnrealBuildTool
 				}
 
 				// Set the define for the project name. This allows the executable to locate the correct project file to use, which may not be the same as the game name or target.
-				if(UnrealBuildTool.HasUProjectFile())
+				if(ProjectFile != null)
 				{
-					string ProjectName = UnrealBuildTool.GetUProjectFile().GetFileNameWithoutExtension();
+					string ProjectName = ProjectFile.GetFileNameWithoutExtension();
 					GlobalCompileEnvironment.Config.Definitions.Add(String.Format("UE_PROJECT_NAME={0}", ProjectName));
 				}
 			}
@@ -2946,7 +2964,7 @@ namespace UnrealBuildTool
 			{
 				if (UEBuildPlatform.GetBuildPlatform(TargetPlatform, true) == null)
 				{
-					string DirectoryFragment = String.Format("/{0}/", TargetPlatform.ToString());
+					string DirectoryFragment = Path.DirectorySeparatorChar + TargetPlatform.ToString() + Path.DirectorySeparatorChar;
 					ExcludeFolders.Add(DirectoryFragment);
 				}
 			}
@@ -2958,7 +2976,7 @@ namespace UnrealBuildTool
 			// If we're compiling against the engine, add the plugins enabled for this target
 			if(UEBuildConfiguration.bCompileAgainstEngine)
 			{
-				ProjectDescriptor Project = UnrealBuildTool.HasUProjectFile()? ProjectDescriptor.FromFile(UnrealBuildTool.GetUProjectFile().FullName) : null;
+				ProjectDescriptor Project = (ProjectFile != null)? ProjectDescriptor.FromFile(ProjectFile.FullName) : null;
 				foreach(PluginInfo ValidPlugin in ValidPlugins)
 				{
 					if(UProjectInfo.IsPluginEnabledForProject(ValidPlugin, Project, Platform))
@@ -3007,16 +3025,16 @@ namespace UnrealBuildTool
 		}
 
 		/** Checks whether a plugin path contains a platform directory fragment */
-		private static bool ShouldExcludePlugin(PluginInfo Plugin, List<string> ExcludeFragments)
+		private bool ShouldExcludePlugin(PluginInfo Plugin, List<string> ExcludeFragments)
 		{
 			string RelativePathFromRoot;
 			if(Plugin.LoadedFrom == PluginLoadedFrom.Engine)
 			{
-				RelativePathFromRoot = Utils.CleanDirectorySeparators(Plugin.File.MakeRelativeTo(UnrealBuildTool.EngineDirectory), '/');
+				RelativePathFromRoot = Plugin.File.MakeRelativeTo(UnrealBuildTool.EngineDirectory);
 			}
 			else
 			{
-				RelativePathFromRoot = Utils.CleanDirectorySeparators(Plugin.File.MakeRelativeTo(UnrealBuildTool.GetUProjectPath()), '/');
+				RelativePathFromRoot = Plugin.File.MakeRelativeTo(ProjectFile.Directory);
 			}
 			return ExcludeFragments.Any(x => RelativePathFromRoot.Contains(x));
 		}
@@ -3133,7 +3151,7 @@ namespace UnrealBuildTool
 			// Installed Engine intermediates go to the project's intermediate folder. Installed Engine never writes to the engine intermediate folder. (Those files are immutable)
 			// Also, when compiling in monolithic, all intermediates go to the project's folder.  This is because a project can change definitions that affects all engine translation
 			// units too, so they can't be shared between different targets.  They are effectively project-specific engine intermediates.
-			if( UnrealBuildTool.IsEngineInstalled() || ( UnrealBuildTool.HasUProjectFile() && ShouldCompileMonolithic() ) )
+			if( UnrealBuildTool.IsEngineInstalled() || ( ProjectFile != null && ShouldCompileMonolithic() ) )
 			{
 				var IntermediateConfiguration = Configuration;
 				if( UnrealBuildTool.RunningRocket() )
@@ -3145,9 +3163,9 @@ namespace UnrealBuildTool
 				GlobalCompileEnvironment.Config.OutputDirectory = DirectoryReference.Combine(UnrealBuildTool.EngineDirectory, BuildConfiguration.PlatformIntermediateFolder, OutputAppName, IntermediateConfiguration.ToString());
 				if (ShouldCompileMonolithic())
 				{
-					if(UnrealBuildTool.HasUProjectFile())
+					if(ProjectFile != null)
 					{
-						GlobalCompileEnvironment.Config.OutputDirectory = DirectoryReference.Combine(UnrealBuildTool.GetUProjectPath(), BuildConfiguration.PlatformIntermediateFolder, OutputAppName, IntermediateConfiguration.ToString());
+						GlobalCompileEnvironment.Config.OutputDirectory = DirectoryReference.Combine(ProjectFile.Directory, BuildConfiguration.PlatformIntermediateFolder, OutputAppName, IntermediateConfiguration.ToString());
 					}
 					else if(ForeignPlugins.Count > 0)
 					{
@@ -3329,6 +3347,71 @@ namespace UnrealBuildTool
 			FlatModuleCsData.Add(Module.Name, new FlatModuleCsDataType((Module.RulesFile == null)? null : Module.RulesFile.FullName));
 		}
 
+		/** Create a rules object for the given module, and set any default values for this target */
+		private ModuleRules CreateModuleRulesAndSetDefaults(string ModuleName, out FileReference ModuleFileName)
+		{
+			// Create the rules from the assembly
+			ModuleRules RulesObject = RulesAssembly.CreateModuleRules(ModuleName, TargetInfo, out ModuleFileName);
+
+			// Reads additional dependencies array for project module from project file and fills PrivateDependencyModuleNames. 
+			if(ProjectDescriptor != null)
+			{
+				ModuleDescriptor Module = ProjectDescriptor.Modules.FirstOrDefault(x => x.Name.Equals(ModuleName, StringComparison.InvariantCultureIgnoreCase));
+				if(Module != null && Module.AdditionalDependencies != null)
+				{
+					RulesObject.PrivateDependencyModuleNames.AddRange(Module.AdditionalDependencies);
+				}
+			}
+
+			// Validate rules object
+			if (RulesObject.Type == ModuleRules.ModuleType.CPlusPlus)
+			{
+				if (RulesObject.PrivateAssemblyReferences.Count > 0)
+				{
+					throw new BuildException("Module rules for '{0}' may not specify PrivateAssemblyReferences unless it is a CPlusPlusCLR module type.", ModuleName);
+				}
+
+				var InvalidDependencies = RulesObject.DynamicallyLoadedModuleNames.Intersect(RulesObject.PublicDependencyModuleNames.Concat(RulesObject.PrivateDependencyModuleNames)).ToList();
+				if (InvalidDependencies.Count != 0)
+				{
+					throw new BuildException("Module rules for '{0}' should not be dependent on modules which are also dynamically loaded: {1}", ModuleName, String.Join(", ", InvalidDependencies));
+				}
+
+				// Choose code optimization options based on module type (game/engine) if default optimization method is selected.
+				bool bIsEngineModule = ModuleFileName.IsUnderDirectory(UnrealBuildTool.EngineDirectory);
+				if (RulesObject.OptimizeCode == ModuleRules.CodeOptimization.Default)
+				{
+					// Engine/Source and Engine/Plugins are considered 'Engine' code...
+					if (bIsEngineModule)
+					{
+						// Engine module - always optimize (except Debug).
+						RulesObject.OptimizeCode = ModuleRules.CodeOptimization.Always;
+					}
+					else
+					{
+						// Game module - do not optimize in Debug and DebugGame builds.
+						RulesObject.OptimizeCode = ModuleRules.CodeOptimization.InNonDebugBuilds;
+					}
+				}
+
+				// Disable shared PCHs for game modules by default (but not game plugins, since they won't depend on the game's PCH!)
+				if (RulesObject.PCHUsage == ModuleRules.PCHUsageMode.Default)
+				{
+					if (ProjectFile == null || !ModuleFileName.IsUnderDirectory(DirectoryReference.Combine(ProjectFile.Directory, "Source")))
+					{
+						// Engine module or plugin module -- allow shared PCHs
+						RulesObject.PCHUsage = ModuleRules.PCHUsageMode.UseSharedPCHs;
+					}
+					else
+					{
+						// Game module.  Do not enable shared PCHs by default, because games usually have a large precompiled header of their own and compile times would suffer.
+						RulesObject.PCHUsage = ModuleRules.PCHUsageMode.NoSharedPCHs;
+					}
+				}
+			}
+			return RulesObject;
+		}
+
 		/** Finds a module given its name.  Throws an exception if the module couldn't be found. */
 		public UEBuildModule FindOrCreateModuleByName(string ModuleName)
 		{
@@ -3340,7 +3423,7 @@ namespace UnrealBuildTool
 				// @todo projectfiles: Cross-platform modules can appear here during project generation, but they may have already
 				//   been filtered out by the project generator.  This causes the projects to not be added to directories properly.
 				FileReference ModuleFileName;
-				var RulesObject = RulesAssembly.CreateModuleRules(ModuleName, TargetInfo, out ModuleFileName);
+				ModuleRules RulesObject = CreateModuleRulesAndSetDefaults(ModuleName, out ModuleFileName);
 				DirectoryReference ModuleDirectory = ModuleFileName.Directory;
 				
 				// Get the plugin for this module
@@ -3387,9 +3470,9 @@ namespace UnrealBuildTool
 				var IsGameModule = !ModuleType.IsEngineModule();
 
 				// Get the base directory for paths referenced by the module. If the module's under the UProject source directory use that, otherwise leave it relative to the Engine source directory.
-				DirectoryReference ProjectSourceDirectoryName = UnrealBuildTool.GetUProjectSourcePath();
-				if (ProjectSourceDirectoryName != null)
+				if(ProjectFile != null)
 				{
+					DirectoryReference ProjectSourceDirectoryName = DirectoryReference.Combine(ProjectFile.Directory, "Source");
 					if (ModuleFileName.IsUnderDirectory(ProjectSourceDirectoryName))
 					{
 						RulesObject.PublicIncludePaths = CombinePathList(ProjectSourceDirectoryName, RulesObject.PublicIncludePaths);
@@ -3442,10 +3525,10 @@ namespace UnrealBuildTool
 				List<FileItem> FoundSourceFiles = new List<FileItem>();
 				if (RulesObject.Type == ModuleRules.ModuleType.CPlusPlus || RulesObject.Type == ModuleRules.ModuleType.CPlusPlusCLR)
 				{
-					ProjectFile ProjectFile = null;
-					if (ProjectFileGenerator.bGenerateProjectFiles && ProjectFileGenerator.ModuleToProjectFileMap.TryGetValue(ModuleName, out ProjectFile))
+					ProjectFile ProjectFileForIDE = null;
+					if (ProjectFileGenerator.bGenerateProjectFiles && ProjectFileGenerator.ModuleToProjectFileMap.TryGetValue(ModuleName, out ProjectFileForIDE))
 					{
-						IntelliSenseGatherer = ProjectFile;
+						IntelliSenseGatherer = ProjectFileForIDE;
 					}
 
 					// So all we care about are the game module and/or plugins.
@@ -3453,9 +3536,9 @@ namespace UnrealBuildTool
 					{
 						List<FileReference> SourceFilePaths = new List<FileReference>();
 
-						if (ProjectFile != null)
+						if (ProjectFileForIDE != null)
 						{
-							foreach (var SourceFile in ProjectFile.SourceFiles)
+							foreach (var SourceFile in ProjectFileForIDE.SourceFiles)
 							{
 								SourceFilePaths.Add(SourceFile.Reference);
 							}
@@ -3561,7 +3644,7 @@ namespace UnrealBuildTool
 				}
 				else
 				{
-					BaseSourceDirectory = DirectoryReference.Combine(UnrealBuildTool.GetUProjectPath(), "Source");
+					BaseSourceDirectory = DirectoryReference.Combine(ProjectFile.Directory, "Source");
 				}
 
 				// If it's a game module (plugin or otherwise), add the root source directory to the include paths.
@@ -3746,7 +3829,7 @@ namespace UnrealBuildTool
 			if (Rules != null)
 			{
 				// We don't want to check this for uprojects (but do want to include UE4Game)
-				if (!UnrealBuildTool.HasUProjectFile() && !Rules.bOutputPubliclyDistributable)
+				if (ProjectFile == null && !Rules.bOutputPubliclyDistributable)
 				{
 					return true;
 				}
diff --git a/Engine/Source/Programs/UnrealBuildTool/IOS/IOSToolChain.cs b/Engine/Source/Programs/UnrealBuildTool/IOS/IOSToolChain.cs
--- a/Engine/Source/Programs/UnrealBuildTool/IOS/IOSToolChain.cs
+++ b/Engine/Source/Programs/UnrealBuildTool/IOS/IOSToolChain.cs
@@ -982,7 +982,7 @@ namespace UnrealBuildTool
 					// generate the dummy project so signing works
 					if (AppName == "UE4Game" || AppName == "UE4Client" || Utils.IsFileUnderDirectory(Target.ProjectDirectory + "/" + AppName + ".uproject", Path.GetFullPath("../..")))
 					{
-						UnrealBuildTool.GenerateProjectFiles (new XcodeProjectFileGenerator (), new string[] {"-platforms=IOS", "-NoIntellIsense", "-iosdeployonly", "-ignorejunk"});
+						UnrealBuildTool.GenerateProjectFiles (new XcodeProjectFileGenerator (UnrealBuildTool.GetUProjectFile()), new string[] {"-platforms=IOS", "-NoIntellIsense", "-iosdeployonly", "-ignorejunk"});
 						Project = Path.GetFullPath("../..") + "/UE4_IOS.xcodeproj";
 					}
 					else
@@ -1130,7 +1130,7 @@ namespace UnrealBuildTool
 					if (!bUseDangerouslyFastMode)
 					{
 						// generate the dummy project so signing works
-						UnrealBuildTool.GenerateProjectFiles(new XcodeProjectFileGenerator(), new string[] { "-NoIntellisense", "-iosdeployonly", (UnrealBuildTool.HasUProjectFile() ? "-game" : "") });
+						UnrealBuildTool.GenerateProjectFiles(new XcodeProjectFileGenerator(UnrealBuildTool.GetUProjectFile()), new string[] { "-NoIntellisense", "-iosdeployonly", (UnrealBuildTool.HasUProjectFile() ? "-game" : "") });
 					}
 
 					// now that 
diff --git a/Engine/Source/Programs/UnrealBuildTool/System/ActionHistory.cs b/Engine/Source/Programs/UnrealBuildTool/System/ActionHistory.cs
--- a/Engine/Source/Programs/UnrealBuildTool/System/ActionHistory.cs
+++ b/Engine/Source/Programs/UnrealBuildTool/System/ActionHistory.cs
@@ -99,8 +99,12 @@ namespace UnrealBuildTool
 			{
 				// Monolithic configs and programs have their Action History stored in their respective project folders
 				// or under engine intermediate folder + program name folder
-				DirectoryReference RootDirectory = UnrealBuildTool.GetUProjectPath();
-				if (RootDirectory == null)
+				DirectoryReference RootDirectory;
+				if(Target.ProjectFile != null)
+				{
+					RootDirectory = Target.ProjectFile.Directory;
+				}
+				else
 				{
 					RootDirectory = UnrealBuildTool.EngineDirectory;
 				}
@@ -109,8 +113,8 @@ namespace UnrealBuildTool
 			else
 			{
 				// Shared action history (unless this is a rocket target)
-				Folder = (UnrealBuildTool.RunningRocket() && UnrealBuildTool.HasUProjectFile()) ?
-					DirectoryReference.Combine(UnrealBuildTool.GetUProjectPath(), BuildConfiguration.BaseIntermediateFolder) :
+				Folder = (UnrealBuildTool.RunningRocket() && Target.ProjectFile != null) ?
+					DirectoryReference.Combine(Target.ProjectFile.Directory, BuildConfiguration.BaseIntermediateFolder) :
 					DirectoryReference.Combine(UnrealBuildTool.EngineDirectory, BuildConfiguration.BaseIntermediateFolder);
 			}
 			return FileReference.Combine(Folder, "ActionHistory.bin");
diff --git a/Engine/Source/Programs/UnrealBuildTool/System/CMakefileGenerator.cs b/Engine/Source/Programs/UnrealBuildTool/System/CMakefileGenerator.cs
--- a/Engine/Source/Programs/UnrealBuildTool/System/CMakefileGenerator.cs
+++ b/Engine/Source/Programs/UnrealBuildTool/System/CMakefileGenerator.cs
@@ -36,6 +36,11 @@ namespace UnrealBuildTool
 		/// True if intellisense data should be generated (takes a while longer)
 		bool bGenerateIntelliSenseData = false;
 
+		/// Default constructor
+		public CMakefileGenerator(FileReference InOnlyGameProject) : base(InOnlyGameProject)
+		{
+		}
+
 		/// True if we should include IntelliSense data in the generated project files when possible
 		override public bool ShouldGenerateIntelliSenseData()
 		{
@@ -103,10 +108,10 @@ namespace UnrealBuildTool
 
 			if (!String.IsNullOrEmpty (GameProjectName)) 
 			{
-				CMakeGameRootPath = "set(GAME_ROOT_PATH \"" + UnrealBuildTool.GetUProjectPath() + "\")\n";
+				CMakeGameRootPath = "set(GAME_ROOT_PATH \"" + OnlyGameProject.Directory.FullName + "\")\n";
 
-				GameProjectPath = UnrealBuildTool.GetUProjectPath ().FullName;
-				GameProjectFile = UnrealBuildTool.GetUProjectFile ().FullName;
+				GameProjectPath = OnlyGameProject.Directory.FullName;
+				GameProjectFile = OnlyGameProject.FullName;
 
 				CMakeGameProjectFile = "set(GAME_PROJECT_FILE \"" + GameProjectFile + "\")\n";
 
diff --git a/Engine/Source/Programs/UnrealBuildTool/System/CPPHeaders.cs b/Engine/Source/Programs/UnrealBuildTool/System/CPPHeaders.cs
--- a/Engine/Source/Programs/UnrealBuildTool/System/CPPHeaders.cs
+++ b/Engine/Source/Programs/UnrealBuildTool/System/CPPHeaders.cs
@@ -456,14 +456,14 @@ namespace UnrealBuildTool
 						EndIndex = FileContents.Length;
 					}
 
-					Result.AddRange(CollectHeaders(CPPFile, FileToRead, FileContents, InstalledFolder, StartIndex, EndIndex));
+					Result.AddRange(CollectHeaders(Target.ProjectFile, CPPFile, FileToRead, FileContents, InstalledFolder, StartIndex, EndIndex));
 
 					StartIndex = EndIndex + 1;
 				}
 			}
 			else
 			{
-				Result = CollectHeaders(CPPFile, FileToRead, FileContents, InstalledFolder, 0, FileContents.Length);
+				Result = CollectHeaders(Target.ProjectFile, CPPFile, FileToRead, FileContents, InstalledFolder, 0, FileContents.Length);
 			}
 
 			// Populate cache with results.
@@ -483,7 +483,7 @@ namespace UnrealBuildTool
 		/// <param name="InstalledFolder"></param>
 		/// <param name="StartIndex"></param>
 		/// <param name="EndIndex"></param>
-		private static List<DependencyInclude> CollectHeaders(FileItem CPPFile, string FileToRead, string FileContents, string InstalledFolder, int StartIndex, int EndIndex)
+		private static List<DependencyInclude> CollectHeaders(FileReference ProjectFile, FileItem CPPFile, string FileToRead, string FileContents, string InstalledFolder, int StartIndex, int EndIndex)
 		{
 			var Result = new List<DependencyInclude>();
 
@@ -500,7 +500,7 @@ namespace UnrealBuildTool
 				}
 
 				//@TODO: The intermediate exclusion is to work around autogenerated absolute paths in Module.SomeGame.cpp style files
-				bool bCheckForBackwardSlashes = FileToRead.StartsWith(InstalledFolder) || (UnrealBuildTool.HasUProjectFile() && new FileReference(FileToRead).IsUnderDirectory(UnrealBuildTool.GetUProjectPath()));
+				bool bCheckForBackwardSlashes = FileToRead.StartsWith(InstalledFolder) || ((ProjectFile != null) && new FileReference(FileToRead).IsUnderDirectory(ProjectFile.Directory));
 				if (bCheckForBackwardSlashes && !FileToRead.Contains("Intermediate") && !FileToRead.Contains("ThirdParty") && HeaderValue.IndexOf('\\', 0) >= 0)
 				{
 					throw new BuildException("In {0}: #include \"{1}\" contains backslashes ('\\'), please use forward slashes ('/') instead.", FileToRead, C.Value);
diff --git a/Engine/Source/Programs/UnrealBuildTool/System/CodeLiteGenerator.cs b/Engine/Source/Programs/UnrealBuildTool/System/CodeLiteGenerator.cs
--- a/Engine/Source/Programs/UnrealBuildTool/System/CodeLiteGenerator.cs
+++ b/Engine/Source/Programs/UnrealBuildTool/System/CodeLiteGenerator.cs
@@ -24,6 +24,10 @@ namespace UnrealBuildTool
 		public string CodeCompletionFileName = "CodeCompletionFolders.txt";
 		public string CodeCompletionPreProcessorFileName = "CodeLitePreProcessor.txt";
 
+		public CodeLiteGenerator(FileReference InOnlyGameProject) : base(InOnlyGameProject)
+		{
+		}
+
 		//
 		// Returns CodeLite's project filename extension.
 		//
@@ -34,7 +38,6 @@ namespace UnrealBuildTool
 				return ".project";
 			}
 		}
-
 		protected override bool WriteMasterProjectFile( ProjectFile UBTProject )
 		{
 			var SolutionFileName = MasterProjectName + SolutionExtension;
@@ -208,7 +211,7 @@ namespace UnrealBuildTool
 
 		protected override ProjectFile AllocateProjectFile( FileReference InitFilePath )
 		{
-			return new CodeLiteProject( InitFilePath );
+			return new CodeLiteProject( InitFilePath, OnlyGameProject );
 		}
 
 		public override MasterProjectFolder AllocateMasterProjectFolder( ProjectFileGenerator InitOwnerProjectFileGenerator, string InitFolderName )
diff --git a/Engine/Source/Programs/UnrealBuildTool/System/CodeLiteProject.cs b/Engine/Source/Programs/UnrealBuildTool/System/CodeLiteProject.cs
--- a/Engine/Source/Programs/UnrealBuildTool/System/CodeLiteProject.cs
+++ b/Engine/Source/Programs/UnrealBuildTool/System/CodeLiteProject.cs
@@ -11,8 +11,11 @@ namespace UnrealBuildTool
 {
 	public class CodeLiteProject : ProjectFile
 	{
-		public CodeLiteProject( FileReference InitFilePath ) : base(InitFilePath)
+		FileReference OnlyGameProject;
+
+		public CodeLiteProject( FileReference InitFilePath, FileReference InOnlyGameProject ) : base(InitFilePath)
 		{
+			OnlyGameProject = InOnlyGameProject;
 		}
 		
 		// Check if the XElement is empty.
@@ -46,9 +49,9 @@ namespace UnrealBuildTool
 			//
 
 			string GameWorkingDirectory = "";
-			if (UnrealBuildTool.HasUProjectFile ()) 
+			if (OnlyGameProject != null)
 			{
-				GameWorkingDirectory = Path.Combine (Path.GetDirectoryName (UnrealBuildTool.GetUProjectFile ().FullName), "Binaries", ProjectPlatformName);
+				GameWorkingDirectory = Path.Combine (Path.GetDirectoryName (OnlyGameProject.FullName), "Binaries", ProjectPlatformName);
 			}
 			//
 			// Build the working directory of the UE4Editor executable.
@@ -65,10 +68,9 @@ namespace UnrealBuildTool
 			}
 
 			string GameProjectFile = "";
-			FileReference GameProject = UnrealBuildTool.GetUProjectFile();
-			if (GameProject != null)
+			if (OnlyGameProject != null)
 			{
-				GameProjectFile = GameProject.FullName;
+				GameProjectFile = OnlyGameProject.FullName;
 			}
 
 			//
@@ -467,7 +469,7 @@ namespace UnrealBuildTool
 							string CookGameCommandLine = "mono AutomationTool.exe BuildCookRun ";
 
 							// Projects filename
-							CookGameCommandLine += "-project=\"" + UnrealBuildTool.GetUProjectFile () + "\" ";
+							CookGameCommandLine += "-project=\"" + OnlyGameProject.FullName + "\" ";
 							
 							// Disables Perforce functionality 
 							CookGameCommandLine += "-noP4 ";
diff --git a/Engine/Source/Programs/UnrealBuildTool/System/DependencyCache.cs b/Engine/Source/Programs/UnrealBuildTool/System/DependencyCache.cs
--- a/Engine/Source/Programs/UnrealBuildTool/System/DependencyCache.cs
+++ b/Engine/Source/Programs/UnrealBuildTool/System/DependencyCache.cs
@@ -404,9 +404,9 @@ namespace UnrealBuildTool
 		public static FileReference GetDependencyCachePathForTarget(UEBuildTarget Target)
 		{
 			DirectoryReference PlatformIntermediatePath;
-			if (UnrealBuildTool.HasUProjectFile())
+			if (Target.ProjectFile != null)
 			{
-				PlatformIntermediatePath = DirectoryReference.Combine(UnrealBuildTool.GetUProjectPath(), BuildConfiguration.PlatformIntermediateFolder);
+				PlatformIntermediatePath = DirectoryReference.Combine(Target.ProjectFile.Directory, BuildConfiguration.PlatformIntermediateFolder);
 			}
 			else
 			{
diff --git a/Engine/Source/Programs/UnrealBuildTool/System/ExternalExecution.cs b/Engine/Source/Programs/UnrealBuildTool/System/ExternalExecution.cs
--- a/Engine/Source/Programs/UnrealBuildTool/System/ExternalExecution.cs
+++ b/Engine/Source/Programs/UnrealBuildTool/System/ExternalExecution.cs
@@ -713,7 +713,7 @@ namespace UnrealBuildTool
 					Directory.CreateDirectory(ModuleInfoFileName.Directory.FullName);
 					System.IO.File.WriteAllText(ModuleInfoFileName.FullName, fastJSON.JSON.Instance.ToJSON(Manifest, new fastJSON.JSONParameters{ UseExtensions = false }));
 
-					string CmdLine = (UnrealBuildTool.HasUProjectFile()) ? "\"" + UnrealBuildTool.GetUProjectFile() + "\"" : Target.GetTargetName();
+					string CmdLine = (Target.ProjectFile != null) ? "\"" + Target.ProjectFile.FullName + "\"" : Target.GetTargetName();
 					CmdLine += " \"" + ModuleInfoFileName + "\" -LogCmds=\"loginit warning, logexit warning, logdatabase error\"";
 					if (UnrealBuildTool.RunningRocket())
 					{
diff --git a/Engine/Source/Programs/UnrealBuildTool/System/FlatCPPIncludeDepencencyCache.cs b/Engine/Source/Programs/UnrealBuildTool/System/FlatCPPIncludeDepencencyCache.cs
--- a/Engine/Source/Programs/UnrealBuildTool/System/FlatCPPIncludeDepencencyCache.cs
+++ b/Engine/Source/Programs/UnrealBuildTool/System/FlatCPPIncludeDepencencyCache.cs
@@ -310,9 +310,9 @@ namespace UnrealBuildTool
 		public static FileReference GetDependencyCachePathForTarget(UEBuildTarget Target)
 		{
 			DirectoryReference PlatformIntermediatePath;
-			if (UnrealBuildTool.HasUProjectFile())
+			if (Target.ProjectFile != null)
 			{
-				PlatformIntermediatePath = DirectoryReference.Combine(UnrealBuildTool.GetUProjectPath(), BuildConfiguration.PlatformIntermediateFolder);
+				PlatformIntermediatePath = DirectoryReference.Combine(Target.ProjectFile.Directory, BuildConfiguration.PlatformIntermediateFolder);
 			}
 			else
 			{
diff --git a/Engine/Source/Programs/UnrealBuildTool/System/KDevelopGenerator.cs b/Engine/Source/Programs/UnrealBuildTool/System/KDevelopGenerator.cs
--- a/Engine/Source/Programs/UnrealBuildTool/System/KDevelopGenerator.cs
+++ b/Engine/Source/Programs/UnrealBuildTool/System/KDevelopGenerator.cs
@@ -34,6 +34,9 @@ namespace UnrealBuildTool
 	/// </summary>
 	public class KDevelopGenerator : ProjectFileGenerator
 	{
+		public KDevelopGenerator(FileReference InOnlyGameProject) : base(InOnlyGameProject)
+		{
+		}
 
 		/// File extension for project files we'll be generating (e.g. ".vcxproj")
 		override public string ProjectFileExtension
@@ -83,7 +86,7 @@ namespace UnrealBuildTool
 
 			if (TargetName == GameProjectName)
 			{			
-				ProjectCmdArg = " -project=\"" + UnrealBuildTool.GetUProjectFile () + "\"";
+				ProjectCmdArg = " -project=\"" + OnlyGameProject.FullName + "\"";
 				Executable = "mono";
 				BuildCommand = "Engine/Binaries/DotNET/UnrealBuildTool.exe";
 
@@ -94,7 +97,7 @@ namespace UnrealBuildTool
 
 			} 
 			else if (TargetName == (GameProjectName + "Editor")) {
-				ProjectCmdArg = " -editorrecompile -project=\"" + UnrealBuildTool.GetUProjectFile () + "\"";
+				ProjectCmdArg = " -editorrecompile -project=\"" + OnlyGameProject.FullName + "\"";
 				Executable = "mono";
 				BuildCommand = "Engine/Binaries/DotNET/UnrealBuildTool.exe";
 
diff --git a/Engine/Source/Programs/UnrealBuildTool/System/MakefileGenerator.cs b/Engine/Source/Programs/UnrealBuildTool/System/MakefileGenerator.cs
--- a/Engine/Source/Programs/UnrealBuildTool/System/MakefileGenerator.cs
+++ b/Engine/Source/Programs/UnrealBuildTool/System/MakefileGenerator.cs
@@ -38,6 +38,11 @@ namespace UnrealBuildTool
 		/// Now this is needed for project target generation.
 		bool bGenerateIntelliSenseData = true;
 
+		/// Default constructor
+		public MakefileGenerator(FileReference InOnlyGameProject) : base(InOnlyGameProject)
+		{
+		}
+
 		/// True if we should include IntelliSense data in the generated project files when possible
 		override public bool ShouldGenerateIntelliSenseData()
 		{
@@ -71,7 +76,7 @@ namespace UnrealBuildTool
 
             if (!String.IsNullOrEmpty(GameProjectName))
             {
-                GameProjectFile = UnrealBuildTool.GetUProjectFile().FullName;
+                GameProjectFile = OnlyGameProject.FullName;
                 MakeGameProjectFile = "GAMEPROJECTFILE =" + GameProjectFile + "\n";
                 ProjectBuildCommand = "PROJECTBUILD = mono $(UNREALROOTPATH)/Engine/Binaries/DotNET/UnrealBuildTool.exe\n";
             }
diff --git a/Engine/Source/Programs/UnrealBuildTool/System/ProjectFileGenerator.cs b/Engine/Source/Programs/UnrealBuildTool/System/ProjectFileGenerator.cs
--- a/Engine/Source/Programs/UnrealBuildTool/System/ProjectFileGenerator.cs
+++ b/Engine/Source/Programs/UnrealBuildTool/System/ProjectFileGenerator.cs
@@ -219,6 +219,9 @@ namespace UnrealBuildTool
 		/// @todo projectfiles: Nasty global static list.  This is only really used for IntelliSense, and to avoid extra folder searches for projects we've already cached source files for.
 		public static readonly Dictionary<string, ProjectFile> ModuleToProjectFileMap = new Dictionary<string, ProjectFile>( StringComparer.InvariantCultureIgnoreCase );
 
+		/// If generating project files for a single project, the path to its .uproject file.
+		public readonly FileReference OnlyGameProject;
+
 		/// When generating IntelliSense data, we may want to only generate data for a specific project file, even if other targets make use of modules
 		/// in this project file.  This is useful to prevent unusual or hacky global definitions from Programs affecting the Editor/Engine modules.  We
 		/// always want the most common and useful definitions to be set when working with solutions with many modules.
@@ -241,6 +244,15 @@ namespace UnrealBuildTool
 			return bGenerateIntelliSenseData;
 		}
 
+		/// <summary>
+		/// Default constructor.
+		/// </summary>
+		/// <param name="InOnlyGameProject">The project file passed in on the command line</param>
+		public ProjectFileGenerator(FileReference InOnlyGameProject)
+		{
+			OnlyGameProject = InOnlyGameProject;
+		}
+
 		/// <summary>
 		/// Adds all .automation.csproj files to the solution.
 		/// </summary>
@@ -300,10 +312,10 @@ namespace UnrealBuildTool
 			{
 				Log.TraceInformation("Discovering modules, targets and source code for game...");
 
-				MasterProjectPath = UnrealBuildTool.GetUProjectPath();
+				MasterProjectPath = OnlyGameProject.Directory;
 					
-					// Set the project file name
-				MasterProjectName = Path.GetFileNameWithoutExtension(UnrealBuildTool.GetUProjectFile().FullName);
+				// Set the project file name
+				MasterProjectName = Path.GetFileNameWithoutExtension(OnlyGameProject.FullName);
 
 				if (!DirectoryReference.Combine(MasterProjectPath, "Source").Exists())
 				{
@@ -326,11 +338,11 @@ namespace UnrealBuildTool
 			{
 				Log.TraceInformation("Discovering modules, targets and source code for project...");
 
-				MasterProjectPath = UnrealBuildTool.GetUProjectPath();
+				MasterProjectPath = OnlyGameProject.Directory;
 				IntermediateProjectFilesPath = DirectoryReference.Combine( MasterProjectPath, "Intermediate", "ProjectFiles" );
 
 				// Set the project file name
-				MasterProjectName = Path.GetFileNameWithoutExtension(UnrealBuildTool.GetUProjectFile().FullName);
+				MasterProjectName = OnlyGameProject.GetFileNameWithoutExtension();
 
 				if (!DirectoryReference.Combine(MasterProjectPath, "Source").Exists())
 				{
@@ -796,10 +808,10 @@ namespace UnrealBuildTool
 				}
 
 				// Make sure we can get a valid game name out of this project
-				var GameName = UnrealBuildTool.GetUProjectFile().GetFileNameWithoutExtension();
+				var GameName = OnlyGameProject.GetFileNameWithoutExtension();
 				if( String.IsNullOrEmpty( GameName ) )
 				{
-					throw new BuildException("A valid Rocket game project was not found in the specified location (" + UnrealBuildTool.GetUProjectPath() + ")");
+					throw new BuildException("A valid Rocket game project was not found in the specified location (" + OnlyGameProject.Directory.FullName + ")");
 				}
 
 				IncludeEngineSource = true;
@@ -812,10 +824,10 @@ namespace UnrealBuildTool
 			}
 			else if( bGeneratingGameProjectFiles )
 			{
-				GameProjectName = UnrealBuildTool.GetUProjectFile().GetFileNameWithoutExtension();
+				GameProjectName = OnlyGameProject.GetFileNameWithoutExtension();
 				if (String.IsNullOrEmpty(GameProjectName))
 				{
-					throw new BuildException("A valid game project was not found in the specified location (" + UnrealBuildTool.GetUProjectPath() + ")");
+					throw new BuildException("A valid game project was not found in the specified location (" + OnlyGameProject.Directory.FullName + ")");
 				}
 
 				IncludeEngineSource = bAlwaysIncludeEngineModules;
@@ -976,9 +988,9 @@ namespace UnrealBuildTool
 					}
 				}
 
-				if (UnrealBuildTool.HasUProjectFile())
+				if (OnlyGameProject != null)
 				{
-					string ProjectRoot = UnrealBuildTool.GetUProjectPath().FullName;
+					string ProjectRoot = OnlyGameProject.Directory.FullName;
 					if (Search.StartsWith(ProjectRoot, StringComparison.InvariantCultureIgnoreCase))
 					{
 						if (ProjectRoot.EndsWith("\\") || ProjectRoot.EndsWith("/"))
@@ -1722,7 +1734,7 @@ namespace UnrealBuildTool
 
 						bool bProjectAlreadyExisted;
 						var ProjectFile = FindOrAddProject(ProjectFilePath, IncludeInGeneratedProjects: true, bAlreadyExisted: out bProjectAlreadyExisted);
-						ProjectFile.IsForeignProject = bGeneratingGameProjectFiles && UnrealBuildTool.HasUProjectFile() && TargetFilePath.IsUnderDirectory(UnrealBuildTool.GetUProjectPath());
+						ProjectFile.IsForeignProject = bGeneratingGameProjectFiles && OnlyGameProject != null && TargetFilePath.IsUnderDirectory(OnlyGameProject.Directory);
 						ProjectFile.IsGeneratedProject = true;
 						ProjectFile.IsStubProject = false;
 
diff --git a/Engine/Source/Programs/UnrealBuildTool/System/QMakefileGenerator.cs b/Engine/Source/Programs/UnrealBuildTool/System/QMakefileGenerator.cs
--- a/Engine/Source/Programs/UnrealBuildTool/System/QMakefileGenerator.cs
+++ b/Engine/Source/Programs/UnrealBuildTool/System/QMakefileGenerator.cs
@@ -34,6 +34,11 @@ namespace UnrealBuildTool
 	/// </summary>
 	public class QMakefileGenerator : ProjectFileGenerator
 	{
+		/// Default constructor
+		public QMakefileGenerator(FileReference InOnlyGameProject) : base(InOnlyGameProject)
+		{
+		}
+
 		/// File extension for project files we'll be generating (e.g. ".vcxproj")
 		override public string ProjectFileExtension
 		{
@@ -206,8 +211,8 @@ namespace UnrealBuildTool
 
 			if (!String.IsNullOrEmpty (GameProjectName))
 			{
-				GameProjectPath = UnrealBuildTool.GetUProjectPath ().FullName;
-				GameProjectFile = UnrealBuildTool.GetUProjectFile ().FullName;
+				GameProjectPath = OnlyGameProject.Directory.FullName;
+				GameProjectFile = OnlyGameProject.FullName;
 				QMakeGameProjectFile = "gameProjectFile=" + GameProjectFile + "\n";
 				BuildCommand = "build=mono $$unrealRootPath/Engine/Binaries/DotNET/UnrealBuildTool.exe\n\n";
 			}
diff --git a/Engine/Source/Programs/UnrealBuildTool/System/RulesCompiler.cs b/Engine/Source/Programs/UnrealBuildTool/System/RulesCompiler.cs
--- a/Engine/Source/Programs/UnrealBuildTool/System/RulesCompiler.cs
+++ b/Engine/Source/Programs/UnrealBuildTool/System/RulesCompiler.cs
@@ -387,46 +387,6 @@ namespace UnrealBuildTool
 
 		/** Redistribution override flag for this module. */
 		public bool? IsRedistributableOverride { get; set; }
-
-		/**
-		 * Reads additional dependencies array for project module from project file and fills PrivateDependencyModuleNames. 
-		 *
-		 * @param ProjectFile A path to the .uproject file.
-		 * @param ModuleName Name of the module.
-		 */
-		public void ReadAdditionalDependencies(FileReference ProjectFile, string ModuleName)
-		{
-			// Create a case-insensitive dictionary of the contents
-			Dictionary<string, object> Descriptor = fastJSON.JSON.Instance.ToObject<Dictionary<string, object>>(File.ReadAllText(ProjectFile.FullName));
-			Descriptor = new Dictionary<string,object>(Descriptor, StringComparer.InvariantCultureIgnoreCase);
-
-			// Get the list of plugins
-			object ModulesObject;
-			if (Descriptor.TryGetValue("Modules", out ModulesObject))
-			{
-				foreach(var ModuleObject in (ModulesObject as object[]).Cast<Dictionary<string, object>>())
-				{
-					object NameObject;
-					object AdditionalDependenciesObject;
-
-					if(!ModuleObject.TryGetValue("Name", out NameObject)
-						|| !(NameObject as string).Equals(ModuleName)
-						|| !ModuleObject.TryGetValue("AdditionalDependencies", out AdditionalDependenciesObject))
-					{
-						continue;
-					}
-
-					foreach (var AdditionalDependency in (AdditionalDependenciesObject as object[]).Cast<string>())
-					{
-						if(!PrivateDependencyModuleNames.Contains(AdditionalDependency))
-						{
-							PrivateDependencyModuleNames.Add(AdditionalDependency);
-						}
-					}
-					break;
-				}
-			}
-		}
 	}
 
 	/// <summary>
@@ -1072,11 +1032,6 @@ namespace UnrealBuildTool
 
 	public class RulesAssembly
 	{
-		/// <summary>
-		/// The parent rules assembly that this assembly inherits. Game assemblies inherit the engine assembly, and the engine assembly inherits nothing.
-		/// </summary>
-		private RulesAssembly Parent;
-
 		/// <summary>
 		/// The compiled assembly
 		/// </summary>
@@ -1107,6 +1062,11 @@ namespace UnrealBuildTool
 		/// </summary>
 		private Dictionary<FileReference, bool> ModuleHasSource = new Dictionary<FileReference, bool>();
 
+		/// <summary>
+		/// The parent rules assembly that this assembly inherits. Game assemblies inherit the engine assembly, and the engine assembly inherits nothing.
+		/// </summary>
+		private RulesAssembly Parent;
+
 		/// <summary>
 		/// Constructor. Compiles a rules assembly from the given source files.
 		/// </summary>
@@ -1289,76 +1249,16 @@ namespace UnrealBuildTool
 				throw new BuildException(Ex, "Unable to instantiate instance of '{0}' object type from compiled assembly '{1}'.  Unreal Build Tool creates an instance of your module's 'Rules' object in order to find out about your module's requirements.  The CLR exception details may provide more information:  {2}", ModuleTypeName, CompiledAssembly.FullName, Ex.ToString());
 			}
 
-			// Have to do absolute here as this could be a project that is under the root
-			bool bProjectModule = UnrealBuildTool.HasUProjectFile() && ModuleFileName.IsUnderDirectory(UnrealBuildTool.GetUProjectPath());
-			if (bProjectModule)
+			// Update the run-time dependencies path to remove $(PluginDir) and replace with a full path. When the receipt is saved it'll be converted to a $(ProjectDir) or $(EngineDir) equivalent.
+			foreach (var Dependency in RulesObject.RuntimeDependencies)
 			{
-				RulesObject.ReadAdditionalDependencies(UnrealBuildTool.GetUProjectFile(), ModuleName);
-			}
-
-			// Validate rules object
-			{
-				if (RulesObject.Type == ModuleRules.ModuleType.CPlusPlus)
-				{
-					if (RulesObject.PrivateAssemblyReferences.Count > 0)
-					{
-						throw new BuildException("Module rules for '{0}' may not specify PrivateAssemblyReferences unless it is a CPlusPlusCLR module type.", CompiledAssembly.FullName);
-					}
-
-					var InvalidDependencies = RulesObject.DynamicallyLoadedModuleNames.Intersect(RulesObject.PublicDependencyModuleNames.Concat(RulesObject.PrivateDependencyModuleNames)).ToList();
-					if (InvalidDependencies.Count != 0)
-					{
-						throw new BuildException("Module rules for '{0}' should not be dependent on modules which are also dynamically loaded: {1}", ModuleName, String.Join(", ", InvalidDependencies));
-					}
-
-					// Choose code optimization options based on module type (game/engine) if
-					// default optimization method is selected.
-					bool bIsEngineModule = ModuleFileName.IsUnderDirectory(UnrealBuildTool.EngineDirectory);
-					bool IsGamePluginModule = bProjectModule ? ModuleFileName.IsUnderDirectory(DirectoryReference.Combine(UnrealBuildTool.GetUProjectPath(), "Plugins")) : false;
-					if (RulesObject.OptimizeCode == ModuleRules.CodeOptimization.Default)
-					{
-						// Engine/Source and Engine/Plugins are considered 'Engine' code...
-						if (bIsEngineModule)
-						{
-							// Engine module - always optimize (except Debug).
-							RulesObject.OptimizeCode = ModuleRules.CodeOptimization.Always;
-						}
-						else
-						{
-							// Game module - do not optimize in Debug and DebugGame builds.
-							RulesObject.OptimizeCode = ModuleRules.CodeOptimization.InNonDebugBuilds;
-						}
-					}
-
-					// Disable shared PCHs for game modules by default (but not game plugins, since they won't depend on the game's PCH!)
-					if (RulesObject.PCHUsage == ModuleRules.PCHUsageMode.Default)
-					{
-						// Note that bIsEngineModule includes Engine/Plugins, so Engine/Plugins will use shared PCHs.
-						var IsProgramTarget = Target.Type != null && Target.Type == TargetRules.TargetType.Program;
-						if (bIsEngineModule || IsProgramTarget || IsGamePluginModule)
-						{
-							// Engine module or plugin module -- allow shared PCHs
-							RulesObject.PCHUsage = ModuleRules.PCHUsageMode.UseSharedPCHs;
-						}
-						else
-						{
-							// Game module.  Do not enable shared PCHs by default, because games usually have a large precompiled header of their own and compile times would suffer.
-							RulesObject.PCHUsage = ModuleRules.PCHUsageMode.NoSharedPCHs;
-						}
-					}
-				}
-
-				// Update the run-time dependencies path to remove $(PluginDir) and replace with a full path. When the receipt is saved it'll be converted to a $(ProjectDir) or $(EngineDir) equivalent.
-				foreach (var Dependency in RulesObject.RuntimeDependencies)
+				const string PluginDirVariable = "$(PluginDir)";
+				if(Dependency.Path.StartsWith(PluginDirVariable, StringComparison.InvariantCultureIgnoreCase))
 				{
-					const string PluginDirVariable = "$(PluginDir)";
-					if(Dependency.Path.StartsWith(PluginDirVariable, StringComparison.InvariantCultureIgnoreCase))
+					PluginInfo Plugin;
+					if(ModuleFileToPluginInfo.TryGetValue(ModuleFileName, out Plugin))
 					{
-						PluginInfo Plugin;
-						if(ModuleFileToPluginInfo.TryGetValue(ModuleFileName, out Plugin))
-						{
-							Dependency.Path = Plugin.Directory + Dependency.Path.Substring(PluginDirVariable.Length);
-						}
+						Dependency.Path = Plugin.Directory + Dependency.Path.Substring(PluginDirVariable.Length);
 					}
 				}
 			}
@@ -1366,27 +1266,6 @@ namespace UnrealBuildTool
 			return RulesObject;
 		}
 
-		/// <summary>
-		/// Creates an instance of a module rules descriptor object for the specified module name
-		/// </summary>
-		/// <param name="ModuleName">Name of the module</param>
-		/// <param name="Target">Information about the target associated with this module</param>
-		/// <param name="Rules">Output </param>
-		/// <returns>Compiled module rule info</returns>
-		public bool TryCreateModuleRules(string ModuleName, TargetInfo Target, out ModuleRules Rules)
-		{
-			if(GetModuleFileName( ModuleName ) == null)
-			{
-				Rules = null;
-				return false;
-			}
-			else
-			{
-				Rules = CreateModuleRules( ModuleName, Target );
-				return true;
-			}
-		}
-
 		/// <summary>
 		/// Determines whether the given module name is a game module (as opposed to an engine module)
 		/// </summary>
diff --git a/Engine/Source/Programs/UnrealBuildTool/System/UnrealBuildTool.cs b/Engine/Source/Programs/UnrealBuildTool/System/UnrealBuildTool.cs
--- a/Engine/Source/Programs/UnrealBuildTool/System/UnrealBuildTool.cs
+++ b/Engine/Source/Programs/UnrealBuildTool/System/UnrealBuildTool.cs
@@ -189,20 +189,6 @@ namespace UnrealBuildTool
             return UProjectPath;
         }
 
-        /// <summary>
-        /// The Unreal project source directory.  This should always be valid when compiling game projects.  Engine targets and program targets may not have a UProject set. */
-        /// </summary>
-        /// <returns>Relative path to the project's source directory</returns>
-        static public DirectoryReference GetUProjectSourcePath()
-        {
-            DirectoryReference ProjectPath = GetUProjectPath();
-            if (ProjectPath != null)
-            {
-                ProjectPath = DirectoryReference.Combine(ProjectPath, "Source");
-            }
-            return ProjectPath;
-        }
-
         static public string GetUBTPath()
         {
             string UnrealBuildToolPath = Assembly.GetExecutingAssembly().GetOriginalLocation();
@@ -1279,34 +1265,36 @@ namespace UnrealBuildTool
 
 					if (bGenerateCodeLiteFiles || bGenerateVCProjectFiles || bGenerateXcodeProjectFiles || bGenerateMakefiles || bGenerateCMakefiles || bGenerateQMakefiles || bGenerateKDevelopFiles)
                     {
+						FileReference UProjectFile = GetUProjectFile();
+
                         bool bGenerationSuccess = true;
                         if (bGenerateVCProjectFiles)
                         {
-                            bGenerationSuccess &= GenerateProjectFiles(new VCProjectFileGenerator(), Arguments);
+                            bGenerationSuccess &= GenerateProjectFiles(new VCProjectFileGenerator(UProjectFile), Arguments);
                         }
                         if (bGenerateXcodeProjectFiles)
                         {
-                            bGenerationSuccess &= GenerateProjectFiles(new XcodeProjectFileGenerator(), Arguments);
+                            bGenerationSuccess &= GenerateProjectFiles(new XcodeProjectFileGenerator(UProjectFile), Arguments);
                         }
                         if (bGenerateMakefiles)
                         {
-                            bGenerationSuccess &= GenerateProjectFiles(new MakefileGenerator(), Arguments);
+                            bGenerationSuccess &= GenerateProjectFiles(new MakefileGenerator(UProjectFile), Arguments);
                         }
 						if (bGenerateCMakefiles)
 						{
-							bGenerationSuccess &= GenerateProjectFiles(new CMakefileGenerator(), Arguments);
+							bGenerationSuccess &= GenerateProjectFiles(new CMakefileGenerator(UProjectFile), Arguments);
 						}
 						if (bGenerateQMakefiles)
 						{
-							bGenerationSuccess &= GenerateProjectFiles(new QMakefileGenerator(), Arguments);
+							bGenerationSuccess &= GenerateProjectFiles(new QMakefileGenerator(UProjectFile), Arguments);
 						}
 						if (bGenerateKDevelopFiles)
 						{
-							bGenerationSuccess &= GenerateProjectFiles(new KDevelopGenerator(), Arguments);
+							bGenerationSuccess &= GenerateProjectFiles(new KDevelopGenerator(UProjectFile), Arguments);
 						}
 						if (bGenerateCodeLiteFiles)
 						{
-							bGenerationSuccess &= GenerateProjectFiles(new CodeLiteGenerator(), Arguments);
+							bGenerationSuccess &= GenerateProjectFiles(new CodeLiteGenerator(UProjectFile), Arguments);
 						}
                         if(!bGenerationSuccess)
                         {
@@ -1365,7 +1353,7 @@ namespace UnrealBuildTool
                                 {
                                     // We need to be able to identify the Target.Type we can derive it from the Arguments.
                                     BuildConfiguration.bFlushBuildDirOnRemoteMac = false;
-									var TargetDescs = UEBuildTarget.ParseTargetCommandLine( Arguments );
+									var TargetDescs = UEBuildTarget.ParseTargetCommandLine( Arguments, GetUProjectFile() );
 			                        UEBuildTarget CheckTarget = UEBuildTarget.CreateTarget( TargetDescs[0] );	// @todo ubtmake: This may not work in assembler only mode.  We don't want to be loading target rules assemblies here either.
                                     CheckTarget.SetupGlobalEnvironment();
                                     if ((CheckTarget.TargetType == TargetRules.TargetType.Game) ||
@@ -1630,7 +1618,7 @@ namespace UnrealBuildTool
 
 					foreach (string[] TargetSetting in TargetSettings)
 					{
-						TargetDescs.AddRange( UEBuildTarget.ParseTargetCommandLine( TargetSetting ) );
+						TargetDescs.AddRange( UEBuildTarget.ParseTargetCommandLine( TargetSetting, GetUProjectFile() ) );
 					}
 
 					if( BuildConfiguration.bPrintPerformanceInfo )
@@ -2191,7 +2179,7 @@ namespace UnrealBuildTool
 				var EditorProcessName = Path.GetFileNameWithoutExtension(EditorProcessFilename);
 				var EditorProcesses = BuildHostPlatform.Current.GetProcessesByName(EditorProcessName);
 				var BinariesPath = Path.GetFullPath(Path.GetDirectoryName(EditorProcessFilename));
-				var PerProjectBinariesPath = UnrealBuildTool.HasUProjectFile()? Path.Combine(UnrealBuildTool.GetUProjectPath().FullName, BinariesPath.Substring(BinariesPath.LastIndexOf("Binaries"))) : "";
+				var PerProjectBinariesPath = (TargetDesc.ProjectFile != null)? Path.Combine(TargetDesc.ProjectFile.Directory.FullName, BinariesPath.Substring(BinariesPath.LastIndexOf("Binaries"))) : "";
 				bIsRunning = EditorProcesses.FirstOrDefault(EditorProc =>
 					{
 						if(!Path.GetFullPath(EditorProc.Filename).StartsWith(BinariesPath, StringComparison.InvariantCultureIgnoreCase))
diff --git a/Engine/Source/Programs/UnrealBuildTool/System/VCProject.cs b/Engine/Source/Programs/UnrealBuildTool/System/VCProject.cs
--- a/Engine/Source/Programs/UnrealBuildTool/System/VCProject.cs
+++ b/Engine/Source/Programs/UnrealBuildTool/System/VCProject.cs
@@ -268,6 +268,8 @@ namespace UnrealBuildTool
 			
 	public class VCProjectFile : MSBuildProjectFile
 	{
+		FileReference OnlyGameProject;
+
 		// This is the GUID that Visual Studio uses to identify a C++ project file in the solution
 		public override string ProjectTypeGUID
 		{
@@ -278,9 +280,10 @@ namespace UnrealBuildTool
 		/// Constructs a new project file object
 		/// </summary>
 		/// <param name="InitFilePath">The path to the project file on disk</param>
-		public VCProjectFile(FileReference InitFilePath)
-			: base( InitFilePath )
+		public VCProjectFile(FileReference InFilePath, FileReference InOnlyGameProject)
+			: base( InFilePath )
 		{
+			OnlyGameProject = InOnlyGameProject;
 		}
 
 
@@ -960,9 +963,9 @@ namespace UnrealBuildTool
 					DirectoryReference RootDirectory = UnrealBuildTool.EngineDirectory;
 					if ((TargetRules.IsAGame(TargetRulesObject.Type) || TargetRulesObject.Type == TargetRules.TargetType.Server) && bShouldCompileMonolithic && !TargetRulesObject.bOutputToEngineBinaries)
 					{
-						if (UnrealBuildTool.HasUProjectFile() && TargetFilePath.IsUnderDirectory(UnrealBuildTool.GetUProjectPath()))
+						if (OnlyGameProject != null && TargetFilePath.IsUnderDirectory(OnlyGameProject.Directory))
 						{
-							RootDirectory = UnrealBuildTool.GetUProjectPath();
+							RootDirectory = OnlyGameProject.Directory;
 						}
 						else
 						{
diff --git a/Engine/Source/Programs/UnrealBuildTool/System/VCProjectFileGenerator.cs b/Engine/Source/Programs/UnrealBuildTool/System/VCProjectFileGenerator.cs
--- a/Engine/Source/Programs/UnrealBuildTool/System/VCProjectFileGenerator.cs
+++ b/Engine/Source/Programs/UnrealBuildTool/System/VCProjectFileGenerator.cs
@@ -40,6 +40,11 @@ namespace UnrealBuildTool
 	/// </summary>
 	public class VCProjectFileGenerator : ProjectFileGenerator
 	{
+		/// Default constructor
+		public VCProjectFileGenerator(FileReference InOnlyGameProject) : base(InOnlyGameProject)
+		{
+		}
+
 		/// File extension for project files we'll be generating (e.g. ".vcxproj")
 		override public string ProjectFileExtension
 		{
@@ -102,7 +107,7 @@ namespace UnrealBuildTool
 		/// <returns>The newly allocated project file object</returns>
 		protected override ProjectFile AllocateProjectFile( FileReference InitFilePath )
 		{
-			return new VCProjectFile( InitFilePath );
+			return new VCProjectFile( InitFilePath, OnlyGameProject );
 		}
 
 
diff --git a/Engine/Source/Programs/UnrealBuildTool/System/XcodeProjectFileGenerator.cs b/Engine/Source/Programs/UnrealBuildTool/System/XcodeProjectFileGenerator.cs
--- a/Engine/Source/Programs/UnrealBuildTool/System/XcodeProjectFileGenerator.cs
+++ b/Engine/Source/Programs/UnrealBuildTool/System/XcodeProjectFileGenerator.cs
@@ -192,6 +192,10 @@ namespace UnrealBuildTool
 		// always seed the random number the same, so multiple runs of the generator will generate the same project
 		static Random Rand = new Random(0);
 
+		public XcodeProjectFileGenerator(FileReference InOnlyGameProject) : base(InOnlyGameProject)
+		{
+		}
+
 		/**
 		 * Make a random Guid string usable by Xcode (24 characters exactly)
 		 */
@@ -351,15 +355,15 @@ namespace UnrealBuildTool
 			if (Target.Type == XcodeTargetType.Legacy)
 			{
 				string UProjectPath = "";
-				if (UnrealBuildTool.HasUProjectFile() && Target.TargetFilePath != null && Target.TargetFilePath.IsUnderDirectory(UnrealBuildTool.GetUProjectPath()))
+				if (OnlyGameProject != null && Target.TargetFilePath != null && Target.TargetFilePath.IsUnderDirectory(OnlyGameProject.Directory))
 				{
-					if (MasterProjectPath == UnrealBuildTool.GetUProjectPath())
+					if (MasterProjectPath == OnlyGameProject.Directory)
 					{
-						UProjectPath = " " + "\\\"$(PROJECT_DIR)/" + Path.GetFileName(UnrealBuildTool.GetUProjectFile().FullName) + "\\\"";
+						UProjectPath = " " + "\\\"$(PROJECT_DIR)/" + OnlyGameProject.GetFileName() + "\\\"";
 					}
 					else
 					{
-						UProjectPath = " " + "\\\"" + UnrealBuildTool.GetUProjectFile() + "\\\"";
+						UProjectPath = " " + "\\\"" + OnlyGameProject.FullName + "\\\"";
 					}
 				}
 				// Xcode provides $ACTION argument for determining if we are building or cleaning a project
@@ -1425,16 +1429,16 @@ namespace UnrealBuildTool
 					continue;
 				}
 
-				if (UnrealBuildTool.HasUProjectFile() && Target.TargetFilePath != null && Target.TargetFilePath.IsUnderDirectory(UnrealBuildTool.GetUProjectPath()))
+				if (OnlyGameProject != null && Target.TargetFilePath != null && Target.TargetFilePath.IsUnderDirectory(OnlyGameProject.Directory))
 				{
 					if (Target.TargetName.EndsWith("Editor"))
 					{
-						WriteSchemeFile(XcodeProjectPath, Target, ".app", "UE4Editor", new List<string>(new string[] { "&quot;" + UnrealBuildTool.GetUProjectFile() + "&quot;" }));
+						WriteSchemeFile(XcodeProjectPath, Target, ".app", "UE4Editor", new List<string>(new string[] { "&quot;" + OnlyGameProject.FullName + "&quot;" }));
 					}
 					else
 					{
-						string ProjectBinariesDir = UnrealBuildTool.GetUProjectPath() + "/Binaries/Mac/";
-						WriteSchemeFile(XcodeProjectPath, Target, ".app", ProjectBinariesDir + Target.TargetName, new List<string>(new string[] { "&quot;" + UnrealBuildTool.GetUProjectFile() + "&quot;" }));
+						string ProjectBinariesDir = OnlyGameProject.Directory.FullName + "/Binaries/Mac/";
+						WriteSchemeFile(XcodeProjectPath, Target, ".app", ProjectBinariesDir + Target.TargetName, new List<string>(new string[] { "&quot;" + OnlyGameProject.FullName + "&quot;" }));
 					}
 				}
 				else if (Target.TargetName.EndsWith("Editor") && Target.TargetName != "UE4Editor")
@@ -1456,11 +1460,11 @@ namespace UnrealBuildTool
 		private void WriteProject(StringBuilder XcodeProjectFileContent, List<XcodeProjectTarget> ProjectTargets)
 		{
 			// @todo: Windows and Mac should store the project in the same folder
-			string PathBranch = (BuildHostPlatform.Current.Platform != UnrealTargetPlatform.Mac) ? ((UnrealBuildTool.GetUProjectPath() != null) ? "/Intermediate/IOS" : "/Engine/Intermediate/IOS") : "";
+			string PathBranch = (BuildHostPlatform.Current.Platform != UnrealTargetPlatform.Mac) ? ((OnlyGameProject != null) ? "/Intermediate/IOS" : "/Engine/Intermediate/IOS") : "";
 			var XcodeProjectPath = MasterProjectPath.FullName + PathBranch + "/" + MasterProjectName + ".xcodeproj";
 			if (bGeneratingGameProjectFiles && GameProjectName == "UE4Game")
 			{
-				XcodeProjectPath = UnrealBuildTool.GetUProjectPath() + "/" + MasterProjectName + ".xcodeproj";
+				XcodeProjectPath = OnlyGameProject.Directory.FullName + "/" + MasterProjectName + ".xcodeproj";
 			}
 			if (!Directory.Exists(XcodeProjectPath))
 			{
@@ -1492,11 +1496,11 @@ namespace UnrealBuildTool
 
 		private void WriteWorkspace()
 		{
-			string PathBranch = (BuildHostPlatform.Current.Platform != UnrealTargetPlatform.Mac) ? ((UnrealBuildTool.GetUProjectPath() != null) ? "/Intermediate/IOS" : "/Engine/Intermediate/IOS") : "";
+			string PathBranch = (BuildHostPlatform.Current.Platform != UnrealTargetPlatform.Mac) ? ((OnlyGameProject != null) ? "/Intermediate/IOS" : "/Engine/Intermediate/IOS") : "";
 			var XcodeProjectPath = MasterProjectPath + PathBranch + "/" + MasterProjectName + ".xcodeproj";
 			if (bGeneratingGameProjectFiles && GameProjectName == "UE4Game")
 			{
-				XcodeProjectPath = UnrealBuildTool.GetUProjectPath() + "/" + MasterProjectName + ".xcodeproj";
+				XcodeProjectPath = OnlyGameProject.Directory.FullName + "/" + MasterProjectName + ".xcodeproj";
 			}
 			if (!Directory.Exists(XcodeProjectPath))
 			{
diff --git a/Engine/Source/Programs/UnrealBuildTool/ToolChain/UEToolChain.cs b/Engine/Source/Programs/UnrealBuildTool/ToolChain/UEToolChain.cs
--- a/Engine/Source/Programs/UnrealBuildTool/ToolChain/UEToolChain.cs
+++ b/Engine/Source/Programs/UnrealBuildTool/ToolChain/UEToolChain.cs
@@ -113,21 +113,7 @@ namespace UnrealBuildTool
 		public static FileReference GetResponseFileName( LinkEnvironment LinkEnvironment, FileItem OutputFile )
 		{
 			// Construct a relative path for the intermediate response file
-			FileReference ResponseFileName = FileReference.Combine( LinkEnvironment.Config.IntermediateDirectory, Path.GetFileName( OutputFile.AbsolutePath ) + ".response" );
-			if (UnrealBuildTool.HasUProjectFile())
-			{
-				// If this is the uproject being built, redirect the intermediate
-				if (OutputFile.Reference.IsUnderDirectory( UnrealBuildTool.GetUProjectPath() ))
-				{
-					ResponseFileName = FileReference.Combine(
-						UnrealBuildTool.GetUProjectPath(), 
-						BuildConfiguration.PlatformIntermediateFolder,
-						UnrealBuildTool.GetUProjectFile().GetFileNameWithoutExtension(),
-						LinkEnvironment.Config.Target.Configuration.ToString(),
-						Path.GetFileName(OutputFile.AbsolutePath) + ".response");
-				}
-			}
-			return ResponseFileName;
+			return FileReference.Combine( LinkEnvironment.Config.IntermediateDirectory, OutputFile.Reference.GetFileName() + ".response" );
 		}
 
 		/** Converts the passed in path from UBT host to compiler native format. */
