UBT: Fix CodeLite generator.

diff --git a/Engine/Source/Programs/UnrealBuildTool/System/CodeLiteProject.cs b/Engine/Source/Programs/UnrealBuildTool/System/CodeLiteProject.cs
--- a/Engine/Source/Programs/UnrealBuildTool/System/CodeLiteProject.cs
+++ b/Engine/Source/Programs/UnrealBuildTool/System/CodeLiteProject.cs
@@ -469,8 +469,11 @@ namespace UnrealBuildTool
 							string CookGameCommandLine = "mono AutomationTool.exe BuildCookRun ";
 
 							// Projects filename
-							CookGameCommandLine += "-project=\"" + OnlyGameProject.FullName + "\" ";
-							
+							if (OnlyGameProject != null)
+							{
+								CookGameCommandLine += "-project=\"" + OnlyGameProject.FullName + "\" ";
+							}
+
 							// Disables Perforce functionality 
 							CookGameCommandLine += "-noP4 ";
 							
