Remove "DO NOT MERGE THIS BACK TO MAIN" code.

diff --git a/Engine/Source/Programs/UnrealBuildTool/System/Unity.cs b/Engine/Source/Programs/UnrealBuildTool/System/Unity.cs
--- a/Engine/Source/Programs/UnrealBuildTool/System/Unity.cs
+++ b/Engine/Source/Programs/UnrealBuildTool/System/Unity.cs
@@ -277,9 +277,7 @@ namespace UnrealBuildTool
 				AllUnityFiles = CPPUnityFileBuilder.GetUnityFiles();
 			}
 
-			//THIS CHANGE WAS MADE TO FIX THE BUILD IN OUR BRANCH
-			//DO NOT MERGE THIS BACK TO MAIN
-			string PCHHeaderNameInCode = CPPFiles.Count > 0 ? CPPFiles[0].PCHHeaderNameInCode : "";
+			string PCHHeaderNameInCode = CPPFiles[0].PCHHeaderNameInCode;
 			if( CompileEnvironment.Config.PrecompiledHeaderIncludeFilename != null )
 			{
 				PCHHeaderNameInCode = ToolChain.ConvertPath( CompileEnvironment.Config.PrecompiledHeaderIncludeFilename.FullName );
