Remove calls to UEToolChain.GetPlatformToolChain(). iOS and Mac deployment for remote tool chains is now done through UEToolChain.PostBuildSync() rather than UEBuildDeploy.PrepTargetForDeploy().

diff --git a/Engine/Source/Programs/UnrealBuildTool/Configuration/UEBuildBinary.cs b/Engine/Source/Programs/UnrealBuildTool/Configuration/UEBuildBinary.cs
--- a/Engine/Source/Programs/UnrealBuildTool/Configuration/UEBuildBinary.cs
+++ b/Engine/Source/Programs/UnrealBuildTool/Configuration/UEBuildBinary.cs
@@ -549,7 +549,7 @@ namespace UnrealBuildTool
 		/// <param name="CompileEnvironment">The environment to compile the binary in</param>
 		/// <param name="LinkEnvironment">The environment to link the binary in</param>
 		/// <returns></returns>
-		public override IEnumerable<FileItem> Build(IUEToolChain TargetToolChain, CPPEnvironment CompileEnvironment, LinkEnvironment LinkEnvironment)
+		public override IEnumerable<FileItem> Build(IUEToolChain ToolChain, CPPEnvironment CompileEnvironment, LinkEnvironment LinkEnvironment)
 		{
 			// UnrealCodeAnalyzer produces output files only for a specific module.
 			if (BuildConfiguration.bRunUnrealCodeAnalyzer && !(Modules.Any(x => x.Name == BuildConfiguration.UCAModuleToAnalyze)))
@@ -558,10 +558,10 @@ namespace UnrealBuildTool
 			}
 
 			// Setup linking environment.
-			var BinaryLinkEnvironment = SetupBinaryLinkEnvironment(LinkEnvironment, CompileEnvironment);
+			var BinaryLinkEnvironment = SetupBinaryLinkEnvironment(ToolChain, LinkEnvironment, CompileEnvironment);
 
 			// Return linked files.
-			return SetupOutputFiles(TargetToolChain, ref BinaryLinkEnvironment);
+			return SetupOutputFiles(ToolChain, ref BinaryLinkEnvironment);
 		}
 
 		/// <summary>
@@ -660,7 +660,7 @@ namespace UnrealBuildTool
 			return Config.OutputFilePath.FullName;
 		}
 
-		private LinkEnvironment SetupBinaryLinkEnvironment(LinkEnvironment LinkEnvironment, CPPEnvironment CompileEnvironment)
+		private LinkEnvironment SetupBinaryLinkEnvironment(IUEToolChain ToolChain, LinkEnvironment LinkEnvironment, CPPEnvironment CompileEnvironment)
 		{
 			var BinaryLinkEnvironment = LinkEnvironment.DeepCopy();
 			var LinkEnvironmentVisitedModules = new HashSet<UEBuildModule>();
@@ -688,7 +688,7 @@ namespace UnrealBuildTool
 				{
 					// Compile each module.
 					Log.TraceVerbose("Compile module: " + Module.Name);
-					LinkInputFiles = Module.Compile(CompileEnvironment, BinaryCompileEnvironment);
+					LinkInputFiles = Module.Compile(ToolChain, CompileEnvironment, BinaryCompileEnvironment);
 
 					// NOTE: Because of 'Shared PCHs', in monolithic builds the same PCH file may appear as a link input
 					// multiple times for a single binary.  We'll check for that here, and only add it once.  This avoids
@@ -745,7 +745,7 @@ namespace UnrealBuildTool
 			return BinaryLinkEnvironment;
 		}
 
-		private List<FileItem> SetupOutputFiles(IUEToolChain TargetToolChain, ref LinkEnvironment BinaryLinkEnvironment)
+		private List<FileItem> SetupOutputFiles(IUEToolChain ToolChain, ref LinkEnvironment BinaryLinkEnvironment)
 		{
 			// Early exits first
 			if (ProjectFileGenerator.bGenerateProjectFiles)
@@ -782,14 +782,14 @@ namespace UnrealBuildTool
 				if (BinaryLinkEnvironment.Config.Target.Platform != CPPTargetPlatform.Mac && BinaryLinkEnvironment.Config.Target.Platform != CPPTargetPlatform.Linux)
 				{
 					// Create the import library.
-					OutputFiles.AddRange(BinaryLinkEnvironment.LinkExecutable(true));
+					OutputFiles.AddRange(ToolChain.LinkAllFiles(BinaryLinkEnvironment, true));
 				}
 			}
 
 			BinaryLinkEnvironment.Config.bIncludeDependentLibrariesInLibrary = bIncludeDependentLibrariesInLibrary;
 
 			// Link the binary.
-			FileItem[] Executables = BinaryLinkEnvironment.LinkExecutable(false);
+			FileItem[] Executables = ToolChain.LinkAllFiles(BinaryLinkEnvironment, false);
 			OutputFiles.AddRange(Executables);
 
 			// Produce additional console app if requested
@@ -802,12 +802,12 @@ namespace UnrealBuildTool
 				ConsoleAppLinkEvironment.Config.OutputFilePaths = ConsoleAppLinkEvironment.Config.OutputFilePaths.Select(Path => GetAdditionalConsoleAppPath(Path)).ToList();
 
 				// Link the console app executable
-				OutputFiles.AddRange(ConsoleAppLinkEvironment.LinkExecutable(false));
+				OutputFiles.AddRange(ToolChain.LinkAllFiles(ConsoleAppLinkEvironment, false));
 			}
 
 			foreach (var Executable in Executables)
 			{
-				OutputFiles.AddRange(TargetToolChain.PostBuild(Executable, BinaryLinkEnvironment));
+				OutputFiles.AddRange(ToolChain.PostBuild(Executable, BinaryLinkEnvironment));
 			}
 
 			return OutputFiles;
@@ -875,9 +875,7 @@ namespace UnrealBuildTool
 			}
 			ProjectCSharpEnviroment.EnvironmentTargetPlatform = LinkEnvironment.Config.Target.Platform;
 
-			// Currently only supported by windows...
-			UEToolChain.GetPlatformToolChain(CPPTargetPlatform.Win64).CompileCSharpProject(
-				ProjectCSharpEnviroment, Config.ProjectFilePath, Config.OutputFilePath);
+			ToolChain.CompileCSharpProject(ProjectCSharpEnviroment, Config.ProjectFilePath, Config.OutputFilePath);
 
 			return new FileItem[] { FileItem.GetItemByFileReference(Config.OutputFilePath) };
 		}
diff --git a/Engine/Source/Programs/UnrealBuildTool/Configuration/UEBuildModule.cs b/Engine/Source/Programs/UnrealBuildTool/Configuration/UEBuildModule.cs
--- a/Engine/Source/Programs/UnrealBuildTool/Configuration/UEBuildModule.cs
+++ b/Engine/Source/Programs/UnrealBuildTool/Configuration/UEBuildModule.cs
@@ -558,7 +558,7 @@ namespace UnrealBuildTool
 		}
 
 		/** Compiles the module, and returns a list of files output by the compiler. */
-		public abstract List<FileItem> Compile( CPPEnvironment GlobalCompileEnvironment, CPPEnvironment CompileEnvironment );
+		public abstract List<FileItem> Compile( IUEToolChain ToolChain, CPPEnvironment GlobalCompileEnvironment, CPPEnvironment CompileEnvironment );
 		
 		// Object interface.
 		public override string ToString()
@@ -676,7 +676,7 @@ namespace UnrealBuildTool
 		}
 
 		// UEBuildModule interface.
-		public override List<FileItem> Compile(CPPEnvironment GlobalCompileEnvironment, CPPEnvironment CompileEnvironment)
+		public override List<FileItem> Compile(IUEToolChain ToolChain, CPPEnvironment GlobalCompileEnvironment, CPPEnvironment CompileEnvironment)
 		{
 			return new List<FileItem>();
 		}
@@ -876,7 +876,7 @@ namespace UnrealBuildTool
 		}
 
 		// UEBuildModule interface.
-		public override List<FileItem> Compile(CPPEnvironment GlobalCompileEnvironment, CPPEnvironment CompileEnvironment)
+		public override List<FileItem> Compile(IUEToolChain ToolChain, CPPEnvironment GlobalCompileEnvironment, CPPEnvironment CompileEnvironment)
 		{
 			var BuildPlatform = UEBuildPlatform.GetBuildPlatformForCPPTargetPlatform(CompileEnvironment.Config.Target.Platform);
 
@@ -1195,7 +1195,7 @@ namespace UnrealBuildTool
 					if (bModuleUsesUnityBuild)
 					{
 						// unity files generated for only the set of files which share the same PCH environment
-						CPPFilesToBuild = Unity.GenerateUnityCPPs( Target, CPPFilesToBuild, ModulePCHCompileEnvironment, Name );
+						CPPFilesToBuild = Unity.GenerateUnityCPPs( ToolChain, Target, CPPFilesToBuild, ModulePCHCompileEnvironment, Name );
 					}
 
 					// Check if there are enough unity files to warrant pch generation (and we haven't already generated the shared one)
@@ -1207,6 +1207,7 @@ namespace UnrealBuildTool
 							if (SharedPCHHeaderFile == null)
 							{
 								PCHOutput = PrecompileHeaderEnvironment.GeneratePCHCreationAction( 
+									ToolChain,
 									Target,
 									CPPFilesToBuild[0].PCHHeaderNameInCode,
 									ModulePCHEnvironment.PrecompiledHeaderIncludeFilename,
@@ -1232,6 +1233,7 @@ namespace UnrealBuildTool
 									new HashSet<UEBuildModule>());
 
 								PCHOutput = PrecompileHeaderEnvironment.GeneratePCHCreationAction( 
+									ToolChain,
 									Target,
 									CPPFilesToBuild[0].PCHHeaderNameInCode,
 									ModulePCHEnvironment.PrecompiledHeaderIncludeFilename,
@@ -1264,7 +1266,7 @@ namespace UnrealBuildTool
 						CPPCompileEnvironment = ModulePCHCompileEnvironment;
 					}
 
-					LinkInputFiles.AddRange( CPPCompileEnvironment.CompileFiles( Target, CPPFilesToBuild, Name ).ObjectFiles );
+					LinkInputFiles.AddRange( ToolChain.CompileCPPFiles( Target, CPPCompileEnvironment, CPPFilesToBuild, Name ).ObjectFiles );
 					bWasModuleCodeCompiled = true;
 				}
 
@@ -1280,9 +1282,9 @@ namespace UnrealBuildTool
 				var CPPFilesToCompile = SourceFilesToBuild.CPPFiles;
 				if (bModuleUsesUnityBuild)
 				{
-					CPPFilesToCompile = Unity.GenerateUnityCPPs( Target, CPPFilesToCompile, CPPCompileEnvironment, Name );
+					CPPFilesToCompile = Unity.GenerateUnityCPPs( ToolChain, Target, CPPFilesToCompile, CPPCompileEnvironment, Name );
 				}
-				LinkInputFiles.AddRange( CPPCompileEnvironment.CompileFiles( Target, CPPFilesToCompile, Name ).ObjectFiles );
+				LinkInputFiles.AddRange( ToolChain.CompileCPPFiles(Target, CPPCompileEnvironment, CPPFilesToCompile, Name ).ObjectFiles );
 			}
 
 			if (AutoGenerateCppInfo != null && AutoGenerateCppInfo.BuildInfo != null && !CPPCompileEnvironment.bHackHeaderGenerator)
@@ -1295,21 +1297,21 @@ namespace UnrealBuildTool
 					CachePCHUsageForModuleSourceFile(this.Target, CPPCompileEnvironment, GeneratedCppFileItem);
 
 					// @todo ubtmake: Check for ALL other places where we might be injecting .cpp or .rc files for compiling without caching CachedCPPIncludeInfo first (anything platform specific?)
-					LinkInputFiles.AddRange(CPPCompileEnvironment.CompileFiles(Target, new List<FileItem> { GeneratedCppFileItem }, Name).ObjectFiles);
+					LinkInputFiles.AddRange(ToolChain.CompileCPPFiles(Target, CPPCompileEnvironment, new List<FileItem> { GeneratedCppFileItem }, Name).ObjectFiles);
 				}
 			}
 
 			// Compile C files directly.
-			LinkInputFiles.AddRange(CPPCompileEnvironment.CompileFiles( Target, SourceFilesToBuild.CFiles, Name).ObjectFiles);
+			LinkInputFiles.AddRange(ToolChain.CompileCPPFiles(Target, CPPCompileEnvironment, SourceFilesToBuild.CFiles, Name).ObjectFiles);
 
 			// Compile CC files directly.
-			LinkInputFiles.AddRange(CPPCompileEnvironment.CompileFiles( Target, SourceFilesToBuild.CCFiles, Name).ObjectFiles);
+			LinkInputFiles.AddRange(ToolChain.CompileCPPFiles(Target, CPPCompileEnvironment, SourceFilesToBuild.CCFiles, Name).ObjectFiles);
 
 			// Compile MM files directly.
-			LinkInputFiles.AddRange(CPPCompileEnvironment.CompileFiles( Target, SourceFilesToBuild.MMFiles, Name).ObjectFiles);
+			LinkInputFiles.AddRange(ToolChain.CompileCPPFiles(Target, CPPCompileEnvironment, SourceFilesToBuild.MMFiles, Name).ObjectFiles);
 
 			// Compile RC files.
-			LinkInputFiles.AddRange(CPPCompileEnvironment.CompileRCFiles(Target, SourceFilesToBuild.RCFiles).ObjectFiles);
+			LinkInputFiles.AddRange(ToolChain.CompileRCFiles(Target, CPPCompileEnvironment, SourceFilesToBuild.RCFiles).ObjectFiles);
 
 			return LinkInputFiles;
 		}
@@ -1779,7 +1781,7 @@ namespace UnrealBuildTool
 		}
 
 		// UEBuildModule interface.
-		public override List<FileItem> Compile( CPPEnvironment GlobalCompileEnvironment, CPPEnvironment CompileEnvironment )
+		public override List<FileItem> Compile( IUEToolChain ToolChain, CPPEnvironment GlobalCompileEnvironment, CPPEnvironment CompileEnvironment )
 		{
 			var ModuleCLREnvironment = CompileEnvironment.DeepCopy();
 
@@ -1793,7 +1795,7 @@ namespace UnrealBuildTool
 			}
 
 			// Pass the CLR compilation environment to the standard C++ module compilation code.
-			return base.Compile(GlobalCompileEnvironment, ModuleCLREnvironment );
+			return base.Compile(ToolChain, GlobalCompileEnvironment, ModuleCLREnvironment );
 		}
 
 		public override void SetupPrivateLinkEnvironment(
@@ -1894,13 +1896,12 @@ namespace UnrealBuildTool
 		/// <param name="ModuleName">Name of the module this PCH is being generated for</param>
 		/// <param name="bAllowDLLExports">True if we should allow DLLEXPORT definitions for this PCH</param>
 		/// <returns>the compilation output result of the created pch.</returns>
-		public static CPPOutput GeneratePCHCreationAction(UEBuildTarget Target, string PCHHeaderNameInCode, FileItem PrecompiledHeaderIncludeFilename, CPPEnvironment ProjectCPPEnvironment, DirectoryReference OutputDirectory, string ModuleName, bool bAllowDLLExports )
+		public static CPPOutput GeneratePCHCreationAction(IUEToolChain ToolChain, UEBuildTarget Target, string PCHHeaderNameInCode, FileItem PrecompiledHeaderIncludeFilename, CPPEnvironment ProjectCPPEnvironment, DirectoryReference OutputDirectory, string ModuleName, bool bAllowDLLExports )
 		{
 			// Find the header file to be precompiled. Don't skip external headers
 			if (PrecompiledHeaderIncludeFilename.bExists)
 			{
 				// Create a Dummy wrapper around the PCH to avoid problems with #pragma once on clang
-				var ToolChain = UEToolChain.GetPlatformToolChain(ProjectCPPEnvironment.Config.Target.Platform);
 				string PCHGuardDefine = Path.GetFileNameWithoutExtension(PrecompiledHeaderIncludeFilename.AbsolutePath).ToUpper();
 				string LocalPCHHeaderNameInCode = ToolChain.ConvertPath(PrecompiledHeaderIncludeFilename.AbsolutePath);
 				string TmpPCHHeaderContents = String.Format("#ifndef __AUTO_{0}_H__\n#define __AUTO_{0}_H__\n//Last Write: {2}\n#include \"{1}\"\n#endif//__AUTO_{0}_H__", PCHGuardDefine, LocalPCHHeaderNameInCode, PrecompiledHeaderIncludeFilename.LastWriteTime);
@@ -1935,7 +1936,7 @@ namespace UnrealBuildTool
 				Log.TraceVerbose( "Found PCH file \"{0}\".", PrecompiledHeaderIncludeFilename );
 
 				// Create the action to compile the PCH file.
-				return ProjectPCHEnvironment.CompileFiles(Target, new List<FileItem>() { DummyPCH }, ModuleName);
+				return ToolChain.CompileCPPFiles(Target, ProjectPCHEnvironment, new List<FileItem>() { DummyPCH }, ModuleName);
 			}
 			throw new BuildException( "Couldn't find PCH file \"{0}\".", PrecompiledHeaderIncludeFilename );
 		}
diff --git a/Engine/Source/Programs/UnrealBuildTool/Configuration/UEBuildTarget.cs b/Engine/Source/Programs/UnrealBuildTool/Configuration/UEBuildTarget.cs
--- a/Engine/Source/Programs/UnrealBuildTool/Configuration/UEBuildTarget.cs
+++ b/Engine/Source/Programs/UnrealBuildTool/Configuration/UEBuildTarget.cs
@@ -2026,7 +2026,7 @@ namespace UnrealBuildTool
 					// Execute the header tool
 					FileReference ModuleInfoFileName = FileReference.Combine( ProjectIntermediateDirectory, "UnrealHeaderTool.manifest" );
 					ECompilationResult UHTResult = ECompilationResult.OtherCompilationError;
-					if (!ExternalExecution.ExecuteHeaderToolIfNecessary(this, GlobalCompileEnvironment, UObjectModules, ModuleInfoFileName, ref UHTResult))
+					if (!ExternalExecution.ExecuteHeaderToolIfNecessary(TargetToolChain, this, GlobalCompileEnvironment, UObjectModules, ModuleInfoFileName, ref UHTResult))
 					{
 						Log.TraceInformation("UnrealHeaderTool failed for target '" + GetTargetName() + "' (platform: " + Platform.ToString() + ", module info: " + ModuleInfoFileName + ").");
 						return UHTResult;
@@ -2109,7 +2109,7 @@ namespace UnrealBuildTool
 		public void PreBuildSetup(IUEToolChain TargetToolChain)
 		{
 			// Set up the global compile and link environment in GlobalCompileEnvironment and GlobalLinkEnvironment.
-			SetupGlobalEnvironment();
+			SetupGlobalEnvironment(TargetToolChain);
 
 			// Setup the target's modules.
 			SetupModules();
@@ -3081,11 +3081,10 @@ namespace UnrealBuildTool
 		}
 
 		/** Sets up the global compile and link environment for the target. */
-		public virtual void SetupGlobalEnvironment()
+		public virtual void SetupGlobalEnvironment(IUEToolChain ToolChain)
 		{
 			var BuildPlatform = UEBuildPlatform.GetBuildPlatform(Platform);
 
-			var ToolChain = UEToolChain.GetPlatformToolChain(BuildPlatform.GetCPPTargetPlatform(Platform));
 			ToolChain.SetUpGlobalEnvironment();
 
 			// Allow each target type (Game/Editor/Server) to set a default global environment state
diff --git a/Engine/Source/Programs/UnrealBuildTool/IOS/IOSToolChain.cs b/Engine/Source/Programs/UnrealBuildTool/IOS/IOSToolChain.cs
--- a/Engine/Source/Programs/UnrealBuildTool/IOS/IOSToolChain.cs
+++ b/Engine/Source/Programs/UnrealBuildTool/IOS/IOSToolChain.cs
@@ -732,7 +732,7 @@ namespace UnrealBuildTool
 		 * 
 		 * @param Executable FileItem describing the executable to generate debug info for
 		 */
-		static public FileItem GenerateDebugInfo(FileItem Executable)
+		public FileItem GenerateDebugInfo(FileItem Executable)
 		{
 			// Make a file item for the source and destination files
 			string FullDestPathRoot = Executable.AbsolutePath + ".dSYM";
@@ -755,8 +755,7 @@ namespace UnrealBuildTool
 				GenDebugAction.ActionHandler = new Action.BlockingActionHandler(RPCUtilHelper.RPCActionHandler);
 			}
 			
-			IOSToolChain Toolchain = UEToolChain.GetPlatformToolChain(CPPTargetPlatform.IOS) as IOSToolChain;
-			GenDebugAction.WorkingDirectory = Toolchain.GetMacDevSrcRoot();
+			GenDebugAction.WorkingDirectory = GetMacDevSrcRoot();
 			GenDebugAction.CommandPath = "sh";
 
 			// note that the source and dest are switched from a copy command
@@ -1282,6 +1281,42 @@ namespace UnrealBuildTool
 						RPCUtilHelper.CopyDirectory( RemoteDest, LocalDest, RPCUtilHelper.ECopyOptions.None );
 					}
 				}
+
+				// If it is requested, send the app bundle back to the platform executing these commands.
+				if (BuildConfiguration.bCopyAppBundleBackToDevice)
+				{
+					Log.TraceInformation("Copying binaries back to this device...");
+
+					try
+					{
+						string BinaryDir = Path.GetDirectoryName(Target.OutputPath.FullName) + "\\";
+						if (BinaryDir.EndsWith(Target.AppName + "\\Binaries\\IOS\\") && Target.TargetType != TargetRules.TargetType.Game)
+						{
+							BinaryDir = BinaryDir.Replace(Target.TargetType.ToString(), "Game");
+						}
+
+						// Get the app bundle's name
+						string AppFullName = Target.AppName;
+						if (Target.Configuration != UnrealTargetConfiguration.Development)
+						{
+							AppFullName += "-" + Target.Platform.ToString();
+							AppFullName += "-" + Target.Configuration.ToString();
+						}
+
+						foreach (string BinaryPath in BuiltBinaries.Select(x => x.FullName))
+						{
+							if (!BinaryPath.Contains("Dummy"))
+							{
+								RPCUtilHelper.CopyFile(ConvertPath(BinaryPath), BinaryPath, false);
+							}
+						}
+						Log.TraceInformation("Copied binaries successfully.");
+					}
+					catch (Exception)
+					{
+						Log.TraceInformation("Copying binaries back to this device failed.");
+					}
+				}
 			}
 		}
 
diff --git a/Engine/Source/Programs/UnrealBuildTool/IOS/UEDeployIOS.cs b/Engine/Source/Programs/UnrealBuildTool/IOS/UEDeployIOS.cs
--- a/Engine/Source/Programs/UnrealBuildTool/IOS/UEDeployIOS.cs
+++ b/Engine/Source/Programs/UnrealBuildTool/IOS/UEDeployIOS.cs
@@ -701,44 +701,6 @@ namespace UnrealBuildTool.IOS
 			}
 			else
 			{
-				// If it is requested, send the app bundle back to the platform executing these commands.
-				if (BuildConfiguration.bCopyAppBundleBackToDevice)
-				{
-					Log.TraceInformation("Copying binaries back to this device...");
-
-					IOSToolChain Toolchain = UEToolChain.GetPlatformToolChain(CPPTargetPlatform.IOS) as IOSToolChain;
-
-					try
-					{
-						string BinaryDir = Path.GetDirectoryName(InTarget.OutputPath.FullName) + "\\";
-						if (BinaryDir.EndsWith(InTarget.AppName + "\\Binaries\\IOS\\") && InTarget.TargetType != TargetRules.TargetType.Game)
-						{
-							BinaryDir = BinaryDir.Replace(InTarget.TargetType.ToString(), "Game");
-						}
-
-						// Get the app bundle's name
-						string AppFullName = InTarget.AppName;
-						if (InTarget.Configuration != UnrealTargetConfiguration.Development)
-						{
-							AppFullName += "-" + InTarget.Platform.ToString();
-							AppFullName += "-" + InTarget.Configuration.ToString();
-						}
-
-						foreach (string BinaryPath in Toolchain.BuiltBinaries.Select(x => x.FullName))
-						{
-							if (!BinaryPath.Contains("Dummy"))
-							{
-								RPCUtilHelper.CopyFile(Toolchain.ConvertPath(BinaryPath), BinaryPath, false);
-							}
-						}
-						Log.TraceInformation("Copied binaries successfully.");
-					}
-					catch (Exception)
-					{
-						Log.TraceInformation("Copying binaries back to this device failed.");
-					}
-				}
-
 				GeneratePList(ProjectDirectory, bIsUE4Game, GameName, UnrealBuildTool.GetUProjectFile() == null ? "" : Path.GetFileNameWithoutExtension(UnrealBuildTool.GetUProjectFile().FullName), "../../Engine", "");
 			}
 			return true;
diff --git a/Engine/Source/Programs/UnrealBuildTool/Mac/MacToolChain.cs b/Engine/Source/Programs/UnrealBuildTool/Mac/MacToolChain.cs
--- a/Engine/Source/Programs/UnrealBuildTool/Mac/MacToolChain.cs
+++ b/Engine/Source/Programs/UnrealBuildTool/Mac/MacToolChain.cs
@@ -8,6 +8,7 @@ using System.Diagnostics;
 using System.Security.AccessControl;
 using System.Text;
 using System.Linq;
+using Ionic.Zip;
 
 namespace UnrealBuildTool
 {
@@ -59,8 +60,6 @@ namespace UnrealBuildTool
 
 		private static List<FileItem> BundleDependencies = new List<FileItem>();
 
-		public List<string> BuiltBinaries = new List<string>();
-
 		private static void SetupXcodePaths(bool bVerbose)
 		{
 			SelectXcode(ref XcodeDeveloperDir, bVerbose);
@@ -1428,19 +1427,164 @@ namespace UnrealBuildTool
 
 		public override void PreBuildSync()
 		{
-			BuiltBinaries = new List<string>();
-
 			base.PreBuildSync();
 		}
 
-		public override void PostBuildSync(UEBuildTarget Target)
+		public override void PostBuildSync(UEBuildTarget InTarget)
 		{
-			foreach (UEBuildBinary Binary in Target.AppBinaries)
+			if (BuildHostPlatform.Current.Platform != UnrealTargetPlatform.Mac)
 			{
-				BuiltBinaries.Add(Path.GetFullPath(Binary.ToString()));
+				List<string> BuiltBinaries = new List<string>();
+				foreach (UEBuildBinary Binary in InTarget.AppBinaries)
+				{
+					BuiltBinaries.Add(Path.GetFullPath(Binary.ToString()));
+				}
+
+				string IntermediateDirectory = InTarget.EngineIntermediateDirectory.FullName;
+				if (!Directory.Exists(IntermediateDirectory))
+				{
+					IntermediateDirectory = Path.GetFullPath("../../" + InTarget.AppName + "/Intermediate/Build/Mac/" + InTarget.AppName + "/" + InTarget.Configuration);
+					if (!Directory.Exists(IntermediateDirectory))
+					{
+						IntermediateDirectory = Path.GetFullPath("../../Engine/Intermediate/Build/Mac/" + InTarget.AppName + "/" + InTarget.Configuration);
+					}
+				}
+			
+				string FixDylibDepsScript = Path.Combine(IntermediateDirectory, "FixDylibDependencies.sh");
+				string FinalizeAppBundleScript = Path.Combine(IntermediateDirectory, "FinalizeAppBundle.sh");
+
+				string RemoteWorkingDir = "";
+
+				bool bIsStaticLibrary = InTarget.OutputPath.HasExtension(".a");
+
+				if (!bIsStaticLibrary)
+				{
+					// Copy the command scripts to the intermediate on the target Mac.
+					string RemoteFixDylibDepsScript = ConvertPath(Path.GetFullPath(FixDylibDepsScript));
+					RemoteFixDylibDepsScript = RemoteFixDylibDepsScript.Replace("../../../../", "../../");
+					RPCUtilHelper.CopyFile(Path.GetFullPath(FixDylibDepsScript), RemoteFixDylibDepsScript, true);
+
+					if (!InTarget.GlobalLinkEnvironment.Config.bIsBuildingConsoleApplication)
+					{
+						string RemoteFinalizeAppBundleScript = ConvertPath(Path.GetFullPath(FinalizeAppBundleScript));
+						RemoteFinalizeAppBundleScript = RemoteFinalizeAppBundleScript.Replace("../../../../", "../../");
+						RPCUtilHelper.CopyFile(Path.GetFullPath(FinalizeAppBundleScript), RemoteFinalizeAppBundleScript, true);
+					}
+
+
+					// run it remotely
+					RemoteWorkingDir = ConvertPath(Path.GetDirectoryName(Path.GetFullPath(FixDylibDepsScript)));
+
+					Log.TraceInformation("Running FixDylibDependencies.sh...");
+					Hashtable Results = RPCUtilHelper.Command(RemoteWorkingDir, "/bin/sh", "FixDylibDependencies.sh", null);
+					if (Results != null)
+					{
+						string Result = (string)Results["CommandOutput"];
+						if (Result != null)
+						{
+							Log.TraceInformation(Result);
+						}
+					}
+
+					if (!InTarget.GlobalLinkEnvironment.Config.bIsBuildingConsoleApplication)
+					{
+						Log.TraceInformation("Running FinalizeAppBundle.sh...");
+						Results = RPCUtilHelper.Command(RemoteWorkingDir, "/bin/sh", "FinalizeAppBundle.sh", null);
+						if (Results != null)
+						{
+							string Result = (string)Results["CommandOutput"];
+							if (Result != null)
+							{
+								Log.TraceInformation(Result);
+							}
+						}
+					}
+				}
+
+				
+				// If it is requested, send the app bundle back to the platform executing these commands.
+				if (BuildConfiguration.bCopyAppBundleBackToDevice)
+				{
+					Log.TraceInformation("Copying binaries back to this device...");
+
+					try
+					{
+						string BinaryDir = InTarget.OutputPath.Directory + "\\";
+						if (BinaryDir.EndsWith(InTarget.AppName + "\\Binaries\\Mac\\") && InTarget.TargetType != TargetRules.TargetType.Game)
+						{
+							BinaryDir = BinaryDir.Replace(InTarget.TargetType.ToString(), "Game");
+						}
+
+						string RemoteBinariesDir = ConvertPath( BinaryDir );
+						string LocalBinariesDir = BinaryDir;
+						
+						// Get the app bundle's name
+						string AppFullName = InTarget.AppName;
+						if (InTarget.Configuration != InTarget.Rules.UndecoratedConfiguration)
+						{
+							AppFullName += "-" + InTarget.Platform.ToString();
+							AppFullName += "-" + InTarget.Configuration.ToString();
+						}
+
+						if (!InTarget.GlobalLinkEnvironment.Config.bIsBuildingConsoleApplication)
+						{
+							AppFullName += ".app";
+						}
+
+						List<string> NotBundledBinaries = new List<string>();
+						foreach (string BinaryPath in BuiltBinaries)
+						{
+							if (InTarget.GlobalLinkEnvironment.Config.bIsBuildingConsoleApplication || bIsStaticLibrary || !BinaryPath.StartsWith(LocalBinariesDir + AppFullName))
+							{
+								NotBundledBinaries.Add(BinaryPath);
+							}
+						}
+
+						// Zip the app bundle for transferring.
+						if (!InTarget.GlobalLinkEnvironment.Config.bIsBuildingConsoleApplication && !bIsStaticLibrary)
+						{
+							string ZipCommand = "zip -0 -r -y -T \"" + AppFullName + ".zip\" \"" + AppFullName + "\"";
+							RPCUtilHelper.Command(RemoteBinariesDir, ZipCommand, "", null);
+
+							// Copy the AppBundle back to the source machine
+							string LocalZipFileLocation = LocalBinariesDir + AppFullName + ".zip ";
+							string RemoteZipFileLocation = RemoteBinariesDir + AppFullName + ".zip";
+
+							RPCUtilHelper.CopyFile(RemoteZipFileLocation, LocalZipFileLocation, false);
+
+							// Extract the copied app bundle (in zip format) to the local binaries directory
+							using (ZipFile AppBundleZip = ZipFile.Read(LocalZipFileLocation))
+							{
+								foreach (ZipEntry Entry in AppBundleZip)
+								{
+									Entry.Extract(LocalBinariesDir, ExtractExistingFileAction.OverwriteSilently);
+								}
+							}
+
+							// Delete the zip as we no longer need/want it.
+							File.Delete(LocalZipFileLocation);
+							RPCUtilHelper.Command(RemoteBinariesDir, "rm -f \"" + AppFullName + ".zip\"", "", null);
+						}
+
+						if (NotBundledBinaries.Count > 0)
+						{
+
+							foreach (string BinaryPath in NotBundledBinaries)
+							{
+								RPCUtilHelper.CopyFile(ConvertPath(BinaryPath), BinaryPath, false);
+							}
+						}
+
+						Log.TraceInformation("Copied binaries successfully.");
+					}
+					catch (Exception)
+					{
+						Log.TraceInformation("Copying binaries back to this device failed.");
+					}
+				}				
 			}
 
-			base.PostBuildSync(Target);
+			base.PostBuildSync(InTarget);
 		}
 
 		public override ICollection<FileItem> PostBuild(FileItem Executable, LinkEnvironment BinaryLinkEnvironment)
diff --git a/Engine/Source/Programs/UnrealBuildTool/Mac/UEDeployMac.cs b/Engine/Source/Programs/UnrealBuildTool/Mac/UEDeployMac.cs
--- a/Engine/Source/Programs/UnrealBuildTool/Mac/UEDeployMac.cs
+++ b/Engine/Source/Programs/UnrealBuildTool/Mac/UEDeployMac.cs
@@ -25,155 +25,6 @@ namespace UnrealBuildTool.IOS
 		public override bool PrepTargetForDeployment(UEBuildTarget InTarget)
 		{
 			Log.TraceInformation("Deploying now!");
-
-			string IntermediateDirectory = InTarget.EngineIntermediateDirectory.FullName;
-			if (!Directory.Exists(IntermediateDirectory))
-			{
-				IntermediateDirectory = Path.GetFullPath("../../" + InTarget.AppName + "/Intermediate/Build/Mac/" + InTarget.AppName + "/" + InTarget.Configuration);
-				if (!Directory.Exists(IntermediateDirectory))
-				{
-					IntermediateDirectory = Path.GetFullPath("../../Engine/Intermediate/Build/Mac/" + InTarget.AppName + "/" + InTarget.Configuration);
-				}
-			}
-
-			MacToolChain Toolchain = UEToolChain.GetPlatformToolChain(CPPTargetPlatform.Mac) as MacToolChain;
-			
-			string FixDylibDepsScript = Path.Combine(IntermediateDirectory, "FixDylibDependencies.sh");
-			string FinalizeAppBundleScript = Path.Combine(IntermediateDirectory, "FinalizeAppBundle.sh");
-
-			string RemoteWorkingDir = "";
-
-			bool bIsStaticLibrary = InTarget.OutputPath.HasExtension(".a");
-
-			if (BuildHostPlatform.Current.Platform != UnrealTargetPlatform.Mac)
-			{
-				if (!bIsStaticLibrary)
-				{
-					// Copy the command scripts to the intermediate on the target Mac.
-					string RemoteFixDylibDepsScript = Toolchain.ConvertPath(Path.GetFullPath(FixDylibDepsScript));
-					RemoteFixDylibDepsScript = RemoteFixDylibDepsScript.Replace("../../../../", "../../");
-					RPCUtilHelper.CopyFile(Path.GetFullPath(FixDylibDepsScript), RemoteFixDylibDepsScript, true);
-
-					if (!InTarget.GlobalLinkEnvironment.Config.bIsBuildingConsoleApplication)
-					{
-						string RemoteFinalizeAppBundleScript = Toolchain.ConvertPath(Path.GetFullPath(FinalizeAppBundleScript));
-						RemoteFinalizeAppBundleScript = RemoteFinalizeAppBundleScript.Replace("../../../../", "../../");
-						RPCUtilHelper.CopyFile(Path.GetFullPath(FinalizeAppBundleScript), RemoteFinalizeAppBundleScript, true);
-					}
-
-
-					// run it remotely
-					RemoteWorkingDir = Toolchain.ConvertPath(Path.GetDirectoryName(Path.GetFullPath(FixDylibDepsScript)));
-
-					Log.TraceInformation("Running FixDylibDependencies.sh...");
-					Hashtable Results = RPCUtilHelper.Command(RemoteWorkingDir, "/bin/sh", "FixDylibDependencies.sh", null);
-					if (Results != null)
-					{
-						string Result = (string)Results["CommandOutput"];
-						if (Result != null)
-						{
-							Log.TraceInformation(Result);
-						}
-					}
-
-					if (!InTarget.GlobalLinkEnvironment.Config.bIsBuildingConsoleApplication)
-					{
-						Log.TraceInformation("Running FinalizeAppBundle.sh...");
-						Results = RPCUtilHelper.Command(RemoteWorkingDir, "/bin/sh", "FinalizeAppBundle.sh", null);
-						if (Results != null)
-						{
-							string Result = (string)Results["CommandOutput"];
-							if (Result != null)
-							{
-								Log.TraceInformation(Result);
-							}
-						}
-					}
-				}
-
-				
-				// If it is requested, send the app bundle back to the platform executing these commands.
-				if (BuildConfiguration.bCopyAppBundleBackToDevice)
-				{
-					Log.TraceInformation("Copying binaries back to this device...");
-
-					try
-					{
-						string BinaryDir = InTarget.OutputPath.Directory + "\\";
-						if (BinaryDir.EndsWith(InTarget.AppName + "\\Binaries\\Mac\\") && InTarget.TargetType != TargetRules.TargetType.Game)
-						{
-							BinaryDir = BinaryDir.Replace(InTarget.TargetType.ToString(), "Game");
-						}
-
-						string RemoteBinariesDir = Toolchain.ConvertPath( BinaryDir );
-						string LocalBinariesDir = BinaryDir;
-						
-						// Get the app bundle's name
-						string AppFullName = InTarget.AppName;
-						if (InTarget.Configuration != InTarget.Rules.UndecoratedConfiguration)
-						{
-							AppFullName += "-" + InTarget.Platform.ToString();
-							AppFullName += "-" + InTarget.Configuration.ToString();
-						}
-
-						if (!InTarget.GlobalLinkEnvironment.Config.bIsBuildingConsoleApplication)
-						{
-							AppFullName += ".app";
-						}
-
-						List<string> NotBundledBinaries = new List<string>();
-						foreach (string BinaryPath in Toolchain.BuiltBinaries)
-						{
-							if (InTarget.GlobalLinkEnvironment.Config.bIsBuildingConsoleApplication || bIsStaticLibrary || !BinaryPath.StartsWith(LocalBinariesDir + AppFullName))
-							{
-								NotBundledBinaries.Add(BinaryPath);
-							}
-						}
-
-						// Zip the app bundle for transferring.
-						if (!InTarget.GlobalLinkEnvironment.Config.bIsBuildingConsoleApplication && !bIsStaticLibrary)
-						{
-							string ZipCommand = "zip -0 -r -y -T \"" + AppFullName + ".zip\" \"" + AppFullName + "\"";
-							RPCUtilHelper.Command(RemoteBinariesDir, ZipCommand, "", null);
-
-							// Copy the AppBundle back to the source machine
-							string LocalZipFileLocation = LocalBinariesDir + AppFullName + ".zip ";
-							string RemoteZipFileLocation = RemoteBinariesDir + AppFullName + ".zip";
-
-							RPCUtilHelper.CopyFile(RemoteZipFileLocation, LocalZipFileLocation, false);
-
-							// Extract the copied app bundle (in zip format) to the local binaries directory
-							using (ZipFile AppBundleZip = ZipFile.Read(LocalZipFileLocation))
-							{
-								foreach (ZipEntry Entry in AppBundleZip)
-								{
-									Entry.Extract(LocalBinariesDir, ExtractExistingFileAction.OverwriteSilently);
-								}
-							}
-
-							// Delete the zip as we no longer need/want it.
-							File.Delete(LocalZipFileLocation);
-							RPCUtilHelper.Command(RemoteBinariesDir, "rm -f \"" + AppFullName + ".zip\"", "", null);
-						}
-
-						if (NotBundledBinaries.Count > 0)
-						{
-
-							foreach (string BinaryPath in NotBundledBinaries)
-							{
-								RPCUtilHelper.CopyFile(Toolchain.ConvertPath(BinaryPath), BinaryPath, false);
-							}
-						}
-
-						Log.TraceInformation("Copied binaries successfully.");
-					}
-					catch (Exception)
-					{
-						Log.TraceInformation("Copying binaries back to this device failed.");
-					}
-				}				
-			}
-
 			return true;
 		}
 	}
diff --git a/Engine/Source/Programs/UnrealBuildTool/System/CPPEnvironment.cs b/Engine/Source/Programs/UnrealBuildTool/System/CPPEnvironment.cs
--- a/Engine/Source/Programs/UnrealBuildTool/System/CPPEnvironment.cs
+++ b/Engine/Source/Programs/UnrealBuildTool/System/CPPEnvironment.cs
@@ -359,27 +359,6 @@ namespace UnrealBuildTool
 			Config = new CPPEnvironmentConfiguration(InCopyEnvironment.Config);
 		}
 
-		/**
-		 * Creates actions to compile a set of C++ source files.
-		 * @param CPPFiles - The C++ source files to compile.
-		 * @param ModuleName - Name of the module these files are associated with
-		 * @return The object files produced by the actions.
-		 */
-		public CPPOutput CompileFiles(UEBuildTarget Target, List<FileItem> CPPFiles, string ModuleName)
-		{
-			return UEToolChain.GetPlatformToolChain(Config.Target.Platform).CompileCPPFiles(Target, this, CPPFiles, ModuleName);
-		}
-
-		/**
-		 * Creates actions to compile a set of Windows resource script files.
-		 * @param RCFiles - The resource script files to compile.
-		 * @return The compiled resource (.res) files produced by the actions.
-		 */
-		public CPPOutput CompileRCFiles(UEBuildTarget Target, List<FileItem> RCFiles)
-		{
-			return UEToolChain.GetPlatformToolChain(Config.Target.Platform).CompileRCFiles(Target, this, RCFiles);
-		}
-
 		/**
 		 * Whether to use PCH files with the current target
 		 * 
diff --git a/Engine/Source/Programs/UnrealBuildTool/System/ExternalExecution.cs b/Engine/Source/Programs/UnrealBuildTool/System/ExternalExecution.cs
--- a/Engine/Source/Programs/UnrealBuildTool/System/ExternalExecution.cs
+++ b/Engine/Source/Programs/UnrealBuildTool/System/ExternalExecution.cs
@@ -606,7 +606,7 @@ namespace UnrealBuildTool
 		 * Builds and runs the header tool and touches the header directories.
 		 * Performs any early outs if headers need no changes, given the UObject modules, tool path, game name, and configuration
 		 */
-		public static bool ExecuteHeaderToolIfNecessary( UEBuildTarget Target, CPPEnvironment GlobalCompileEnvironment, List<UHTModuleInfo> UObjectModules, FileReference ModuleInfoFileName, ref ECompilationResult UHTResult )
+		public static bool ExecuteHeaderToolIfNecessary( IUEToolChain ToolChain, UEBuildTarget Target, CPPEnvironment GlobalCompileEnvironment, List<UHTModuleInfo> UObjectModules, FileReference ModuleInfoFileName, ref ECompilationResult UHTResult )
 		{
 			if(ProgressWriter.bWriteMarkup)
 			{
@@ -617,9 +617,6 @@ namespace UnrealBuildTool
 				// We never want to try to execute the header tool when we're already trying to build it!
 				var bIsBuildingUHT = Target.GetTargetName().Equals( "UnrealHeaderTool", StringComparison.InvariantCultureIgnoreCase );
 
-				var BuildPlatform = UEBuildPlatform.GetBuildPlatform(Target.Platform);
-				var CppPlatform = BuildPlatform.GetCPPTargetPlatform(Target.Platform);
-				var ToolChain = UEToolChain.GetPlatformToolChain(CppPlatform);
 				var RootLocalPath  = Path.GetFullPath(ProjectFileGenerator.RootRelativePath);
 
 				// check if UHT is out of date
diff --git a/Engine/Source/Programs/UnrealBuildTool/System/LinkEnvironment.cs b/Engine/Source/Programs/UnrealBuildTool/System/LinkEnvironment.cs
--- a/Engine/Source/Programs/UnrealBuildTool/System/LinkEnvironment.cs
+++ b/Engine/Source/Programs/UnrealBuildTool/System/LinkEnvironment.cs
@@ -171,12 +171,6 @@ namespace UnrealBuildTool
 			Config = new LinkEnvironmentConfiguration(InCopyEnvironment.Config);
 		}
 
-		/** Links the input files into an executable. */
-		public FileItem[] LinkExecutable( bool bBuildImportLibraryOnly )
-		{
-			return UEToolChain.GetPlatformToolChain(Config.Target.Platform).LinkAllFiles(this, bBuildImportLibraryOnly);
-		}
-
 		/// <summary>
 		/// Performs a deep copy of this LinkEnvironment object.
 		/// </summary>
diff --git a/Engine/Source/Programs/UnrealBuildTool/System/Unity.cs b/Engine/Source/Programs/UnrealBuildTool/System/Unity.cs
--- a/Engine/Source/Programs/UnrealBuildTool/System/Unity.cs
+++ b/Engine/Source/Programs/UnrealBuildTool/System/Unity.cs
@@ -148,6 +148,7 @@ namespace UnrealBuildTool
 		/// <param name="BaseName">Base name to use for the Unity files</param>
 		/// <returns>The "unity" C++ files.</returns>
 		public static List<FileItem> GenerateUnityCPPs(
+			IUEToolChain ToolChain,
 			UEBuildTarget Target,
 			List<FileItem> CPPFiles, 
 			CPPEnvironment CompileEnvironment,
@@ -156,8 +157,6 @@ namespace UnrealBuildTool
 		{
 			var NewCPPFiles = new List<FileItem>();
 
-			var ToolChain = UEToolChain.GetPlatformToolChain(CompileEnvironment.Config.Target.Platform);
-
 			var BuildPlatform = UEBuildPlatform.GetBuildPlatformForCPPTargetPlatform(CompileEnvironment.Config.Target.Platform);
 
 			// Figure out size of all input files combined. We use this to determine whether to use larger unity threshold or not.
diff --git a/Engine/Source/Programs/UnrealBuildTool/System/UnrealBuildTool.cs b/Engine/Source/Programs/UnrealBuildTool/System/UnrealBuildTool.cs
--- a/Engine/Source/Programs/UnrealBuildTool/System/UnrealBuildTool.cs
+++ b/Engine/Source/Programs/UnrealBuildTool/System/UnrealBuildTool.cs
@@ -1355,7 +1355,9 @@ namespace UnrealBuildTool
                                     BuildConfiguration.bFlushBuildDirOnRemoteMac = false;
 									var TargetDescs = UEBuildTarget.ParseTargetCommandLine( Arguments, GetUProjectFile() );
 			                        UEBuildTarget CheckTarget = UEBuildTarget.CreateTarget( TargetDescs[0] );	// @todo ubtmake: This may not work in assembler only mode.  We don't want to be loading target rules assemblies here either.
-                                    CheckTarget.SetupGlobalEnvironment();
+	                                IUEBuildPlatform BuildPlatform = UEBuildPlatform.GetBuildPlatform(TargetDescs[0].Platform);
+									IUEToolChain ToolChain = UEToolChain.GetPlatformToolChain(BuildPlatform.GetCPPTargetPlatform(TargetDescs[0].Platform);
+                                    CheckTarget.SetupGlobalEnvironment(ToolChain));
                                     if ((CheckTarget.TargetType == TargetRules.TargetType.Game) ||
                                         (CheckTarget.TargetType == TargetRules.TargetType.Server) ||
                                         (CheckTarget.TargetType == TargetRules.TargetType.Client))
@@ -1962,7 +1964,7 @@ namespace UnrealBuildTool
                                         // Execute the header tool
                                         FileReference ModuleInfoFileName = FileReference.Combine(Target.ProjectIntermediateDirectory, "UnrealHeaderTool.manifest");
                                         ECompilationResult UHTResult = ECompilationResult.OtherCompilationError;
-                                        if (!ExternalExecution.ExecuteHeaderToolIfNecessary(Target, GlobalCompileEnvironment:null, UObjectModules:TargetUObjectModules, ModuleInfoFileName:ModuleInfoFileName, UHTResult:ref UHTResult))
+                                        if (!ExternalExecution.ExecuteHeaderToolIfNecessary(ToolChain, Target, GlobalCompileEnvironment:null, UObjectModules:TargetUObjectModules, ModuleInfoFileName:ModuleInfoFileName, UHTResult:ref UHTResult))
                                         {
                                             Log.TraceInformation("UnrealHeaderTool failed for target '" + Target.GetTargetName() + "' (platform: " + Target.Platform.ToString() + ", module info: " + ModuleInfoFileName + ").");
                                             BuildResult = UHTResult;
diff --git a/Engine/Source/Programs/UnrealBuildTool/System/XcodeProjectFileGenerator.cs b/Engine/Source/Programs/UnrealBuildTool/System/XcodeProjectFileGenerator.cs
--- a/Engine/Source/Programs/UnrealBuildTool/System/XcodeProjectFileGenerator.cs
+++ b/Engine/Source/Programs/UnrealBuildTool/System/XcodeProjectFileGenerator.cs
@@ -736,24 +736,23 @@ namespace UnrealBuildTool
 			bool bIsUE4Client = Target.TargetName.Equals("UE4Client", StringComparison.InvariantCultureIgnoreCase);
 			string EngineRelative = "";
 
-			IUEToolChain Toolchain = UEToolChain.GetPlatformToolChain(CPPTargetPlatform.IOS);
 			foreach (string GameFolder in GameFolders.Select(x => x.FullName))
 			{
 				if (File.Exists(Path.Combine(GameFolder, Target.TargetName+".uproject")))
 				{
 					IsAGame = true;
-					GamePath = Toolchain.ConvertPath(GameFolder);
+					GamePath = ConvertPath(GameFolder);
 					break;
 				}
 				else if (!string.IsNullOrEmpty(GameNameFromClientName) && GameFolder.EndsWith(GameNameFromClientName))
 				{
-					GamePath = Toolchain.ConvertPath(GameFolder);
+					GamePath = ConvertPath(GameFolder);
 					break;
 				}
 			}
 
 			EngineRelative = Path.GetFullPath(EngineRelativePath + "/../");
-			EngineRelative = Toolchain.ConvertPath(EngineRelative);
+			EngineRelative = ConvertPath(EngineRelative);
 
 			if (!bGeneratingRocketProjectFiles)
 			{
@@ -765,6 +764,16 @@ namespace UnrealBuildTool
 			AppendSingleConfig(ref Contents, Target, "Shipping", Target.ShippingConfigGuid, PreprocessorDefinitions, HeaderSearchPaths, EngineRelative, GamePath, bIsUE4Game, IsAGame, bIsUE4Client);
 		}
 
+		/// <summary>
+		/// Convert all paths to Apple/Unix format (with forward slashes)
+		/// </summary>
+		/// <param name="InPath">The path to convert</param>
+		/// <returns>The normalized path</returns>
+		private static string ConvertPath(string InPath)
+		{
+			return InPath.Replace("\\", "/");
+		}
+		
 		/// <summary>
 		/// Appends a build configuration list section for specific target.
 		/// </summary>
