Fix misplaced parenthesis

diff --git a/Engine/Source/Programs/UnrealBuildTool/System/UnrealBuildTool.cs b/Engine/Source/Programs/UnrealBuildTool/System/UnrealBuildTool.cs
--- a/Engine/Source/Programs/UnrealBuildTool/System/UnrealBuildTool.cs
+++ b/Engine/Source/Programs/UnrealBuildTool/System/UnrealBuildTool.cs
@@ -1356,8 +1356,8 @@ namespace UnrealBuildTool
 									var TargetDescs = UEBuildTarget.ParseTargetCommandLine( Arguments, GetUProjectFile() );
 			                        UEBuildTarget CheckTarget = UEBuildTarget.CreateTarget( TargetDescs[0] );	// @todo ubtmake: This may not work in assembler only mode.  We don't want to be loading target rules assemblies here either.
 	                                IUEBuildPlatform BuildPlatform = UEBuildPlatform.GetBuildPlatform(TargetDescs[0].Platform);
-									IUEToolChain ToolChain = UEToolChain.GetPlatformToolChain(BuildPlatform.GetCPPTargetPlatform(TargetDescs[0].Platform);
-                                    CheckTarget.SetupGlobalEnvironment(ToolChain));
+									IUEToolChain ToolChain = UEToolChain.GetPlatformToolChain(BuildPlatform.GetCPPTargetPlatform(TargetDescs[0].Platform));
+                                    CheckTarget.SetupGlobalEnvironment(ToolChain);
                                     if ((CheckTarget.TargetType == TargetRules.TargetType.Game) ||
                                         (CheckTarget.TargetType == TargetRules.TargetType.Server) ||
                                         (CheckTarget.TargetType == TargetRules.TargetType.Client))
