Merging - cl# 2704196 - PS4 - Don't add PS4 c# projects if PS4 SDK is not installed

diff --git a/Engine/Source/Programs/UnrealBuildTool/System/ProjectFileGenerator.cs b/Engine/Source/Programs/UnrealBuildTool/System/ProjectFileGenerator.cs
--- a/Engine/Source/Programs/UnrealBuildTool/System/ProjectFileGenerator.cs
+++ b/Engine/Source/Programs/UnrealBuildTool/System/ProjectFileGenerator.cs
@@ -1309,12 +1309,15 @@ namespace UnrealBuildTool
         /// </summary>
         private void AddPS4Projects(MasterProjectFolder Folder)
         {
-            string ProjectFolderName = Path.Combine(EngineRelativePath, "Source", "Programs", "PS4");
-            DirectoryInfo ProjectFolderInfo = new DirectoryInfo(ProjectFolderName);
-            if (ProjectFolderInfo.Exists)
-            {
-                Folder.ChildProjects.Add(AddSimpleCSharpProject("PS4/PS4DevKitUtil"));
-            }
+			if (UEBuildPlatform.BuildPlatformDictionary.ContainsKey(UnrealTargetPlatform.PS4))
+			{
+				string ProjectFolderName = Path.Combine(EngineRelativePath, "Source", "Programs", "PS4");
+				DirectoryInfo ProjectFolderInfo = new DirectoryInfo(ProjectFolderName);
+				if (ProjectFolderInfo.Exists)
+				{
+					Folder.ChildProjects.Add(AddSimpleCSharpProject("PS4/PS4DevKitUtil"));
+				}
+			}
         }
 
 		/// <summary>
