Auto-format UBT for consistent spacing. All default Visual Studio settings, but tabs instead of spaces.

diff --git a/Engine/Source/Programs/UnrealBuildTool/Configuration/UEBuildPlatform.cs b/Engine/Source/Programs/UnrealBuildTool/Configuration/UEBuildPlatform.cs
--- a/Engine/Source/Programs/UnrealBuildTool/Configuration/UEBuildPlatform.cs
+++ b/Engine/Source/Programs/UnrealBuildTool/Configuration/UEBuildPlatform.cs
@@ -31,10 +31,10 @@ namespace UnrealBuildTool
 		 */
 		bool CanUseDistcc();
 
-        /**
-         * If this platform can be compiled with SN-DBS
-         */
-        bool CanUseSNDBS();
+		/**
+		 * If this platform can be compiled with SN-DBS
+		 */
+		bool CanUseSNDBS();
 
 		/**
 		 * Register the platform with the UEBuildPlatform class
@@ -46,7 +46,7 @@ namespace UnrealBuildTool
 		 */
 		void ManageAndValidateSDK();
 
-        /**
+		/**
 		 * Retrieve the CPPTargetPlatform for the given UnrealTargetPlatform
 		 *
 		 * @param InUnrealTargetPlatform The UnrealTargetPlatform being build
@@ -466,13 +466,13 @@ namespace UnrealBuildTool
 			return false;
 		}
 
-        /**
-         *	If this platform can be compiled with SN-DBS
-         */
-        public virtual bool CanUseSNDBS()
-        {
-            return false;
-        }
+		/**
+		 *	If this platform can be compiled with SN-DBS
+		 */
+		public virtual bool CanUseSNDBS()
+		{
+			return false;
+		}
 
 		/**
 		 *	Register the platform with the UEBuildPlatform class
@@ -953,42 +953,42 @@ namespace UnrealBuildTool
 
 			// look at all bool values
 			if (BoolKeys != null) foreach (string Key in BoolKeys)
-			{
-				bool Default = false, Project = false;
-				DefaultIni.GetBool(Section, Key, out Default);
-				ProjIni.GetBool(Section, Key, out Project);
-				if (Default != Project)
 				{
-					Console.WriteLine(Key + " is not set to default. (" + Default + " vs. " + Project + ")");
-					return false;
+					bool Default = false, Project = false;
+					DefaultIni.GetBool(Section, Key, out Default);
+					ProjIni.GetBool(Section, Key, out Project);
+					if (Default != Project)
+					{
+						Console.WriteLine(Key + " is not set to default. (" + Default + " vs. " + Project + ")");
+						return false;
+					}
 				}
-			}
 
 			// look at all int values
 			if (IntKeys != null) foreach (string Key in IntKeys)
-			{
-				int Default = 0, Project = 0;
-				DefaultIni.GetInt32(Section, Key, out Default);
-				ProjIni.GetInt32(Section, Key, out Project);
-				if (Default != Project)
 				{
-					Console.WriteLine(Key + " is not set to default. (" + Default + " vs. " + Project + ")");
-					return false;
+					int Default = 0, Project = 0;
+					DefaultIni.GetInt32(Section, Key, out Default);
+					ProjIni.GetInt32(Section, Key, out Project);
+					if (Default != Project)
+					{
+						Console.WriteLine(Key + " is not set to default. (" + Default + " vs. " + Project + ")");
+						return false;
+					}
 				}
-			}
 
 			// look for all string values
 			if (StringKeys != null) foreach (string Key in StringKeys)
-			{
-				string Default = "", Project = "";
-				DefaultIni.GetString(Section, Key, out Default);
-				ProjIni.GetString(Section, Key, out Project);
-				if (Default != Project)
 				{
-					Console.WriteLine(Key + " is not set to default. (" + Default + " vs. " + Project + ")");
-					return false;
+					string Default = "", Project = "";
+					DefaultIni.GetString(Section, Key, out Default);
+					ProjIni.GetString(Section, Key, out Project);
+					if (Default != Project)
+					{
+						Console.WriteLine(Key + " is not set to default. (" + Default + " vs. " + Project + ")");
+						return false;
+					}
 				}
-			}
 
 			// if we get here, we match all important settings
 			return true;
@@ -1238,7 +1238,7 @@ namespace UnrealBuildTool
 					{
 						Writer.WriteLine(InstalledSDKVersionString);
 						Writer.WriteLine("ManualSDK");
-					}					
+					}
 				}
 			}
 		}
@@ -1294,7 +1294,7 @@ namespace UnrealBuildTool
 			if (!IsAutoSDKSafe())
 			{
 				Console.ForegroundColor = ConsoleColor.Red;
-				LogAutoSDK( GetSDKTargetPlatformName() + " attempted to run SDK hook which could have damaged manual SDK install!");				
+				LogAutoSDK(GetSDKTargetPlatformName() + " attempted to run SDK hook which could have damaged manual SDK install!");
 				Console.ResetColor();
 
 				return false;
@@ -1307,7 +1307,7 @@ namespace UnrealBuildTool
 				if (File.Exists(HookExe))
 				{
 					LogAutoSDK("Running {0} hook {1}", Hook, HookExe);
-					
+
 					// run it
 					Process HookProcess = new Process();
 					HookProcess.StartInfo.WorkingDirectory = SDKDirectory;
@@ -1441,7 +1441,7 @@ namespace UnrealBuildTool
 						{
 							if (PathVar.IndexOf(PathRemove, StringComparison.OrdinalIgnoreCase) >= 0)
 							{
-								LogAutoSDK("Removing Path: '{0}'", PathVar);								
+								LogAutoSDK("Removing Path: '{0}'", PathVar);
 								ModifiedPathVars.Remove(PathVar);
 							}
 						}
@@ -1455,7 +1455,7 @@ namespace UnrealBuildTool
 						{
 							if (String.Compare(PathAdd, PathVar, true) == 0)
 							{
-								LogAutoSDK("Removing Path: '{0}'", PathVar);								
+								LogAutoSDK("Removing Path: '{0}'", PathVar);
 								ModifiedPathVars.Remove(PathVar);
 							}
 						}
@@ -1466,7 +1466,7 @@ namespace UnrealBuildTool
 					{
 						if (!ModifiedPathVars.Contains(PathAdd))
 						{
-							LogAutoSDK("Adding Path: '{0}'", PathAdd);							
+							LogAutoSDK("Adding Path: '{0}'", PathAdd);
 							ModifiedPathVars.Add(PathAdd);
 						}
 					}
diff --git a/Engine/Source/Programs/UnrealBuildTool/Linux/UEBuildLinux.cs b/Engine/Source/Programs/UnrealBuildTool/Linux/UEBuildLinux.cs
--- a/Engine/Source/Programs/UnrealBuildTool/Linux/UEBuildLinux.cs
+++ b/Engine/Source/Programs/UnrealBuildTool/Linux/UEBuildLinux.cs
@@ -8,65 +8,65 @@ using System.IO;
 
 namespace UnrealBuildTool
 {
-    class LinuxPlatform : UEBuildPlatform
-    {
-        /** This is the SDK version we support */
+	class LinuxPlatform : UEBuildPlatform
+	{
+		/** This is the SDK version we support */
 		static private Dictionary<string, string> ExpectedSDKVersions = new Dictionary<string, string>()
 		{
 			{ "x86_64-unknown-linux-gnu", "v6_clang-3.6.0_ld-2.24_glibc-2.12.2" },
 			{ "arm-unknown-linux-gnueabihf", "arm-unknown-linux-gnueabihf_v5_clang-3.5.0-ld-2.23.1-glibc-2.13" },
 		};
 
-        /** Platform name (embeds architecture for now) */
-        static private string TargetPlatformName = "Linux_x64";
+		/** Platform name (embeds architecture for now) */
+		static private string TargetPlatformName = "Linux_x64";
 
-        /** Linux architecture (compiler target triplet) */
+		/** Linux architecture (compiler target triplet) */
 		// FIXME: for now switching between architectures is hard-coded
 		static private string DefaultArchitecture = "x86_64-unknown-linux-gnu";
 		//static private string DefaultArchitecture = "arm-unknown-linux-gnueabihf";
 
-        /** The current architecture */
-        public override string GetActiveArchitecture()
-        {
-            return DefaultArchitecture;
-        }
-
-        /**
-         * Allow the platform to apply architecture-specific name according to its rules
-         */
-        public override string ApplyArchitectureName(string BinaryName)
-        {
-            // Linux ignores architecture-specific names, although it might be worth it to prepend architecture
-            return BinaryName;
-        }
-
-        /** 
-         * Whether platform supports switching SDKs during runtime
-         * 
-         * @return true if supports
-         */
-        protected override bool PlatformSupportsAutoSDKs()
-        {
-            return true;
-        }
-
-        /** 
-         * Returns platform-specific name used in SDK repository
-         * 
-         * @return path to SDK Repository
-         */
-        public override string GetSDKTargetPlatformName()
-        {
-            return TargetPlatformName;
-        }
-
-        /** 
-         * Returns SDK string as required by the platform 
-         * 
-         * @return Valid SDK string
-         */
-        protected override string GetRequiredSDKString()
-        {
+		/** The current architecture */
+		public override string GetActiveArchitecture()
+		{
+			return DefaultArchitecture;
+		}
+
+		/**
+		 * Allow the platform to apply architecture-specific name according to its rules
+		 */
+		public override string ApplyArchitectureName(string BinaryName)
+		{
+			// Linux ignores architecture-specific names, although it might be worth it to prepend architecture
+			return BinaryName;
+		}
+
+		/** 
+		 * Whether platform supports switching SDKs during runtime
+		 * 
+		 * @return true if supports
+		 */
+		protected override bool PlatformSupportsAutoSDKs()
+		{
+			return true;
+		}
+
+		/** 
+		 * Returns platform-specific name used in SDK repository
+		 * 
+		 * @return path to SDK Repository
+		 */
+		public override string GetSDKTargetPlatformName()
+		{
+			return TargetPlatformName;
+		}
+
+		/** 
+		 * Returns SDK string as required by the platform 
+		 * 
+		 * @return Valid SDK string
+		 */
+		protected override string GetRequiredSDKString()
+		{
 			string SDKString;
 			if (!ExpectedSDKVersions.TryGetValue(GetActiveArchitecture(), out SDKString))
 			{
@@ -74,44 +74,44 @@ namespace UnrealBuildTool
 			}
 
 			return SDKString;
-        }
-
-        protected override String GetRequiredScriptVersionString()
-        {
-            return "3.0";
-        }
-
-        protected override bool PreferAutoSDK()
-        {
-            // having LINUX_ROOT set (for legacy reasons or for convenience of cross-compiling certain third party libs) should not make UBT skip AutoSDKs
-            return true;
-        }
-
-        /**
-         *	Whether the required external SDKs are installed for this platform
-         */
-        protected override SDKStatus HasRequiredManualSDKInternal()
-        {
-            if (BuildHostPlatform.Current.Platform == UnrealTargetPlatform.Linux)
-            {
-                return SDKStatus.Valid;
-            }
-
-            string BaseLinuxPath = Environment.GetEnvironmentVariable("LINUX_ROOT");
-
-            // we don't have an LINUX_ROOT specified
-            if (String.IsNullOrEmpty(BaseLinuxPath))
-                return SDKStatus.Invalid;
-
-            // paths to our toolchains
-            BaseLinuxPath = BaseLinuxPath.Replace("\"", "");
-            string ClangPath = Path.Combine(BaseLinuxPath, @"bin\Clang++.exe");
-            
-            if (File.Exists(ClangPath))
-                return SDKStatus.Valid;
-
-            return SDKStatus.Invalid;
-        }
+		}
+
+		protected override String GetRequiredScriptVersionString()
+		{
+			return "3.0";
+		}
+
+		protected override bool PreferAutoSDK()
+		{
+			// having LINUX_ROOT set (for legacy reasons or for convenience of cross-compiling certain third party libs) should not make UBT skip AutoSDKs
+			return true;
+		}
+
+		/**
+		 *	Whether the required external SDKs are installed for this platform
+		 */
+		protected override SDKStatus HasRequiredManualSDKInternal()
+		{
+			if (BuildHostPlatform.Current.Platform == UnrealTargetPlatform.Linux)
+			{
+				return SDKStatus.Valid;
+			}
+
+			string BaseLinuxPath = Environment.GetEnvironmentVariable("LINUX_ROOT");
+
+			// we don't have an LINUX_ROOT specified
+			if (String.IsNullOrEmpty(BaseLinuxPath))
+				return SDKStatus.Invalid;
+
+			// paths to our toolchains
+			BaseLinuxPath = BaseLinuxPath.Replace("\"", "");
+			string ClangPath = Path.Combine(BaseLinuxPath, @"bin\Clang++.exe");
+
+			if (File.Exists(ClangPath))
+				return SDKStatus.Valid;
+
+			return SDKStatus.Invalid;
+		}
 
 		public override bool CanUseXGE()
 		{
@@ -122,287 +122,287 @@ namespace UnrealBuildTool
 		/**
 		 *	Register the platform with the UEBuildPlatform class
 		 */
-        protected override void RegisterBuildPlatformInternal()
-        {
+		protected override void RegisterBuildPlatformInternal()
+		{
 			if ((ProjectFileGenerator.bGenerateProjectFiles == true) || (HasRequiredSDKsInstalled() == SDKStatus.Valid))
-            {
-                bool bRegisterBuildPlatform = true;
-
-                string EngineSourcePath = Path.Combine(ProjectFileGenerator.RootRelativePath, "Engine", "Source");
-                string LinuxTargetPlatformFile = Path.Combine(EngineSourcePath, "Developer", "Linux", "LinuxTargetPlatform", "LinuxTargetPlatform.Build.cs");
-
-                if (File.Exists(LinuxTargetPlatformFile) == false)
-                {
-                    bRegisterBuildPlatform = false;
-                }
-
-                if (bRegisterBuildPlatform == true)
-                {
-                    // Register this build platform for Linux
-                    if (BuildConfiguration.bPrintDebugInfo)
-                    {
-                        Console.WriteLine("        Registering for {0}", UnrealTargetPlatform.Linux.ToString());
-                    }
-                    UEBuildPlatform.RegisterBuildPlatform(UnrealTargetPlatform.Linux, this);
-                    UEBuildPlatform.RegisterPlatformWithGroup(UnrealTargetPlatform.Linux, UnrealPlatformGroup.Unix);
-                }
-            }
-        }
-
-        /**
-         *	Retrieve the CPPTargetPlatform for the given UnrealTargetPlatform
-         *
-         *	@param	InUnrealTargetPlatform		The UnrealTargetPlatform being build
-         *	
-         *	@return	CPPTargetPlatform			The CPPTargetPlatform to compile for
-         */
-        public override CPPTargetPlatform GetCPPTargetPlatform(UnrealTargetPlatform InUnrealTargetPlatform)
-        {
-            switch (InUnrealTargetPlatform)
-            {
-                case UnrealTargetPlatform.Linux:
-                    return CPPTargetPlatform.Linux;
-            }
-            throw new BuildException("LinuxPlatform::GetCPPTargetPlatform: Invalid request for {0}", InUnrealTargetPlatform.ToString());
-        }
-
-        /**
-         *	Get the extension to use for the given binary type
-         *	
-         *	@param	InBinaryType		The binary type being built
-         *	
-         *	@return	string				The binary extension (i.e. 'exe' or 'dll')
-         */
-        public override string GetBinaryExtension(UEBuildBinaryType InBinaryType)
-        {
-            switch (InBinaryType)
-            {
-                case UEBuildBinaryType.DynamicLinkLibrary:
-                    return ".so";
-                case UEBuildBinaryType.Executable:
-                    return "";
-                case UEBuildBinaryType.StaticLibrary:
-                    return ".a";
+			{
+				bool bRegisterBuildPlatform = true;
+
+				string EngineSourcePath = Path.Combine(ProjectFileGenerator.RootRelativePath, "Engine", "Source");
+				string LinuxTargetPlatformFile = Path.Combine(EngineSourcePath, "Developer", "Linux", "LinuxTargetPlatform", "LinuxTargetPlatform.Build.cs");
+
+				if (File.Exists(LinuxTargetPlatformFile) == false)
+				{
+					bRegisterBuildPlatform = false;
+				}
+
+				if (bRegisterBuildPlatform == true)
+				{
+					// Register this build platform for Linux
+					if (BuildConfiguration.bPrintDebugInfo)
+					{
+						Console.WriteLine("        Registering for {0}", UnrealTargetPlatform.Linux.ToString());
+					}
+					UEBuildPlatform.RegisterBuildPlatform(UnrealTargetPlatform.Linux, this);
+					UEBuildPlatform.RegisterPlatformWithGroup(UnrealTargetPlatform.Linux, UnrealPlatformGroup.Unix);
+				}
+			}
+		}
+
+		/**
+		 *	Retrieve the CPPTargetPlatform for the given UnrealTargetPlatform
+		 *
+		 *	@param	InUnrealTargetPlatform		The UnrealTargetPlatform being build
+		 *	
+		 *	@return	CPPTargetPlatform			The CPPTargetPlatform to compile for
+		 */
+		public override CPPTargetPlatform GetCPPTargetPlatform(UnrealTargetPlatform InUnrealTargetPlatform)
+		{
+			switch (InUnrealTargetPlatform)
+			{
+				case UnrealTargetPlatform.Linux:
+					return CPPTargetPlatform.Linux;
+			}
+			throw new BuildException("LinuxPlatform::GetCPPTargetPlatform: Invalid request for {0}", InUnrealTargetPlatform.ToString());
+		}
+
+		/**
+		 *	Get the extension to use for the given binary type
+		 *	
+		 *	@param	InBinaryType		The binary type being built
+		 *	
+		 *	@return	string				The binary extension (i.e. 'exe' or 'dll')
+		 */
+		public override string GetBinaryExtension(UEBuildBinaryType InBinaryType)
+		{
+			switch (InBinaryType)
+			{
+				case UEBuildBinaryType.DynamicLinkLibrary:
+					return ".so";
+				case UEBuildBinaryType.Executable:
+					return "";
+				case UEBuildBinaryType.StaticLibrary:
+					return ".a";
 				case UEBuildBinaryType.Object:
 					return ".o";
 				case UEBuildBinaryType.PrecompiledHeader:
 					return ".gch";
-            }
-            return base.GetBinaryExtension(InBinaryType);
-        }
-
-        /**
-         *	Get the extension to use for debug info for the given binary type
-         *	
-         *	@param	InBinaryType		The binary type being built
-         *	
-         *	@return	string				The debug info extension (i.e. 'pdb')
-         */
-        public override string GetDebugInfoExtension(UEBuildBinaryType InBinaryType)
-        {
-            return "";
-        }
-
-        /**
-         *	Gives the platform a chance to 'override' the configuration settings 
-         *	that are overridden on calls to RunUBT.
-         *	
-         *	@param	InPlatform			The UnrealTargetPlatform being built
-         *	@param	InConfiguration		The UnrealTargetConfiguration being built
-         */
-        public override void ResetBuildConfiguration(UnrealTargetPlatform InPlatform, UnrealTargetConfiguration InConfiguration)
-        {
-            ValidateUEBuildConfiguration();
-        }
-
-        /**
-         *	Validate configuration for this platform
-         *	NOTE: This function can/will modify BuildConfiguration!
-         *	
-         *	@param	InPlatform			The CPPTargetPlatform being built
-         *	@param	InConfiguration		The CPPTargetConfiguration being built
-         *	@param	bInCreateDebugInfo	true if debug info is getting create, false if not
-         */
-        public override void ValidateBuildConfiguration(CPPTargetConfiguration Configuration, CPPTargetPlatform Platform, bool bCreateDebugInfo)
-        {
-            UEBuildConfiguration.bCompileSimplygon = false;
-        }
-
-        /**
-         *	Validate the UEBuildConfiguration for this platform
-         *	This is called BEFORE calling UEBuildConfiguration to allow setting 
-         *	various fields used in that function such as CompileLeanAndMean...
-         */
-        public override void ValidateUEBuildConfiguration()
-        {
-            if (ProjectFileGenerator.bGenerateProjectFiles && !ProjectFileGenerator.bGeneratingRocketProjectFiles)
-            {
-                // When generating non-Rocket project files we need intellisense generator to include info from all modules,
-                // including editor-only third party libs
-                UEBuildConfiguration.bCompileLeanAndMeanUE = false;
-            }
-
-            BuildConfiguration.bUseUnityBuild = true;
-
-            // Don't stop compilation at first error...
-            BuildConfiguration.bStopXGECompilationAfterErrors = true;
-
-            BuildConfiguration.bUseSharedPCHs = false;
-        }
-
-        /**
-         *	Whether PDB files should be used
-         *	
-         *	@param	InPlatform			The CPPTargetPlatform being built
-         *	@param	InConfiguration		The CPPTargetConfiguration being built
-         *	@param	bInCreateDebugInfo	true if debug info is getting create, false if not
-         *	
-         *	@return	bool	true if PDB files should be used, false if not
-         */
-        public override bool ShouldUsePDBFiles(CPPTargetPlatform Platform, CPPTargetConfiguration Configuration, bool bCreateDebugInfo)
-        {
-            return true;
-        }
-
-        /**
-         *	Get a list of extra modules the platform requires.
-         *	This is to allow undisclosed platforms to add modules they need without exposing information about the platfomr.
-         *	
-         *	@param	Target						The target being build
-         *	@param	BuildTarget					The UEBuildTarget getting build
-         *	@param	PlatformExtraModules		OUTPUT the list of extra modules the platform needs to add to the target
-         */
-        public override void GetExtraModules(TargetInfo Target, UEBuildTarget BuildTarget, ref List<string> PlatformExtraModules)
-        {
-        }
-
-        /**
-         *	Modify the newly created module passed in for this platform.
-         *	This is not required - but allows for hiding details of a
-         *	particular platform.
-         *	
+			}
+			return base.GetBinaryExtension(InBinaryType);
+		}
+
+		/**
+		 *	Get the extension to use for debug info for the given binary type
+		 *	
+		 *	@param	InBinaryType		The binary type being built
+		 *	
+		 *	@return	string				The debug info extension (i.e. 'pdb')
+		 */
+		public override string GetDebugInfoExtension(UEBuildBinaryType InBinaryType)
+		{
+			return "";
+		}
+
+		/**
+		 *	Gives the platform a chance to 'override' the configuration settings 
+		 *	that are overridden on calls to RunUBT.
+		 *	
+		 *	@param	InPlatform			The UnrealTargetPlatform being built
+		 *	@param	InConfiguration		The UnrealTargetConfiguration being built
+		 */
+		public override void ResetBuildConfiguration(UnrealTargetPlatform InPlatform, UnrealTargetConfiguration InConfiguration)
+		{
+			ValidateUEBuildConfiguration();
+		}
+
+		/**
+		 *	Validate configuration for this platform
+		 *	NOTE: This function can/will modify BuildConfiguration!
+		 *	
+		 *	@param	InPlatform			The CPPTargetPlatform being built
+		 *	@param	InConfiguration		The CPPTargetConfiguration being built
+		 *	@param	bInCreateDebugInfo	true if debug info is getting create, false if not
+		 */
+		public override void ValidateBuildConfiguration(CPPTargetConfiguration Configuration, CPPTargetPlatform Platform, bool bCreateDebugInfo)
+		{
+			UEBuildConfiguration.bCompileSimplygon = false;
+		}
+
+		/**
+		 *	Validate the UEBuildConfiguration for this platform
+		 *	This is called BEFORE calling UEBuildConfiguration to allow setting 
+		 *	various fields used in that function such as CompileLeanAndMean...
+		 */
+		public override void ValidateUEBuildConfiguration()
+		{
+			if (ProjectFileGenerator.bGenerateProjectFiles && !ProjectFileGenerator.bGeneratingRocketProjectFiles)
+			{
+				// When generating non-Rocket project files we need intellisense generator to include info from all modules,
+				// including editor-only third party libs
+				UEBuildConfiguration.bCompileLeanAndMeanUE = false;
+			}
+
+			BuildConfiguration.bUseUnityBuild = true;
+
+			// Don't stop compilation at first error...
+			BuildConfiguration.bStopXGECompilationAfterErrors = true;
+
+			BuildConfiguration.bUseSharedPCHs = false;
+		}
+
+		/**
+		 *	Whether PDB files should be used
+		 *	
+		 *	@param	InPlatform			The CPPTargetPlatform being built
+		 *	@param	InConfiguration		The CPPTargetConfiguration being built
+		 *	@param	bInCreateDebugInfo	true if debug info is getting create, false if not
+		 *	
+		 *	@return	bool	true if PDB files should be used, false if not
+		 */
+		public override bool ShouldUsePDBFiles(CPPTargetPlatform Platform, CPPTargetConfiguration Configuration, bool bCreateDebugInfo)
+		{
+			return true;
+		}
+
+		/**
+		 *	Get a list of extra modules the platform requires.
+		 *	This is to allow undisclosed platforms to add modules they need without exposing information about the platfomr.
+		 *	
+		 *	@param	Target						The target being build
+		 *	@param	BuildTarget					The UEBuildTarget getting build
+		 *	@param	PlatformExtraModules		OUTPUT the list of extra modules the platform needs to add to the target
+		 */
+		public override void GetExtraModules(TargetInfo Target, UEBuildTarget BuildTarget, ref List<string> PlatformExtraModules)
+		{
+		}
+
+		/**
+		 *	Modify the newly created module passed in for this platform.
+		 *	This is not required - but allows for hiding details of a
+		 *	particular platform.
+		 *	
 		 *  @param  Name			The name of the module
 		 *	@param	Rules			The module rules
 		 *	@param	Target			The target being build
 		 */
 		public override void ModifyModuleRules(string ModuleName, ModuleRules Rules, TargetInfo Target)
-        {
-            if ((Target.Platform == UnrealTargetPlatform.Win32) || (Target.Platform == UnrealTargetPlatform.Win64))
-            {
-                if (!UEBuildConfiguration.bBuildRequiresCookedData)
-                {
-                    if (ModuleName == "Engine")
-                    {
-                        if (UEBuildConfiguration.bBuildDeveloperTools)
-                        {
-                            Rules.PlatformSpecificDynamicallyLoadedModuleNames.Add("LinuxTargetPlatform");
-                            Rules.PlatformSpecificDynamicallyLoadedModuleNames.Add("LinuxNoEditorTargetPlatform");
-                            Rules.PlatformSpecificDynamicallyLoadedModuleNames.Add("LinuxServerTargetPlatform");
-                        }
-                    }
-                }
-
-                // allow standalone tools to use targetplatform modules, without needing Engine
-                if (UEBuildConfiguration.bForceBuildTargetPlatforms && ModuleName == "TargetPlatform")
-                {
-                    Rules.PlatformSpecificDynamicallyLoadedModuleNames.Add("LinuxTargetPlatform");
-                    Rules.PlatformSpecificDynamicallyLoadedModuleNames.Add("LinuxNoEditorTargetPlatform");
-                    Rules.PlatformSpecificDynamicallyLoadedModuleNames.Add("LinuxServerTargetPlatform");
-                }
-            }
-            else if (Target.Platform == UnrealTargetPlatform.Linux)
-            {
-                bool bBuildShaderFormats = UEBuildConfiguration.bForceBuildShaderFormats;
-
-                if (!UEBuildConfiguration.bBuildRequiresCookedData)
-                {
-                    if (ModuleName == "TargetPlatform")
-                    {
-                        bBuildShaderFormats = true;
-                    }
-                }
-
-                // allow standalone tools to use target platform modules, without needing Engine
-				if(ModuleName == "TargetPlatform")
+		{
+			if ((Target.Platform == UnrealTargetPlatform.Win32) || (Target.Platform == UnrealTargetPlatform.Win64))
+			{
+				if (!UEBuildConfiguration.bBuildRequiresCookedData)
 				{
-                    if (UEBuildConfiguration.bForceBuildTargetPlatforms)
-                    {
-                        Rules.DynamicallyLoadedModuleNames.Add("LinuxTargetPlatform");
-                        Rules.DynamicallyLoadedModuleNames.Add("LinuxNoEditorTargetPlatform");
-                        Rules.DynamicallyLoadedModuleNames.Add("LinuxServerTargetPlatform");
-					    Rules.DynamicallyLoadedModuleNames.Add("AllDesktopTargetPlatform");
-                    }
-    
-                    if (bBuildShaderFormats)
-                    {
-					    // Rules.DynamicallyLoadedModuleNames.Add("ShaderFormatD3D");
-                        Rules.DynamicallyLoadedModuleNames.Add("ShaderFormatOpenGL");
-                    }
+					if (ModuleName == "Engine")
+					{
+						if (UEBuildConfiguration.bBuildDeveloperTools)
+						{
+							Rules.PlatformSpecificDynamicallyLoadedModuleNames.Add("LinuxTargetPlatform");
+							Rules.PlatformSpecificDynamicallyLoadedModuleNames.Add("LinuxNoEditorTargetPlatform");
+							Rules.PlatformSpecificDynamicallyLoadedModuleNames.Add("LinuxServerTargetPlatform");
+						}
+					}
 				}
-            }
-        }
-
-        /**
-         *	Setup the target environment for building
-         *	
-         *	@param	InBuildTarget		The target being built
-         */
-        public override void SetUpEnvironment(UEBuildTarget InBuildTarget)
-        {
-            InBuildTarget.GlobalCompileEnvironment.Config.Definitions.Add("PLATFORM_LINUX=1");
-            InBuildTarget.GlobalCompileEnvironment.Config.Definitions.Add("LINUX=1");
-
-            InBuildTarget.GlobalCompileEnvironment.Config.Definitions.Add("WITH_DATABASE_SUPPORT=0");		//@todo linux: valid?
-
-            if (GetActiveArchitecture().StartsWith("arm"))
-            {
-                InBuildTarget.GlobalCompileEnvironment.Config.Definitions.Add("REQUIRES_ALIGNED_INT_ACCESS");
-            }
-
-            // link with Linux libraries.
-            InBuildTarget.GlobalLinkEnvironment.Config.AdditionalLibraries.Add("pthread");
-
-            // Disable Simplygon support if compiling against the NULL RHI.
-            if (InBuildTarget.GlobalCompileEnvironment.Config.Definitions.Contains("USE_NULL_RHI=1"))
-            {
-                UEBuildConfiguration.bCompileSimplygon = false;
-            }
+
+				// allow standalone tools to use targetplatform modules, without needing Engine
+				if (UEBuildConfiguration.bForceBuildTargetPlatforms && ModuleName == "TargetPlatform")
+				{
+					Rules.PlatformSpecificDynamicallyLoadedModuleNames.Add("LinuxTargetPlatform");
+					Rules.PlatformSpecificDynamicallyLoadedModuleNames.Add("LinuxNoEditorTargetPlatform");
+					Rules.PlatformSpecificDynamicallyLoadedModuleNames.Add("LinuxServerTargetPlatform");
+				}
+			}
+			else if (Target.Platform == UnrealTargetPlatform.Linux)
+			{
+				bool bBuildShaderFormats = UEBuildConfiguration.bForceBuildShaderFormats;
+
+				if (!UEBuildConfiguration.bBuildRequiresCookedData)
+				{
+					if (ModuleName == "TargetPlatform")
+					{
+						bBuildShaderFormats = true;
+					}
+				}
+
+				// allow standalone tools to use target platform modules, without needing Engine
+				if (ModuleName == "TargetPlatform")
+				{
+					if (UEBuildConfiguration.bForceBuildTargetPlatforms)
+					{
+						Rules.DynamicallyLoadedModuleNames.Add("LinuxTargetPlatform");
+						Rules.DynamicallyLoadedModuleNames.Add("LinuxNoEditorTargetPlatform");
+						Rules.DynamicallyLoadedModuleNames.Add("LinuxServerTargetPlatform");
+						Rules.DynamicallyLoadedModuleNames.Add("AllDesktopTargetPlatform");
+					}
+
+					if (bBuildShaderFormats)
+					{
+						// Rules.DynamicallyLoadedModuleNames.Add("ShaderFormatD3D");
+						Rules.DynamicallyLoadedModuleNames.Add("ShaderFormatOpenGL");
+					}
+				}
+			}
+		}
+
+		/**
+		 *	Setup the target environment for building
+		 *	
+		 *	@param	InBuildTarget		The target being built
+		 */
+		public override void SetUpEnvironment(UEBuildTarget InBuildTarget)
+		{
+			InBuildTarget.GlobalCompileEnvironment.Config.Definitions.Add("PLATFORM_LINUX=1");
+			InBuildTarget.GlobalCompileEnvironment.Config.Definitions.Add("LINUX=1");
+
+			InBuildTarget.GlobalCompileEnvironment.Config.Definitions.Add("WITH_DATABASE_SUPPORT=0");		//@todo linux: valid?
+
+			if (GetActiveArchitecture().StartsWith("arm"))
+			{
+				InBuildTarget.GlobalCompileEnvironment.Config.Definitions.Add("REQUIRES_ALIGNED_INT_ACCESS");
+			}
+
+			// link with Linux libraries.
+			InBuildTarget.GlobalLinkEnvironment.Config.AdditionalLibraries.Add("pthread");
+
+			// Disable Simplygon support if compiling against the NULL RHI.
+			if (InBuildTarget.GlobalCompileEnvironment.Config.Definitions.Contains("USE_NULL_RHI=1"))
+			{
+				UEBuildConfiguration.bCompileSimplygon = false;
+			}
 
 			if (InBuildTarget.TargetType == TargetRules.TargetType.Server)
 			{
 				// Localization shouldn't be needed on servers by default, and ICU is pretty heavy
 				UEBuildConfiguration.bCompileICU = false;
 			}
-        }
-
-        /**
-         *	Whether this platform should create debug information or not
-         *	
-         *	@param	InPlatform			The UnrealTargetPlatform being built
-         *	@param	InConfiguration		The UnrealTargetConfiguration being built
-         *	
-         *	@return	bool				true if debug info should be generated, false if not
-         */
-        public override bool ShouldCreateDebugInfo(UnrealTargetPlatform Platform, UnrealTargetConfiguration Configuration)
-        {
-            switch (Configuration)
-            {
-                case UnrealTargetConfiguration.Development:
-                case UnrealTargetConfiguration.Shipping:
-                case UnrealTargetConfiguration.Test:
-                case UnrealTargetConfiguration.Debug:
-                default:
-                    return true;
-            };
-        }
-
-        /**
-         *	Setup the binaries for this specific platform.
-         *	
-         *	@param	InBuildTarget		The target being built
-         */
-        public override void SetupBinaries(UEBuildTarget InBuildTarget)
-        {
-        }
-    }
+		}
+
+		/**
+		 *	Whether this platform should create debug information or not
+		 *	
+		 *	@param	InPlatform			The UnrealTargetPlatform being built
+		 *	@param	InConfiguration		The UnrealTargetConfiguration being built
+		 *	
+		 *	@return	bool				true if debug info should be generated, false if not
+		 */
+		public override bool ShouldCreateDebugInfo(UnrealTargetPlatform Platform, UnrealTargetConfiguration Configuration)
+		{
+			switch (Configuration)
+			{
+				case UnrealTargetConfiguration.Development:
+				case UnrealTargetConfiguration.Shipping:
+				case UnrealTargetConfiguration.Test:
+				case UnrealTargetConfiguration.Debug:
+				default:
+					return true;
+			};
+		}
+
+		/**
+		 *	Setup the binaries for this specific platform.
+		 *	
+		 *	@param	InBuildTarget		The target being built
+		 */
+		public override void SetupBinaries(UEBuildTarget InBuildTarget)
+		{
+		}
+	}
 }
