Auto-format UBT for consistent spacing. All default Visual Studio settings, but tabs instead of spaces.

diff --git a/Engine/Source/Programs/UnrealBuildTool/System/UnrealBuildTool.cs b/Engine/Source/Programs/UnrealBuildTool/System/UnrealBuildTool.cs
--- a/Engine/Source/Programs/UnrealBuildTool/System/UnrealBuildTool.cs
+++ b/Engine/Source/Programs/UnrealBuildTool/System/UnrealBuildTool.cs
@@ -14,56 +14,56 @@ using Tools.DotNETCommon;
 
 namespace UnrealBuildTool
 {
-    public partial class UnrealBuildTool
-    {
-        /** Time at which control fist passes to UBT. */
-        static public DateTime StartTime;
+	public partial class UnrealBuildTool
+	{
+		/** Time at which control fist passes to UBT. */
+		static public DateTime StartTime;
 
-        static public string BuildGuid = Guid.NewGuid().ToString();
+		static public string BuildGuid = Guid.NewGuid().ToString();
 
-        /** Total time spent building projects. */
-        static public double TotalBuildProjectTime = 0;
+		/** Total time spent building projects. */
+		static public double TotalBuildProjectTime = 0;
 
-        /** Total time spent compiling. */
-        static public double TotalCompileTime = 0;
+		/** Total time spent compiling. */
+		static public double TotalCompileTime = 0;
 
-        /** Total time spent creating app bundles. */
-        static public double TotalCreateAppBundleTime = 0;
+		/** Total time spent creating app bundles. */
+		static public double TotalCreateAppBundleTime = 0;
 
-        /** Total time spent generating debug information. */
-        static public double TotalGenerateDebugInfoTime = 0;
+		/** Total time spent generating debug information. */
+		static public double TotalGenerateDebugInfoTime = 0;
 
-        /** Total time spent linking. */
-        static public double TotalLinkTime = 0;
+		/** Total time spent linking. */
+		static public double TotalLinkTime = 0;
 
-        /** Total time spent on other actions. */
-        static public double TotalOtherActionsTime = 0;
+		/** Total time spent on other actions. */
+		static public double TotalOtherActionsTime = 0;
 
-        /// How much time was spent scanning for include dependencies for outdated C++ files
-        public static double TotalDeepIncludeScanTime = 0.0;
+		/// How much time was spent scanning for include dependencies for outdated C++ files
+		public static double TotalDeepIncludeScanTime = 0.0;
 
-        /** The command line */
-        static public List<string> CmdLine = new List<string>();
+		/** The command line */
+		static public List<string> CmdLine = new List<string>();
 
-        /** The environment at boot time. */
-        static public System.Collections.IDictionary InitialEnvironment;
+		/** The environment at boot time. */
+		static public System.Collections.IDictionary InitialEnvironment;
 
-        /** Should we skip all extra modules for platforms? */
-        static bool bSkipNonHostPlatforms = false;
+		/** Should we skip all extra modules for platforms? */
+		static bool bSkipNonHostPlatforms = false;
 
-        static UnrealTargetPlatform OnlyPlatformSpecificFor = UnrealTargetPlatform.Unknown;
+		static UnrealTargetPlatform OnlyPlatformSpecificFor = UnrealTargetPlatform.Unknown;
 
-        /** Are we running for Rocket */
-        static public bool bRunningRocket = false;
+		/** Are we running for Rocket */
+		static public bool bRunningRocket = false;
 
-        /** The project file */
-        static FileReference UProjectFile;
+		/** The project file */
+		static FileReference UProjectFile;
 
-        /** The Unreal project file path.  This should always be valid when compiling game projects.  Engine targets and program targets may not have a UProject set. */
-        static bool bHasUProjectFile = false;
+		/** The Unreal project file path.  This should always be valid when compiling game projects.  Engine targets and program targets may not have a UProject set. */
+		static bool bHasUProjectFile = false;
 
-        /** The Unreal project directory.  This should always be valid when compiling game projects.  Engine targets and program targets may not have a UProject set. */
-        static DirectoryReference UProjectPath;
+		/** The Unreal project directory.  This should always be valid when compiling game projects.  Engine targets and program targets may not have a UProject set. */
+		static DirectoryReference UProjectPath;
 
 		/** The full name of the Root UE4 directory */
 		public static readonly DirectoryReference RootDirectory = new DirectoryReference(Path.Combine(Path.GetDirectoryName(Assembly.GetExecutingAssembly().GetOriginalLocation()), "..", "..", ".."));
@@ -74,8 +74,8 @@ namespace UnrealBuildTool
 		/** The full name of the Engine/Source directory */
 		public static readonly DirectoryReference EngineSourceDirectory = DirectoryReference.Combine(EngineDirectory, "Source");
 
-        /** This is set to true during UBT startup, after it is safe to be accessing the IsGatheringBuild and IsAssemblingBuild properties.  Never access this directly. */
-        private static bool bIsSafeToCheckIfGatheringOrAssemblingBuild = false;
+		/** This is set to true during UBT startup, after it is safe to be accessing the IsGatheringBuild and IsAssemblingBuild properties.  Never access this directly. */
+		private static bool bIsSafeToCheckIfGatheringOrAssemblingBuild = false;
 
 		/** True if this run of UBT should only invalidate Makefile. */
 		static public bool bIsInvalidatingMakefilesOnly = false;
@@ -83,74 +83,74 @@ namespace UnrealBuildTool
 		/** Cached array of all the supported target platforms */
 		static public UnrealTargetPlatform[] AllPlatforms = (UnrealTargetPlatform[])Enum.GetValues(typeof(UnrealTargetPlatform));
 
-        /** True if we should gather module dependencies for building in this run.  If this is false, then we'll expect to be loading information from disk about
-            the target's modules before we'll be able to build anything.  One or both of IsGatheringBuild or IsAssemblingBuild must be true. */
-        static public bool IsGatheringBuild
-        {
-            get
-            {
-                if( bIsSafeToCheckIfGatheringOrAssemblingBuild )
-                { 
-                    return bIsGatheringBuild_Unsafe;
-                }
-                else
-                {
-                    throw new BuildException( "Accessed UnrealBuildTool.IsGatheringBuild before it was fully initialized at startup (bIsSafeToCheckIfGatheringOrAssemblingBuild)" );
-                }
-            }
-        }
-        static private bool bIsGatheringBuild_Unsafe = true;
-
-        /** True if we should go ahead and actually compile and link in this run.  If this is false, no actions will be executed and we'll stop after gathering
-            and saving information about dependencies.  One or both of IsGatheringBuild or IsAssemblingBuild must be true. */
-        static public bool IsAssemblingBuild
-        {
-            get
-            {
-                if( bIsSafeToCheckIfGatheringOrAssemblingBuild )
-                { 
-                    return bIsAssemblingBuild_Unsafe;
-                }
-                else
-                {
-                    throw new BuildException( "Accessed UnrealBuildTool.IsGatheringBuild before it was fully initialized at startup (bIsSafeToCheckIfGatheringOrAssemblingBuild)" );
-                }
-            }
-        }
-        static private bool bIsAssemblingBuild_Unsafe = true;
-
-        /** Used when BuildConfiguration.bUseUBTMakefiles is enabled.  If true, it means that our cached includes may not longer be
-            valid (or never were), and we need to recover by forcibly scanning included headers for all build prerequisites to make sure that our
-            cached set of includes is actually correct, before determining which files are outdated. */
-        static public bool bNeedsFullCPPIncludeRescan = false;
-
-
-        /// <summary>
-        /// Returns true if UnrealBuildTool should skip everything other than the host platform as far as finding other modules that need to be compiled.
-        /// </summary>
-        /// <returns>True if running with "-skipnonhostplatforms"</returns>
-        static public bool SkipNonHostPlatforms()
-        {
-            return bSkipNonHostPlatforms;
-        }
-
-        /// <summary>
-        /// Returns the platform we should compile extra modules for, in addition to the host
-        /// </summary>
-        /// <returns>The extra platfrom to compile for with "-onlyhostplatformand="</returns>
-        static public UnrealTargetPlatform GetOnlyPlatformSpecificFor()
-        {
-            return OnlyPlatformSpecificFor;
-        }
-
-        /// <summary>
-        /// Returns true if UnrealBuildTool is running with the "-rocket" option (Rocket mode)
-        /// </summary>
-        /// <returns>True if running with "-rocket"</returns>
-        static public bool RunningRocket()
-        {
-            return bRunningRocket;
-        }
+		/** True if we should gather module dependencies for building in this run.  If this is false, then we'll expect to be loading information from disk about
+			the target's modules before we'll be able to build anything.  One or both of IsGatheringBuild or IsAssemblingBuild must be true. */
+		static public bool IsGatheringBuild
+		{
+			get
+			{
+				if (bIsSafeToCheckIfGatheringOrAssemblingBuild)
+				{
+					return bIsGatheringBuild_Unsafe;
+				}
+				else
+				{
+					throw new BuildException("Accessed UnrealBuildTool.IsGatheringBuild before it was fully initialized at startup (bIsSafeToCheckIfGatheringOrAssemblingBuild)");
+				}
+			}
+		}
+		static private bool bIsGatheringBuild_Unsafe = true;
+
+		/** True if we should go ahead and actually compile and link in this run.  If this is false, no actions will be executed and we'll stop after gathering
+			and saving information about dependencies.  One or both of IsGatheringBuild or IsAssemblingBuild must be true. */
+		static public bool IsAssemblingBuild
+		{
+			get
+			{
+				if (bIsSafeToCheckIfGatheringOrAssemblingBuild)
+				{
+					return bIsAssemblingBuild_Unsafe;
+				}
+				else
+				{
+					throw new BuildException("Accessed UnrealBuildTool.IsGatheringBuild before it was fully initialized at startup (bIsSafeToCheckIfGatheringOrAssemblingBuild)");
+				}
+			}
+		}
+		static private bool bIsAssemblingBuild_Unsafe = true;
+
+		/** Used when BuildConfiguration.bUseUBTMakefiles is enabled.  If true, it means that our cached includes may not longer be
+			valid (or never were), and we need to recover by forcibly scanning included headers for all build prerequisites to make sure that our
+			cached set of includes is actually correct, before determining which files are outdated. */
+		static public bool bNeedsFullCPPIncludeRescan = false;
+
+
+		/// <summary>
+		/// Returns true if UnrealBuildTool should skip everything other than the host platform as far as finding other modules that need to be compiled.
+		/// </summary>
+		/// <returns>True if running with "-skipnonhostplatforms"</returns>
+		static public bool SkipNonHostPlatforms()
+		{
+			return bSkipNonHostPlatforms;
+		}
+
+		/// <summary>
+		/// Returns the platform we should compile extra modules for, in addition to the host
+		/// </summary>
+		/// <returns>The extra platfrom to compile for with "-onlyhostplatformand="</returns>
+		static public UnrealTargetPlatform GetOnlyPlatformSpecificFor()
+		{
+			return OnlyPlatformSpecificFor;
+		}
+
+		/// <summary>
+		/// Returns true if UnrealBuildTool is running with the "-rocket" option (Rocket mode)
+		/// </summary>
+		/// <returns>True if running with "-rocket"</returns>
+		static public bool RunningRocket()
+		{
+			return bRunningRocket;
+		}
 
 		/// <summary>
 		/// Returns true if UnrealBuildTool is running using installed Engine components
@@ -162,38 +162,38 @@ namespace UnrealBuildTool
 			return RunningRocket();
 		}
 
-        /// <summary>
-        /// Returns true if a UProject file is available.  This should always be the case when working with game projects.  For engine or program targets, a UProject will not be available.
-        /// </summary>
-        /// <returns></returns>
-        static public bool HasUProjectFile()
-        {
-            return bHasUProjectFile;
-        }
-
-        /// <summary>
-        /// The Unreal project file path.  This should always be valid when compiling game projects.  Engine targets and program targets may not have a UProject set.
-        /// </summary>
-        /// <returns>The path to the file</returns>
-        static public FileReference GetUProjectFile()
-        {
-            return UProjectFile;
-        }
-
-        /// <summary>
-        /// The Unreal project directory.  This should always be valid when compiling game projects.  Engine targets and program targets may not have a UProject set. */
-        /// </summary>
-        /// <returns>The directory path</returns>
-        static public DirectoryReference GetUProjectPath()
-        {
-            return UProjectPath;
-        }
-
-        static public string GetUBTPath()
-        {
-            string UnrealBuildToolPath = Assembly.GetExecutingAssembly().GetOriginalLocation();
-            return UnrealBuildToolPath;
-        }
+		/// <summary>
+		/// Returns true if a UProject file is available.  This should always be the case when working with game projects.  For engine or program targets, a UProject will not be available.
+		/// </summary>
+		/// <returns></returns>
+		static public bool HasUProjectFile()
+		{
+			return bHasUProjectFile;
+		}
+
+		/// <summary>
+		/// The Unreal project file path.  This should always be valid when compiling game projects.  Engine targets and program targets may not have a UProject set.
+		/// </summary>
+		/// <returns>The path to the file</returns>
+		static public FileReference GetUProjectFile()
+		{
+			return UProjectFile;
+		}
+
+		/// <summary>
+		/// The Unreal project directory.  This should always be valid when compiling game projects.  Engine targets and program targets may not have a UProject set. */
+		/// </summary>
+		/// <returns>The directory path</returns>
+		static public DirectoryReference GetUProjectPath()
+		{
+			return UProjectPath;
+		}
+
+		static public string GetUBTPath()
+		{
+			string UnrealBuildToolPath = Assembly.GetExecutingAssembly().GetOriginalLocation();
+			return UnrealBuildToolPath;
+		}
 
 		// @todo projectfiles: Move this into the ProjectPlatformGeneration class?
 		/// <summary>
@@ -212,39 +212,39 @@ namespace UnrealBuildTool
 				);
 		}
 
-        // @todo projectfiles: Move this into the ProjectPlatformGeneration class?
-        /// <summary>
-        /// IsEditorPlatform
-        /// </summary>
-        /// <param name="InPlatform">The platform of interest</param>
-        /// <returns>True if the given platform is a platform that supports the editor.</returns>
-        static public bool IsEditorPlatform(UnrealTargetPlatform InPlatform)
-        {
-            // Windows 64, Linux and Mac are editor platforms.
-            return (
-                (InPlatform == UnrealTargetPlatform.Win64) ||
-                (InPlatform == UnrealTargetPlatform.Linux) ||
-                (InPlatform == UnrealTargetPlatform.Mac)
-                );
-        }
-
-        // @todo projectfiles: Move this into the ProjectPlatformGeneration class?
-        /// <summary>
-        /// IsServerPlatform
-        /// </summary>
-        /// <param name="InPlatform">The platform of interest</param>
-        /// <returns>True if the given platform is a 'server' platform</returns>
-        static public bool IsServerPlatform(UnrealTargetPlatform InPlatform)
-        {
-            // Windows, Mac, and Linux are desktop platforms.
-            return (
-                (InPlatform == UnrealTargetPlatform.Win64) ||
-                (InPlatform == UnrealTargetPlatform.Win32) ||
-                (InPlatform == UnrealTargetPlatform.Mac) ||
-                (InPlatform == UnrealTargetPlatform.WinRT) ||
-                (InPlatform == UnrealTargetPlatform.Linux)
-                );
-        }
+		// @todo projectfiles: Move this into the ProjectPlatformGeneration class?
+		/// <summary>
+		/// IsEditorPlatform
+		/// </summary>
+		/// <param name="InPlatform">The platform of interest</param>
+		/// <returns>True if the given platform is a platform that supports the editor.</returns>
+		static public bool IsEditorPlatform(UnrealTargetPlatform InPlatform)
+		{
+			// Windows 64, Linux and Mac are editor platforms.
+			return (
+				(InPlatform == UnrealTargetPlatform.Win64) ||
+				(InPlatform == UnrealTargetPlatform.Linux) ||
+				(InPlatform == UnrealTargetPlatform.Mac)
+				);
+		}
+
+		// @todo projectfiles: Move this into the ProjectPlatformGeneration class?
+		/// <summary>
+		/// IsServerPlatform
+		/// </summary>
+		/// <param name="InPlatform">The platform of interest</param>
+		/// <returns>True if the given platform is a 'server' platform</returns>
+		static public bool IsServerPlatform(UnrealTargetPlatform InPlatform)
+		{
+			// Windows, Mac, and Linux are desktop platforms.
+			return (
+				(InPlatform == UnrealTargetPlatform.Win64) ||
+				(InPlatform == UnrealTargetPlatform.Win32) ||
+				(InPlatform == UnrealTargetPlatform.Mac) ||
+				(InPlatform == UnrealTargetPlatform.WinRT) ||
+				(InPlatform == UnrealTargetPlatform.Linux)
+				);
+		}
 
 		// @todo projectfiles: Move this into the ProjectPlatformGeneration class?
 		/// <summary>
@@ -263,96 +263,96 @@ namespace UnrealBuildTool
 				);
 		}
 
-        /// <summary>
-        /// Gets all platforms that satisfies the predicate.
-        /// </summary>
-        /// <param name="OutPlatforms">The list of platform to fill in.</param>
-        /// <param name="Predicate">The predicate to filter the platforms.</param>
-        /// <param name="bCheckValidity">If true, ensure it is a valid platform.</param>
-        /// <returns>True if successful and platforms are found, false if not.</returns>
-        static private bool GetPlatforms(ref List<UnrealTargetPlatform> OutPlatforms, Func<UnrealTargetPlatform, bool> Predicate, bool bCheckValidity = true)
-        {
-            OutPlatforms.Clear();
-            foreach (UnrealTargetPlatform Platform in AllPlatforms)
-            {
-                if (Predicate(Platform))
-                {
-                    var BuildPlatform = UEBuildPlatform.GetBuildPlatform(Platform, true);
-                    if ((bCheckValidity == false) || (BuildPlatform != null))
-                    {
-                        OutPlatforms.Add(Platform);
-                    }
-                }
-            }
-
-            return (OutPlatforms.Count > 0) ? true : false;
-        }
-
-        /// <summary>
-        /// Gets all platforms.
-        /// </summary>
-        /// <param name="OutPlatforms">The list of platform to fill in.</param>
-        /// <param name="bCheckValidity">If true, ensure it is a valid platform.</param>
-        /// <returns>True if successful and platforms are found, false if not.</returns>
-        static public bool GetAllPlatforms(ref List<UnrealTargetPlatform> OutPlatforms, bool bCheckValidity = true)
-        {
-            return GetPlatforms(ref OutPlatforms, (Platform) => { return true; }, bCheckValidity);
-        }
-
-        /// <summary>
-        /// Gets all desktop platforms.
-        /// </summary>
-        /// <param name="OutPlatforms">The list of platform to fill in.</param>
-        /// <param name="bCheckValidity">If true, ensure it is a valid platform.</param>
-        /// <returns>True if successful and platforms are found, false if not.</returns>
-        static public bool GetAllDesktopPlatforms(ref List<UnrealTargetPlatform> OutPlatforms, bool bCheckValidity = true)
-        {
-            return GetPlatforms(ref OutPlatforms, UnrealBuildTool.IsDesktopPlatform, bCheckValidity);
-        }
-
-        /// <summary>
-        /// Gets all editor platforms.
-        /// </summary>
-        /// <param name="OutPlatforms">The list of platform to fill in.</param>
-        /// <param name="bCheckValidity">If true, ensure it is a valid platform.</param>
-        /// <returns>True if successful and platforms are found, false if not.</returns>
-        static public bool GetAllEditorPlatforms(ref List<UnrealTargetPlatform> OutPlatforms, bool bCheckValidity = true)
-        {
-            return GetPlatforms(ref OutPlatforms, UnrealBuildTool.IsEditorPlatform, bCheckValidity);
-        }
-
-        /// <summary>
-        /// Gets all server platforms.
-        /// </summary>
-        /// <param name="OutPlatforms">The list of platform to fill in.</param>
-        /// <param name="bCheckValidity">If true, ensure it is a valid platform.</param>
-        /// <returns>True if successful and platforms are found, false if not.</returns>
-        static public bool GetAllServerPlatforms(ref List<UnrealTargetPlatform> OutPlatforms, bool bCheckValidity = true)
-        {
-            return GetPlatforms(ref OutPlatforms, UnrealBuildTool.IsServerPlatform, bCheckValidity);
-        }
-
-        /// <summary>
-        /// Is this a valid configuration
-        /// Used primarily for Rocket vs non-Rocket
-        /// </summary>
-        /// <param name="InConfiguration"></param>
-        /// <returns>true if valid, false if not</returns>
-        static public bool IsValidConfiguration(UnrealTargetConfiguration InConfiguration)
-        {
-            if (RunningRocket())
-            {
-                if (
-                    (InConfiguration != UnrealTargetConfiguration.Development)
-                    && (InConfiguration != UnrealTargetConfiguration.DebugGame)
-                    && (InConfiguration != UnrealTargetConfiguration.Shipping)
-                    )
-                {
-                    return false;
-                }
-            }
-            return true;
-        }
+		/// <summary>
+		/// Gets all platforms that satisfies the predicate.
+		/// </summary>
+		/// <param name="OutPlatforms">The list of platform to fill in.</param>
+		/// <param name="Predicate">The predicate to filter the platforms.</param>
+		/// <param name="bCheckValidity">If true, ensure it is a valid platform.</param>
+		/// <returns>True if successful and platforms are found, false if not.</returns>
+		static private bool GetPlatforms(ref List<UnrealTargetPlatform> OutPlatforms, Func<UnrealTargetPlatform, bool> Predicate, bool bCheckValidity = true)
+		{
+			OutPlatforms.Clear();
+			foreach (UnrealTargetPlatform Platform in AllPlatforms)
+			{
+				if (Predicate(Platform))
+				{
+					var BuildPlatform = UEBuildPlatform.GetBuildPlatform(Platform, true);
+					if ((bCheckValidity == false) || (BuildPlatform != null))
+					{
+						OutPlatforms.Add(Platform);
+					}
+				}
+			}
+
+			return (OutPlatforms.Count > 0) ? true : false;
+		}
+
+		/// <summary>
+		/// Gets all platforms.
+		/// </summary>
+		/// <param name="OutPlatforms">The list of platform to fill in.</param>
+		/// <param name="bCheckValidity">If true, ensure it is a valid platform.</param>
+		/// <returns>True if successful and platforms are found, false if not.</returns>
+		static public bool GetAllPlatforms(ref List<UnrealTargetPlatform> OutPlatforms, bool bCheckValidity = true)
+		{
+			return GetPlatforms(ref OutPlatforms, (Platform) => { return true; }, bCheckValidity);
+		}
+
+		/// <summary>
+		/// Gets all desktop platforms.
+		/// </summary>
+		/// <param name="OutPlatforms">The list of platform to fill in.</param>
+		/// <param name="bCheckValidity">If true, ensure it is a valid platform.</param>
+		/// <returns>True if successful and platforms are found, false if not.</returns>
+		static public bool GetAllDesktopPlatforms(ref List<UnrealTargetPlatform> OutPlatforms, bool bCheckValidity = true)
+		{
+			return GetPlatforms(ref OutPlatforms, UnrealBuildTool.IsDesktopPlatform, bCheckValidity);
+		}
+
+		/// <summary>
+		/// Gets all editor platforms.
+		/// </summary>
+		/// <param name="OutPlatforms">The list of platform to fill in.</param>
+		/// <param name="bCheckValidity">If true, ensure it is a valid platform.</param>
+		/// <returns>True if successful and platforms are found, false if not.</returns>
+		static public bool GetAllEditorPlatforms(ref List<UnrealTargetPlatform> OutPlatforms, bool bCheckValidity = true)
+		{
+			return GetPlatforms(ref OutPlatforms, UnrealBuildTool.IsEditorPlatform, bCheckValidity);
+		}
+
+		/// <summary>
+		/// Gets all server platforms.
+		/// </summary>
+		/// <param name="OutPlatforms">The list of platform to fill in.</param>
+		/// <param name="bCheckValidity">If true, ensure it is a valid platform.</param>
+		/// <returns>True if successful and platforms are found, false if not.</returns>
+		static public bool GetAllServerPlatforms(ref List<UnrealTargetPlatform> OutPlatforms, bool bCheckValidity = true)
+		{
+			return GetPlatforms(ref OutPlatforms, UnrealBuildTool.IsServerPlatform, bCheckValidity);
+		}
+
+		/// <summary>
+		/// Is this a valid configuration
+		/// Used primarily for Rocket vs non-Rocket
+		/// </summary>
+		/// <param name="InConfiguration"></param>
+		/// <returns>true if valid, false if not</returns>
+		static public bool IsValidConfiguration(UnrealTargetConfiguration InConfiguration)
+		{
+			if (RunningRocket())
+			{
+				if (
+					(InConfiguration != UnrealTargetConfiguration.Development)
+					&& (InConfiguration != UnrealTargetConfiguration.DebugGame)
+					&& (InConfiguration != UnrealTargetConfiguration.Shipping)
+					)
+				{
+					return false;
+				}
+			}
+			return true;
+		}
 
 		/// <summary>
 		/// Is this a valid platform
@@ -364,7 +364,7 @@ namespace UnrealBuildTool
 		{
 			if (RunningRocket())
 			{
-				if(Utils.IsRunningOnMono)
+				if (Utils.IsRunningOnMono)
 				{
 					if (InPlatform != UnrealTargetPlatform.Mac && InPlatform != UnrealTargetPlatform.IOS && InPlatform != UnrealTargetPlatform.Linux)
 					{
@@ -373,7 +373,7 @@ namespace UnrealBuildTool
 				}
 				else
 				{
-					if(InPlatform != UnrealTargetPlatform.Win32 && InPlatform != UnrealTargetPlatform.Win64 && InPlatform != UnrealTargetPlatform.Android)
+					if (InPlatform != UnrealTargetPlatform.Win32 && InPlatform != UnrealTargetPlatform.Win64 && InPlatform != UnrealTargetPlatform.Android)
 					{
 						return false;
 					}
@@ -383,129 +383,129 @@ namespace UnrealBuildTool
 		}
 
 
-        /// <summary>
-        /// Is this a valid platform
-        /// Used primarily for Rocket vs non-Rocket
-        /// </summary>
-        /// <param name="InPlatform"></param>
-        /// <returns>true if valid, false if not</returns>
-        static public bool IsValidDesktopPlatform(UnrealTargetPlatform InPlatform)
-        {
-            switch (BuildHostPlatform.Current.Platform)
-            {
-            case UnrealTargetPlatform.Linux:
-                return InPlatform == UnrealTargetPlatform.Linux;
-            case UnrealTargetPlatform.Mac:
-                return InPlatform == UnrealTargetPlatform.Mac;
-            case UnrealTargetPlatform.Win64:
-                return ((InPlatform == UnrealTargetPlatform.Win32) || (InPlatform == UnrealTargetPlatform.Win64));
-            default:
-                throw new BuildException("Invalid RuntimePlatform:" + BuildHostPlatform.Current.Platform);
-            }
-        }
-
-        /// <summary>
-        /// See if the given string was pass in on the command line
-        /// </summary>
-        /// <param name="InString">The argument to check for (must include '-' if expected)</param>
-        /// <returns>true if the argument is found, false if not</returns>
-        static public bool CommandLineContains(string InString)
-        {
-            string LowercaseInString = InString.ToLowerInvariant();
-            foreach (string Arg in CmdLine)
-            {
-                string LowercaseArg = Arg.ToLowerInvariant();
-                if (LowercaseArg == LowercaseInString)
-                {
-                    return true;
-                }
-            }
-
-            return false;
-        }
-
-        public static void RegisterAllUBTClasses(bool bSkipBuildPlatforms = false)
-        {
-            // Find and register all tool chains and build platforms that are present
-            Assembly UBTAssembly = Assembly.GetExecutingAssembly();
-            if (UBTAssembly != null)
-            {
-                Log.TraceVerbose("Searching for ToolChains, BuildPlatforms, BuildDeploys and ProjectGenerators...");
-
-                List<System.Type> ProjectGeneratorList = new List<System.Type>();
-                var AllTypes = UBTAssembly.GetTypes();
-                if (!bSkipBuildPlatforms)
-                {
-                    // register all build platforms first, since they implement SDK-switching logic that can set environment variables
-                    foreach (var CheckType in AllTypes)
-                    {
-                        if (CheckType.IsClass && !CheckType.IsAbstract)
-                        {
-                            if (Utils.ImplementsInterface<IUEBuildPlatform>(CheckType))
-                            {
-                                Log.TraceVerbose("    Registering build platform: {0}", CheckType.ToString());
-                                var TempInst = (UEBuildPlatform)(UBTAssembly.CreateInstance(CheckType.FullName, true));
-                                TempInst.RegisterBuildPlatform();
-                            }
-                        }
-                    }
-                }
-                // register all the other classes
-                foreach (var CheckType in AllTypes)
-                {
-                    if (CheckType.IsClass && !CheckType.IsAbstract)
-                    {
-                        if (Utils.ImplementsInterface<IUEToolChain>(CheckType))
-                        {
-                            Log.TraceVerbose("    Registering tool chain    : {0}", CheckType.ToString());
-                            var TempInst = (UEToolChain)(UBTAssembly.CreateInstance(CheckType.FullName, true));
-                            TempInst.RegisterToolChain();
-                        }
-                        else if (Utils.ImplementsInterface<IUEBuildDeploy>(CheckType))
-                        {
-                            Log.TraceVerbose("    Registering build deploy  : {0}", CheckType.ToString());
-                            var TempInst = (UEBuildDeploy)(UBTAssembly.CreateInstance(CheckType.FullName, true));
-                            TempInst.RegisterBuildDeploy();
-                        }
-                        else if (CheckType.IsSubclassOf(typeof(UEPlatformProjectGenerator)))
-                        {
-                            ProjectGeneratorList.Add(CheckType);
-                        }
-                    }
-                }
-
-                // We need to process the ProjectGenerators AFTER all BuildPlatforms have been 
-                // registered as they use the BuildPlatform to verify they are legitimate.
-                foreach (var ProjType in ProjectGeneratorList)
-                {
-                    Log.TraceVerbose("    Registering project generator: {0}", ProjType.ToString());
-                    UEPlatformProjectGenerator TempInst = (UEPlatformProjectGenerator)(UBTAssembly.CreateInstance(ProjType.FullName, true));
-                    TempInst.RegisterPlatformProjectGenerator();
-                }
-            }
-        }
-
-        public static void SetProjectFile(FileReference InProjectFile)
-        {
-            if (HasUProjectFile() == false || InProjectFile == null)
-            {
-                UProjectFile = InProjectFile;
+		/// <summary>
+		/// Is this a valid platform
+		/// Used primarily for Rocket vs non-Rocket
+		/// </summary>
+		/// <param name="InPlatform"></param>
+		/// <returns>true if valid, false if not</returns>
+		static public bool IsValidDesktopPlatform(UnrealTargetPlatform InPlatform)
+		{
+			switch (BuildHostPlatform.Current.Platform)
+			{
+				case UnrealTargetPlatform.Linux:
+					return InPlatform == UnrealTargetPlatform.Linux;
+				case UnrealTargetPlatform.Mac:
+					return InPlatform == UnrealTargetPlatform.Mac;
+				case UnrealTargetPlatform.Win64:
+					return ((InPlatform == UnrealTargetPlatform.Win32) || (InPlatform == UnrealTargetPlatform.Win64));
+				default:
+					throw new BuildException("Invalid RuntimePlatform:" + BuildHostPlatform.Current.Platform);
+			}
+		}
+
+		/// <summary>
+		/// See if the given string was pass in on the command line
+		/// </summary>
+		/// <param name="InString">The argument to check for (must include '-' if expected)</param>
+		/// <returns>true if the argument is found, false if not</returns>
+		static public bool CommandLineContains(string InString)
+		{
+			string LowercaseInString = InString.ToLowerInvariant();
+			foreach (string Arg in CmdLine)
+			{
+				string LowercaseArg = Arg.ToLowerInvariant();
+				if (LowercaseArg == LowercaseInString)
+				{
+					return true;
+				}
+			}
+
+			return false;
+		}
+
+		public static void RegisterAllUBTClasses(bool bSkipBuildPlatforms = false)
+		{
+			// Find and register all tool chains and build platforms that are present
+			Assembly UBTAssembly = Assembly.GetExecutingAssembly();
+			if (UBTAssembly != null)
+			{
+				Log.TraceVerbose("Searching for ToolChains, BuildPlatforms, BuildDeploys and ProjectGenerators...");
+
+				List<System.Type> ProjectGeneratorList = new List<System.Type>();
+				var AllTypes = UBTAssembly.GetTypes();
+				if (!bSkipBuildPlatforms)
+				{
+					// register all build platforms first, since they implement SDK-switching logic that can set environment variables
+					foreach (var CheckType in AllTypes)
+					{
+						if (CheckType.IsClass && !CheckType.IsAbstract)
+						{
+							if (Utils.ImplementsInterface<IUEBuildPlatform>(CheckType))
+							{
+								Log.TraceVerbose("    Registering build platform: {0}", CheckType.ToString());
+								var TempInst = (UEBuildPlatform)(UBTAssembly.CreateInstance(CheckType.FullName, true));
+								TempInst.RegisterBuildPlatform();
+							}
+						}
+					}
+				}
+				// register all the other classes
+				foreach (var CheckType in AllTypes)
+				{
+					if (CheckType.IsClass && !CheckType.IsAbstract)
+					{
+						if (Utils.ImplementsInterface<IUEToolChain>(CheckType))
+						{
+							Log.TraceVerbose("    Registering tool chain    : {0}", CheckType.ToString());
+							var TempInst = (UEToolChain)(UBTAssembly.CreateInstance(CheckType.FullName, true));
+							TempInst.RegisterToolChain();
+						}
+						else if (Utils.ImplementsInterface<IUEBuildDeploy>(CheckType))
+						{
+							Log.TraceVerbose("    Registering build deploy  : {0}", CheckType.ToString());
+							var TempInst = (UEBuildDeploy)(UBTAssembly.CreateInstance(CheckType.FullName, true));
+							TempInst.RegisterBuildDeploy();
+						}
+						else if (CheckType.IsSubclassOf(typeof(UEPlatformProjectGenerator)))
+						{
+							ProjectGeneratorList.Add(CheckType);
+						}
+					}
+				}
+
+				// We need to process the ProjectGenerators AFTER all BuildPlatforms have been 
+				// registered as they use the BuildPlatform to verify they are legitimate.
+				foreach (var ProjType in ProjectGeneratorList)
+				{
+					Log.TraceVerbose("    Registering project generator: {0}", ProjType.ToString());
+					UEPlatformProjectGenerator TempInst = (UEPlatformProjectGenerator)(UBTAssembly.CreateInstance(ProjType.FullName, true));
+					TempInst.RegisterPlatformProjectGenerator();
+				}
+			}
+		}
+
+		public static void SetProjectFile(FileReference InProjectFile)
+		{
+			if (HasUProjectFile() == false || InProjectFile == null)
+			{
+				UProjectFile = InProjectFile;
 				if (UProjectFile != null)
-                {
-                    bHasUProjectFile = true;
-                    UProjectPath = UProjectFile.Directory;
-                }
-                else
-                {
-                    bHasUProjectFile = false;
+				{
+					bHasUProjectFile = true;
+					UProjectPath = UProjectFile.Directory;
+				}
+				else
+				{
+					bHasUProjectFile = false;
 					UProjectPath = null;
-                }
-            }
-            else
-            {
-                throw new BuildException("Trying to set the project file to {0} when it is already set to {1}", InProjectFile, UProjectFile);
-            }
-        }
+				}
+			}
+			else
+			{
+				throw new BuildException("Trying to set the project file to {0} when it is already set to {1}", InProjectFile, UProjectFile);
+			}
+		}
 
 		public static void ResetProjectFile()
 		{
@@ -515,501 +515,501 @@ namespace UnrealBuildTool
 			UProjectPath = null;
 		}
 
-        private static bool ParseRocketCommandlineArg(string InArg, ref string OutGameName)
-        {
-            string LowercaseArg = InArg.ToLowerInvariant();
-            
+		private static bool ParseRocketCommandlineArg(string InArg, ref string OutGameName)
+		{
+			string LowercaseArg = InArg.ToLowerInvariant();
+
 			string ProjectArg = null;
-            if (LowercaseArg == "-rocket")
-            {
-                bRunningRocket = true;
-            }
-            else if (LowercaseArg.StartsWith("-project="))
-            {
-				if (BuildHostPlatform.Current.Platform == UnrealTargetPlatform.Linux) 
+			if (LowercaseArg == "-rocket")
+			{
+				bRunningRocket = true;
+			}
+			else if (LowercaseArg.StartsWith("-project="))
+			{
+				if (BuildHostPlatform.Current.Platform == UnrealTargetPlatform.Linux)
 				{
 					ProjectArg = InArg.Substring(9).Trim(new Char[] { ' ', '"', '\'' });
 				}
-				else 
+				else
 				{
 					ProjectArg = InArg.Substring(9);
 				}
-            }
-            else if (LowercaseArg.EndsWith(".uproject"))
-            {
-                ProjectArg = InArg;
-            }
-            else
-            {
-                return false;
-            }
-
-            if (ProjectArg != null && ProjectArg.Length > 0)
-            {
-                bHasUProjectFile = true;
-    
-                string ProjectPath = Path.GetDirectoryName(ProjectArg);
-                OutGameName = Path.GetFileNameWithoutExtension(ProjectArg);
-
-                if (Path.IsPathRooted(ProjectPath) == false)
-                {
-                    if (Directory.Exists(ProjectPath) == false)
-                    {
-                        // If it is *not* rooted, then we potentially received a path that is 
-                        // relative to Engine/Binaries/[PLATFORM] rather than Engine/Source
-                        if (ProjectPath.StartsWith("../") || ProjectPath.StartsWith("..\\"))
-                        {
-                            string AlternativeProjectpath = ProjectPath.Substring(3);
-                            if (Directory.Exists(AlternativeProjectpath))
-                            {
-                                ProjectPath = AlternativeProjectpath;
-                                ProjectArg = ProjectArg.Substring(3);
-                                Debug.Assert(ProjectArg.Length > 0);
-                            }
-                        }
-                    }
-                }
+			}
+			else if (LowercaseArg.EndsWith(".uproject"))
+			{
+				ProjectArg = InArg;
+			}
+			else
+			{
+				return false;
+			}
+
+			if (ProjectArg != null && ProjectArg.Length > 0)
+			{
+				bHasUProjectFile = true;
+
+				string ProjectPath = Path.GetDirectoryName(ProjectArg);
+				OutGameName = Path.GetFileNameWithoutExtension(ProjectArg);
+
+				if (Path.IsPathRooted(ProjectPath) == false)
+				{
+					if (Directory.Exists(ProjectPath) == false)
+					{
+						// If it is *not* rooted, then we potentially received a path that is 
+						// relative to Engine/Binaries/[PLATFORM] rather than Engine/Source
+						if (ProjectPath.StartsWith("../") || ProjectPath.StartsWith("..\\"))
+						{
+							string AlternativeProjectpath = ProjectPath.Substring(3);
+							if (Directory.Exists(AlternativeProjectpath))
+							{
+								ProjectPath = AlternativeProjectpath;
+								ProjectArg = ProjectArg.Substring(3);
+								Debug.Assert(ProjectArg.Length > 0);
+							}
+						}
+					}
+				}
 
 				UProjectFile = new FileReference(ProjectArg);
-                UProjectPath = UProjectFile.Directory;
-            }
-
-            return true;
-        }
-
-        /// <summary>
-        /// This is a safeguard to our initialization order, which tends to be fragile.
-        /// It is not really safe to reference a configuration value until the XmlConfig system loads it (and command lines override it), but non-trivial
-        /// code needs to run before the XmlConfig can load, so it can be hard to tell if a function you are calling in these early 
-        /// stages might unsafely reference a config value.
-        /// 
-        /// This bool is set RIGHT BEFORE the XmlConfig system loads configs, so any references to the config class before that will be trapped.
-        /// If you think you need to change where this value is set, you need to reconsider how soon you are trying to reference config values.
-        /// </summary>
-        public static bool bIsSafeToReferenceConfigurationValues = false;
-
-        /// <summary>
-        /// Paranoid init of the logging system in case someone tries to log super early in startup.
-        /// However, no configuration files have been loaded, and the command line has not been parsed, so 
-        /// certain features like verbose and log file cannot be initialized yet.
-        /// </summary>
-        private static void PreInitLogging()
-        {
-            Log.InitLogging(
-                bLogTimestamps: false,
-                InLogLevel: LogEventType.Log,
-                bLogSeverity: false,
-                bLogSources: false,
-                bColorConsoleOutput: true,
-                TraceListeners: new[] { new ConsoleTraceListener() });
-        }
-
-        /// <summary>
-        /// UBT startup order is fairly fragile, and relies on globals that may or may not be safe to use yet.
-        /// This function is for super early startup stuff that should not access Configuration classes (anything loaded by XmlConfg).
-        /// This should be very minimal startup code.
-        /// </summary>
-        /// <param name="Arguments">Cmdline arguments</param>
-        private static void DoStartupStuffThatCannotAccessConfigurationClasses(string[] Arguments)
-        {
-            // Helpers used for stats tracking.
-            StartTime = DateTime.UtcNow;
-
-            // Do super early log init as a safeguard. We'll re-init with proper config options later.
-            PreInitLogging();
-
-            // ensure we can resolve any external assemblies that are not in the same folder as our assembly.
-            AssemblyUtils.InstallAssemblyResolver(Path.GetDirectoryName(Assembly.GetEntryAssembly().GetOriginalLocation()));
-
-            // Copy off the arguments to allow checking for command-line arguments elsewhere
-            CmdLine = new List<string>(Arguments);
-
-            // Grab the environment.
-            InitialEnvironment = Environment.GetEnvironmentVariables();
-            if (InitialEnvironment.Count < 1)
-            {
-                throw new BuildException("Environment could not be read");
-            }
-
-            // Change the working directory to be the Engine/Source folder. We are likely running from Engine/Binaries/DotNET
-            // This is critical to be done early so any code that relies on the current directory being Engine/Source will work.
-            // UEBuildConfiguration.PostReset is susceptible to this, so we must do this before configs are loaded.
-            string EngineSourceDirectory = Path.Combine(Path.GetDirectoryName(Assembly.GetEntryAssembly().GetOriginalLocation()), "..", "..", "..", "Engine", "Source");
-
-            [REDACTED_EMAIL]: This is a workaround for recompiling game code in editor
-            // The working directory when launching is *not* what we would expect
-            if (Directory.Exists(EngineSourceDirectory) == false)
-            {
-                // We are assuming UBT always runs from <>/Engine/Binaries/DotNET/...
-                EngineSourceDirectory = Assembly.GetExecutingAssembly().GetOriginalLocation();
-                EngineSourceDirectory = EngineSourceDirectory.Replace("\\", "/");
-                Int32 EngineIdx = EngineSourceDirectory.IndexOf("/Engine/Binaries/DotNET/", StringComparison.InvariantCultureIgnoreCase);
-                if (EngineIdx != 0)
-                {
-                    EngineSourceDirectory = Path.Combine(EngineSourceDirectory.Substring(0, EngineIdx), "Engine", "Source");
-                }
-            }
-            if (Directory.Exists(EngineSourceDirectory)) // only set the directory if it exists, this should only happen if we are launching the editor from an artist sync
-            {
-                Directory.SetCurrentDirectory(EngineSourceDirectory);
-            }
-        }
-
-        /// <summary>
-        /// Enough startup work has been done to load configuration classes (like setting the proper current directory).
-        /// It is still not safe to log. Mostly this relates to XmlConfigLoader, but any other code added here cannot log yet either.
-        /// </summary>
-        /// <param name="Arguments"></param>
-        private static void InitConfigurationClasses(string[] Arguments)
-        {
-            // Load the Config XML files first.
-            XmlConfigLoader.Init();
-            
-            // Then let the command lines override any configs necessary.
-            const string UCAUsageThresholdString = "-ucausagethreshold=";
-            const string UCAModuleToAnalyzeString = "-ucamoduletoanalyze=";
-            foreach (var LowercaseArg in Arguments.Select(arg => arg.ToLowerInvariant()))
-            {
-                if (LowercaseArg == "-verbose")
-                {
-                    BuildConfiguration.bPrintDebugInfo = true;
-                    BuildConfiguration.LogLevel = "Verbose";
-                }
-                else if (LowercaseArg.StartsWith("-log="))
-                {
-                    BuildConfiguration.LogFilename = LowercaseArg.Replace("-log=","");
-                }
-                else if (LowercaseArg.StartsWith("-logdetailedcompilertiminginfo="))
-                {
-                    WindowsPlatform.bLogDetailedCompilerTimingInfo = bool.Parse(LowercaseArg.Replace("-logdetailedcompilertiminginfo=",""));
-                }
-                else if (LowercaseArg == "-forceheadergeneration")
-                {
-                    UEBuildConfiguration.bForceHeaderGeneration = true;
-                }
-                else if (LowercaseArg == "-nobuilduht")
-                {
-                    UEBuildConfiguration.bDoNotBuildUHT = true;
-                }
-                else if (LowercaseArg == "-failifgeneratedcodechanges")
-                {
-                    UEBuildConfiguration.bFailIfGeneratedCodeChanges = true;
-                }
-                else if (LowercaseArg == "-clean")
-                {
-                    UEBuildConfiguration.bCleanProject = true;
-                }
-                else if (LowercaseArg == "-prepfordeploy")
-                {
-                    UEBuildConfiguration.bPrepForDeployment = true;
-                }
-                else if (LowercaseArg == "-generateexternalfilelist")
-                {
-                    UEBuildConfiguration.bGenerateExternalFileList = true;
-                }
-                else if (LowercaseArg == "-mergeexternalfilelist")
-                {
-                    UEBuildConfiguration.bMergeExternalFileList = true;
-                }
-                else if (LowercaseArg == "-generatemanifest")
-                {
-                    // Generate a manifest file containing all the files required to be in Perforce
-                    UEBuildConfiguration.bGenerateManifest = true;
-                }
-                else if (LowercaseArg == "-mergemanifests")
-                {
-                    // Whether to add to the existing manifest (if it exists), or start afresh
-                    UEBuildConfiguration.bMergeManifests = true;
-                }
-                else if (LowercaseArg == "-noxge")
-                {
-                    BuildConfiguration.bAllowXGE = false;
-                }
-                else if (LowercaseArg == "-noxgemonitor")
-                {
-                    BuildConfiguration.bShowXGEMonitor = false;
-                }
-                else if (LowercaseArg == "-xgeexport")
-                {
-                    BuildConfiguration.bXGEExport = true;
-                    BuildConfiguration.bAllowXGE = true;
-                }
-                else if (LowercaseArg == "-stresstestunity")
-                {
-                    BuildConfiguration.bStressTestUnity = true;
-                }
-                else if (LowercaseArg == "-disableunity")
-                {
-                    BuildConfiguration.bUseUnityBuild = false;
-                }
-                else if (LowercaseArg == "-forceunity")
-                {
-                    BuildConfiguration.bForceUnityBuild = true;
-                }
-                else if (LowercaseArg == "-shadowvariableerrors")
-                {
-                    BuildConfiguration.bShadowVariableErrors = true;
-                }
-                // New Monolithic Graphics drivers have optional "fast calls" replacing various D3d functions
-                else if (LowercaseArg == "-fastmonocalls")
-                {
-                    BuildConfiguration.bUseFastMonoCalls = true;
-                }
-                else if (LowercaseArg == "-nofastmonocalls")
-                {
-                    BuildConfiguration.bUseFastMonoCalls = false;
-                }
-                else if (LowercaseArg == "-uniqueintermediate")
-                {
-                    BuildConfiguration.BaseIntermediateFolder = "Intermediate/Build/" + BuildGuid + "/";
-                }
-                else if (LowercaseArg == "-nopch")
-                {
-                    BuildConfiguration.bUsePCHFiles = false;
-                }
-                else if (LowercaseArg == "-nosharedpch")
-                {
-                    BuildConfiguration.bUseSharedPCHs = false;
-                }
-                else if (LowercaseArg == "-skipActionHistory")
-                {
-                    BuildConfiguration.bUseActionHistory = false;
-                }
-                else if (LowercaseArg == "-noltcg")
-                {
-                    BuildConfiguration.bAllowLTCG = false;
-                }
-                else if (LowercaseArg == "-nopdb")
-                {
-                    BuildConfiguration.bUsePDBFiles = false;
-                }
-                else if (LowercaseArg == "-deploy")
-                {
-                    BuildConfiguration.bDeployAfterCompile = true;
-                }
-                else if (LowercaseArg == "-copyappbundlebacktodevice")
-                {
-                    BuildConfiguration.bCopyAppBundleBackToDevice = true;
-                }
-                else if (LowercaseArg == "-flushmac")
-                {
-                    BuildConfiguration.bFlushBuildDirOnRemoteMac = true;
-                }
-                else if (LowercaseArg == "-nodebuginfo")
-                {
-                    BuildConfiguration.bDisableDebugInfo = true;
-                }
-                else if (LowercaseArg == "-forcedebuginfo")
-                {
-                    BuildConfiguration.bDisableDebugInfo = false;
-                    BuildConfiguration.bOmitPCDebugInfoInDevelopment = false;
-                }
-                else if (LowercaseArg == "-noubtmakefiles")
-                {
-                    BuildConfiguration.bUseUBTMakefiles = false;
-                }
-                else if (LowercaseArg == "-nosimplygon")
-                {
-                    UEBuildConfiguration.bCompileSimplygon = false;
-                }
-                else if (LowercaseArg == "-nospeedtree")
-                {
-                    UEBuildConfiguration.bCompileSpeedTree = false;
-                }
-                else if (LowercaseArg == "-canskiplink")
-                {
-                    UEBuildConfiguration.bSkipLinkingWhenNothingToCompile = true;
-                }
-                else if (LowercaseArg == "-nocef3")
-                {
-                    UEBuildConfiguration.bCompileCEF3 = false;
-                }
-                else if (LowercaseArg == "-rununrealcodeanalyzer")
-                {
-                    BuildConfiguration.bRunUnrealCodeAnalyzer = true;
-                }
-                else if (LowercaseArg.StartsWith(UCAUsageThresholdString))
-                {
-                    string UCAUsageThresholdValue = LowercaseArg.Substring(UCAUsageThresholdString.Length);
-                    BuildConfiguration.UCAUsageThreshold = float.Parse(UCAUsageThresholdValue.Replace(',', '.'), System.Globalization.CultureInfo.InvariantCulture);
-                }
-                else if (LowercaseArg.StartsWith(UCAModuleToAnalyzeString))
-                {
-                    BuildConfiguration.UCAModuleToAnalyze = LowercaseArg.Substring(UCAModuleToAnalyzeString.Length).Trim();
-                }
+				UProjectPath = UProjectFile.Directory;
+			}
+
+			return true;
+		}
+
+		/// <summary>
+		/// This is a safeguard to our initialization order, which tends to be fragile.
+		/// It is not really safe to reference a configuration value until the XmlConfig system loads it (and command lines override it), but non-trivial
+		/// code needs to run before the XmlConfig can load, so it can be hard to tell if a function you are calling in these early 
+		/// stages might unsafely reference a config value.
+		/// 
+		/// This bool is set RIGHT BEFORE the XmlConfig system loads configs, so any references to the config class before that will be trapped.
+		/// If you think you need to change where this value is set, you need to reconsider how soon you are trying to reference config values.
+		/// </summary>
+		public static bool bIsSafeToReferenceConfigurationValues = false;
+
+		/// <summary>
+		/// Paranoid init of the logging system in case someone tries to log super early in startup.
+		/// However, no configuration files have been loaded, and the command line has not been parsed, so 
+		/// certain features like verbose and log file cannot be initialized yet.
+		/// </summary>
+		private static void PreInitLogging()
+		{
+			Log.InitLogging(
+				bLogTimestamps: false,
+				InLogLevel: LogEventType.Log,
+				bLogSeverity: false,
+				bLogSources: false,
+				bColorConsoleOutput: true,
+				TraceListeners: new[] { new ConsoleTraceListener() });
+		}
+
+		/// <summary>
+		/// UBT startup order is fairly fragile, and relies on globals that may or may not be safe to use yet.
+		/// This function is for super early startup stuff that should not access Configuration classes (anything loaded by XmlConfg).
+		/// This should be very minimal startup code.
+		/// </summary>
+		/// <param name="Arguments">Cmdline arguments</param>
+		private static void DoStartupStuffThatCannotAccessConfigurationClasses(string[] Arguments)
+		{
+			// Helpers used for stats tracking.
+			StartTime = DateTime.UtcNow;
+
+			// Do super early log init as a safeguard. We'll re-init with proper config options later.
+			PreInitLogging();
+
+			// ensure we can resolve any external assemblies that are not in the same folder as our assembly.
+			AssemblyUtils.InstallAssemblyResolver(Path.GetDirectoryName(Assembly.GetEntryAssembly().GetOriginalLocation()));
+
+			// Copy off the arguments to allow checking for command-line arguments elsewhere
+			CmdLine = new List<string>(Arguments);
+
+			// Grab the environment.
+			InitialEnvironment = Environment.GetEnvironmentVariables();
+			if (InitialEnvironment.Count < 1)
+			{
+				throw new BuildException("Environment could not be read");
+			}
+
+			// Change the working directory to be the Engine/Source folder. We are likely running from Engine/Binaries/DotNET
+			// This is critical to be done early so any code that relies on the current directory being Engine/Source will work.
+			// UEBuildConfiguration.PostReset is susceptible to this, so we must do this before configs are loaded.
+			string EngineSourceDirectory = Path.Combine(Path.GetDirectoryName(Assembly.GetEntryAssembly().GetOriginalLocation()), "..", "..", "..", "Engine", "Source");
+
+			[REDACTED_EMAIL]: This is a workaround for recompiling game code in editor
+			// The working directory when launching is *not* what we would expect
+			if (Directory.Exists(EngineSourceDirectory) == false)
+			{
+				// We are assuming UBT always runs from <>/Engine/Binaries/DotNET/...
+				EngineSourceDirectory = Assembly.GetExecutingAssembly().GetOriginalLocation();
+				EngineSourceDirectory = EngineSourceDirectory.Replace("\\", "/");
+				Int32 EngineIdx = EngineSourceDirectory.IndexOf("/Engine/Binaries/DotNET/", StringComparison.InvariantCultureIgnoreCase);
+				if (EngineIdx != 0)
+				{
+					EngineSourceDirectory = Path.Combine(EngineSourceDirectory.Substring(0, EngineIdx), "Engine", "Source");
+				}
+			}
+			if (Directory.Exists(EngineSourceDirectory)) // only set the directory if it exists, this should only happen if we are launching the editor from an artist sync
+			{
+				Directory.SetCurrentDirectory(EngineSourceDirectory);
+			}
+		}
+
+		/// <summary>
+		/// Enough startup work has been done to load configuration classes (like setting the proper current directory).
+		/// It is still not safe to log. Mostly this relates to XmlConfigLoader, but any other code added here cannot log yet either.
+		/// </summary>
+		/// <param name="Arguments"></param>
+		private static void InitConfigurationClasses(string[] Arguments)
+		{
+			// Load the Config XML files first.
+			XmlConfigLoader.Init();
+
+			// Then let the command lines override any configs necessary.
+			const string UCAUsageThresholdString = "-ucausagethreshold=";
+			const string UCAModuleToAnalyzeString = "-ucamoduletoanalyze=";
+			foreach (var LowercaseArg in Arguments.Select(arg => arg.ToLowerInvariant()))
+			{
+				if (LowercaseArg == "-verbose")
+				{
+					BuildConfiguration.bPrintDebugInfo = true;
+					BuildConfiguration.LogLevel = "Verbose";
+				}
+				else if (LowercaseArg.StartsWith("-log="))
+				{
+					BuildConfiguration.LogFilename = LowercaseArg.Replace("-log=", "");
+				}
+				else if (LowercaseArg.StartsWith("-logdetailedcompilertiminginfo="))
+				{
+					WindowsPlatform.bLogDetailedCompilerTimingInfo = bool.Parse(LowercaseArg.Replace("-logdetailedcompilertiminginfo=", ""));
+				}
+				else if (LowercaseArg == "-forceheadergeneration")
+				{
+					UEBuildConfiguration.bForceHeaderGeneration = true;
+				}
+				else if (LowercaseArg == "-nobuilduht")
+				{
+					UEBuildConfiguration.bDoNotBuildUHT = true;
+				}
+				else if (LowercaseArg == "-failifgeneratedcodechanges")
+				{
+					UEBuildConfiguration.bFailIfGeneratedCodeChanges = true;
+				}
+				else if (LowercaseArg == "-clean")
+				{
+					UEBuildConfiguration.bCleanProject = true;
+				}
+				else if (LowercaseArg == "-prepfordeploy")
+				{
+					UEBuildConfiguration.bPrepForDeployment = true;
+				}
+				else if (LowercaseArg == "-generateexternalfilelist")
+				{
+					UEBuildConfiguration.bGenerateExternalFileList = true;
+				}
+				else if (LowercaseArg == "-mergeexternalfilelist")
+				{
+					UEBuildConfiguration.bMergeExternalFileList = true;
+				}
+				else if (LowercaseArg == "-generatemanifest")
+				{
+					// Generate a manifest file containing all the files required to be in Perforce
+					UEBuildConfiguration.bGenerateManifest = true;
+				}
+				else if (LowercaseArg == "-mergemanifests")
+				{
+					// Whether to add to the existing manifest (if it exists), or start afresh
+					UEBuildConfiguration.bMergeManifests = true;
+				}
+				else if (LowercaseArg == "-noxge")
+				{
+					BuildConfiguration.bAllowXGE = false;
+				}
+				else if (LowercaseArg == "-noxgemonitor")
+				{
+					BuildConfiguration.bShowXGEMonitor = false;
+				}
+				else if (LowercaseArg == "-xgeexport")
+				{
+					BuildConfiguration.bXGEExport = true;
+					BuildConfiguration.bAllowXGE = true;
+				}
+				else if (LowercaseArg == "-stresstestunity")
+				{
+					BuildConfiguration.bStressTestUnity = true;
+				}
+				else if (LowercaseArg == "-disableunity")
+				{
+					BuildConfiguration.bUseUnityBuild = false;
+				}
+				else if (LowercaseArg == "-forceunity")
+				{
+					BuildConfiguration.bForceUnityBuild = true;
+				}
+				else if (LowercaseArg == "-shadowvariableerrors")
+				{
+					BuildConfiguration.bShadowVariableErrors = true;
+				}
+				// New Monolithic Graphics drivers have optional "fast calls" replacing various D3d functions
+				else if (LowercaseArg == "-fastmonocalls")
+				{
+					BuildConfiguration.bUseFastMonoCalls = true;
+				}
+				else if (LowercaseArg == "-nofastmonocalls")
+				{
+					BuildConfiguration.bUseFastMonoCalls = false;
+				}
+				else if (LowercaseArg == "-uniqueintermediate")
+				{
+					BuildConfiguration.BaseIntermediateFolder = "Intermediate/Build/" + BuildGuid + "/";
+				}
+				else if (LowercaseArg == "-nopch")
+				{
+					BuildConfiguration.bUsePCHFiles = false;
+				}
+				else if (LowercaseArg == "-nosharedpch")
+				{
+					BuildConfiguration.bUseSharedPCHs = false;
+				}
+				else if (LowercaseArg == "-skipActionHistory")
+				{
+					BuildConfiguration.bUseActionHistory = false;
+				}
+				else if (LowercaseArg == "-noltcg")
+				{
+					BuildConfiguration.bAllowLTCG = false;
+				}
+				else if (LowercaseArg == "-nopdb")
+				{
+					BuildConfiguration.bUsePDBFiles = false;
+				}
+				else if (LowercaseArg == "-deploy")
+				{
+					BuildConfiguration.bDeployAfterCompile = true;
+				}
+				else if (LowercaseArg == "-copyappbundlebacktodevice")
+				{
+					BuildConfiguration.bCopyAppBundleBackToDevice = true;
+				}
+				else if (LowercaseArg == "-flushmac")
+				{
+					BuildConfiguration.bFlushBuildDirOnRemoteMac = true;
+				}
+				else if (LowercaseArg == "-nodebuginfo")
+				{
+					BuildConfiguration.bDisableDebugInfo = true;
+				}
+				else if (LowercaseArg == "-forcedebuginfo")
+				{
+					BuildConfiguration.bDisableDebugInfo = false;
+					BuildConfiguration.bOmitPCDebugInfoInDevelopment = false;
+				}
+				else if (LowercaseArg == "-noubtmakefiles")
+				{
+					BuildConfiguration.bUseUBTMakefiles = false;
+				}
+				else if (LowercaseArg == "-nosimplygon")
+				{
+					UEBuildConfiguration.bCompileSimplygon = false;
+				}
+				else if (LowercaseArg == "-nospeedtree")
+				{
+					UEBuildConfiguration.bCompileSpeedTree = false;
+				}
+				else if (LowercaseArg == "-canskiplink")
+				{
+					UEBuildConfiguration.bSkipLinkingWhenNothingToCompile = true;
+				}
+				else if (LowercaseArg == "-nocef3")
+				{
+					UEBuildConfiguration.bCompileCEF3 = false;
+				}
+				else if (LowercaseArg == "-rununrealcodeanalyzer")
+				{
+					BuildConfiguration.bRunUnrealCodeAnalyzer = true;
+				}
+				else if (LowercaseArg.StartsWith(UCAUsageThresholdString))
+				{
+					string UCAUsageThresholdValue = LowercaseArg.Substring(UCAUsageThresholdString.Length);
+					BuildConfiguration.UCAUsageThreshold = float.Parse(UCAUsageThresholdValue.Replace(',', '.'), System.Globalization.CultureInfo.InvariantCulture);
+				}
+				else if (LowercaseArg.StartsWith(UCAModuleToAnalyzeString))
+				{
+					BuildConfiguration.UCAModuleToAnalyze = LowercaseArg.Substring(UCAModuleToAnalyzeString.Length).Trim();
+				}
 				else if (LowercaseArg == "-ucacheckuobjectthreadsafety")
 				{
 					BuildConfiguration.bUCACheckUObjectThreadSafety = true;
 				}
 				else if (LowercaseArg == "-ucacheckpchfiles")
 				{
-					BuildConfiguration.bUCACheckPCHFiles = true;
+					BuildConfiguration.bUCACheckPCHFiles = true;
+				}
+			}
+		}
+
+		/// <summary>
+		/// Sets up UBT Trace listeners and filters.
+		/// </summary>
+		private static void InitLogging()
+		{
+			// Initialize the shared log system using UBT specific listeners.
+			Log.InitLogging(
+				bLogTimestamps: false,
+				InLogLevel: (LogEventType)Enum.Parse(typeof(LogEventType), BuildConfiguration.LogLevel),
+				bLogSeverity: false,
+				bLogSources: false,
+				bColorConsoleOutput: true,
+				TraceListeners: new[] 
+                {
+                    new ConsoleTraceListener(),
+                    !string.IsNullOrEmpty(BuildConfiguration.LogFilename) ? new TextWriterTraceListener(new StreamWriter(new FileStream(BuildConfiguration.LogFilename, FileMode.Create, FileAccess.ReadWrite, FileShare.Read)) { AutoFlush = true }) : null,
+                });
+		}
+
+		private static int DoPostStartupStuffThatCanAccessConfigs(string[] Arguments)
+		{
+			// We can now do a full initialization of the logging system with proper configuration values.
+			InitLogging();
+
+			// Parse rocket-specific arguments.
+			foreach (string Arg in Arguments)
+			{
+				string TempGameName = null;
+				if (ParseRocketCommandlineArg(Arg, ref TempGameName) == true)
+				{
+					// This is to allow relative paths for the project file
+					Log.TraceVerbose("UBT Running for Rocket: " + UProjectFile);
 				}
-            }
-        }
-
-        /// <summary>
-        /// Sets up UBT Trace listeners and filters.
-        /// </summary>
-        private static void InitLogging()
-        {
-            // Initialize the shared log system using UBT specific listeners.
-            Log.InitLogging(
-                bLogTimestamps: false,
-                InLogLevel: (LogEventType)Enum.Parse(typeof(LogEventType), BuildConfiguration.LogLevel),
-                bLogSeverity: false, 
-                bLogSources: false,
-                bColorConsoleOutput: true,
-                TraceListeners: new[] 
-                {
-                    new ConsoleTraceListener(),
-                    !string.IsNullOrEmpty(BuildConfiguration.LogFilename) ? new TextWriterTraceListener(new StreamWriter(new FileStream(BuildConfiguration.LogFilename, FileMode.Create, FileAccess.ReadWrite, FileShare.Read)) { AutoFlush = true }) : null,
-                });
-        }
-
-        private static int DoPostStartupStuffThatCanAccessConfigs(string[] Arguments)
-        {
-            // We can now do a full initialization of the logging system with proper configuration values.
-            InitLogging();
-
-            // Parse rocket-specific arguments.
-            foreach (string Arg in Arguments)
-            {
-                string TempGameName = null;
-                if (ParseRocketCommandlineArg(Arg, ref TempGameName) == true)
-                {
-                    // This is to allow relative paths for the project file
-                    Log.TraceVerbose("UBT Running for Rocket: " + UProjectFile);
-                }
-            }
-
-            // Build the list of game projects that we know about. When building from the editor (for hot-reload) or for Rocket projects, we require the 
-            // project file to be passed in. Otherwise we scan for projects in directories named in UE4Games.uprojectdirs.
-            if (HasUProjectFile())
-            {
-                UProjectInfo.AddProject(GetUProjectFile());
-            }
-            else
-            {
-                UProjectInfo.FillProjectInfo();
-            }
+			}
+
+			// Build the list of game projects that we know about. When building from the editor (for hot-reload) or for Rocket projects, we require the 
+			// project file to be passed in. Otherwise we scan for projects in directories named in UE4Games.uprojectdirs.
+			if (HasUProjectFile())
+			{
+				UProjectInfo.AddProject(GetUProjectFile());
+			}
+			else
+			{
+				UProjectInfo.FillProjectInfo();
+			}
 
 			var BasicInitStartTime = DateTime.UtcNow;
 
-            ECompilationResult Result = ECompilationResult.Succeeded;
+			ECompilationResult Result = ECompilationResult.Succeeded;
 
-            Log.TraceVerbose("UnrealBuildTool (DEBUG OUTPUT MODE)");
-            Log.TraceVerbose("Command-line: {0}", String.Join(" ", Arguments));
+			Log.TraceVerbose("UnrealBuildTool (DEBUG OUTPUT MODE)");
+			Log.TraceVerbose("Command-line: {0}", String.Join(" ", Arguments));
 
-            Telemetry.Initialize();
+			Telemetry.Initialize();
 
-            bool bRunCopyrightVerification = false;
-            bool bDumpToFile = false;
-            bool bCheckThirdPartyHeaders = false;
-			bool bIgnoreJunk = false;                    
+			bool bRunCopyrightVerification = false;
+			bool bDumpToFile = false;
+			bool bCheckThirdPartyHeaders = false;
+			bool bIgnoreJunk = false;
 
-            // @todo: Ideally we never need to Mutex unless we are invoked with the same target project,
-            // in the same branch/path!  This would allow two clientspecs to build at the same time (even though we need
-            // to make sure that XGE is only used once at a time)
-            bool bUseMutex = true;
-            {
-                int NoMutexArgumentIndex;
-                if (Utils.ParseCommandLineFlag(Arguments, "-NoMutex", out NoMutexArgumentIndex))
-                {
-                    bUseMutex = false;
-                }
-            }
+			// @todo: Ideally we never need to Mutex unless we are invoked with the same target project,
+			// in the same branch/path!  This would allow two clientspecs to build at the same time (even though we need
+			// to make sure that XGE is only used once at a time)
+			bool bUseMutex = true;
+			{
+				int NoMutexArgumentIndex;
+				if (Utils.ParseCommandLineFlag(Arguments, "-NoMutex", out NoMutexArgumentIndex))
+				{
+					bUseMutex = false;
+				}
+			}
 			bool bWaitMutex = false;
-            {
-                int WaitMutexArgumentIndex;
-                if (Utils.ParseCommandLineFlag(Arguments, "-WaitMutex", out WaitMutexArgumentIndex))
-                {
-                    bWaitMutex = true;
-                }
-            }
-
-            bool bAutoSDKOnly = false;
-            {
-                int AutoSDKOnlyArgumentIndex;
-                if (Utils.ParseCommandLineFlag(Arguments, "-autosdkonly", out AutoSDKOnlyArgumentIndex))
-                {
-                    bAutoSDKOnly = true;
+			{
+				int WaitMutexArgumentIndex;
+				if (Utils.ParseCommandLineFlag(Arguments, "-WaitMutex", out WaitMutexArgumentIndex))
+				{
+					bWaitMutex = true;
+				}
+			}
+
+			bool bAutoSDKOnly = false;
+			{
+				int AutoSDKOnlyArgumentIndex;
+				if (Utils.ParseCommandLineFlag(Arguments, "-autosdkonly", out AutoSDKOnlyArgumentIndex))
+				{
+					bAutoSDKOnly = true;
 					bIgnoreJunk = true;
 				}
-            }
-            bool bValidatePlatforms = false;
-            {
-                int ValidatePlatformsArgumentIndex;
-                if (Utils.ParseCommandLineFlag(Arguments, "-validateplatform", out ValidatePlatformsArgumentIndex))
-                {
-                    bValidatePlatforms = true;
+			}
+			bool bValidatePlatforms = false;
+			{
+				int ValidatePlatformsArgumentIndex;
+				if (Utils.ParseCommandLineFlag(Arguments, "-validateplatform", out ValidatePlatformsArgumentIndex))
+				{
+					bValidatePlatforms = true;
 					bIgnoreJunk = true;
 				}
-            }
-            
-			
-            // Don't allow simultaneous execution of Unreal Built Tool. Multi-selection in the UI e.g. causes this and you want serial
-            // execution in this case to avoid contention issues with shared produced items.
-            bool bCreatedMutex = false;
-            {
-                Mutex SingleInstanceMutex = null;
-                if (bUseMutex)
-                {
-                    int LocationHash = Assembly.GetEntryAssembly().GetOriginalLocation().GetHashCode();
-
-                    String MutexName;
-                    if (bAutoSDKOnly || bValidatePlatforms)
-                    {
-                        // this mutex has to be truly global because AutoSDK installs may change global state (like a console copying files into VS directories,
-                        // or installing Target Management services
-                        MutexName = "Global/UnrealBuildTool_Mutex_AutoSDKS";                        
-                    }
-                    else
-                    {                        
-                        MutexName = "Global/UnrealBuildTool_Mutex_" + LocationHash.ToString();
-                    }
-                    SingleInstanceMutex = new Mutex(true, MutexName, out bCreatedMutex);
-                    if (!bCreatedMutex)
-                    {
-						if(bWaitMutex)
+			}
+
+
+			// Don't allow simultaneous execution of Unreal Built Tool. Multi-selection in the UI e.g. causes this and you want serial
+			// execution in this case to avoid contention issues with shared produced items.
+			bool bCreatedMutex = false;
+			{
+				Mutex SingleInstanceMutex = null;
+				if (bUseMutex)
+				{
+					int LocationHash = Assembly.GetEntryAssembly().GetOriginalLocation().GetHashCode();
+
+					String MutexName;
+					if (bAutoSDKOnly || bValidatePlatforms)
+					{
+						// this mutex has to be truly global because AutoSDK installs may change global state (like a console copying files into VS directories,
+						// or installing Target Management services
+						MutexName = "Global/UnrealBuildTool_Mutex_AutoSDKS";
+					}
+					else
+					{
+						MutexName = "Global/UnrealBuildTool_Mutex_" + LocationHash.ToString();
+					}
+					SingleInstanceMutex = new Mutex(true, MutexName, out bCreatedMutex);
+					if (!bCreatedMutex)
+					{
+						if (bWaitMutex)
 						{
 							try
 							{
 								SingleInstanceMutex.WaitOne();
 							}
-							catch(AbandonedMutexException)
+							catch (AbandonedMutexException)
 							{
 							}
 						}
-                        else if (bAutoSDKOnly || bValidatePlatforms)
-                        {
-                            throw new BuildException("A conflicting instance of UnrealBuildTool is already running. Either -autosdkonly or -validateplatform was passed, which allows only a single instance to be run globally. Therefore, the conflicting instance may be in another location or the current location: {0}. A process manager may be used to determine the conflicting process and what tool may have launched it.", Assembly.GetEntryAssembly().GetOriginalLocation());
-                        }
-                        else
-                        {
-                            throw new BuildException("A conflicting instance of UnrealBuildTool is already running at this location: {0}. A process manager may be used to determine the conflicting process and what tool may have launched it.", Assembly.GetEntryAssembly().GetOriginalLocation());
-                        }
-                    }
-                }
-
-                try
-                {
-                    string GameName = null;
-                    var bGenerateVCProjectFiles = false;
-                    var bGenerateXcodeProjectFiles = false;
-                    var bGenerateMakefiles = false;
-                    var bGenerateCMakefiles = false;
-                    var bGenerateQMakefiles = false;
-                    var bGenerateKDevelopFiles = false;
-                    var bGenerateCodeLiteFiles = false;
-                    var bValidPlatformsOnly = false;
-                    var bSpecificModulesOnly = false;
-
-                    // We need to be able to identify if one of the arguments is the platform...
-                    // Leverage the existing parser function in UEBuildTarget to get this information.
-                    UnrealTargetPlatform CheckPlatform;
-                    UnrealTargetConfiguration CheckConfiguration;
-                    UEBuildTarget.ParsePlatformAndConfiguration(Arguments, out CheckPlatform, out CheckConfiguration, false);
+						else if (bAutoSDKOnly || bValidatePlatforms)
+						{
+							throw new BuildException("A conflicting instance of UnrealBuildTool is already running. Either -autosdkonly or -validateplatform was passed, which allows only a single instance to be run globally. Therefore, the conflicting instance may be in another location or the current location: {0}. A process manager may be used to determine the conflicting process and what tool may have launched it.", Assembly.GetEntryAssembly().GetOriginalLocation());
+						}
+						else
+						{
+							throw new BuildException("A conflicting instance of UnrealBuildTool is already running at this location: {0}. A process manager may be used to determine the conflicting process and what tool may have launched it.", Assembly.GetEntryAssembly().GetOriginalLocation());
+						}
+					}
+				}
+
+				try
+				{
+					string GameName = null;
+					var bGenerateVCProjectFiles = false;
+					var bGenerateXcodeProjectFiles = false;
+					var bGenerateMakefiles = false;
+					var bGenerateCMakefiles = false;
+					var bGenerateQMakefiles = false;
+					var bGenerateKDevelopFiles = false;
+					var bGenerateCodeLiteFiles = false;
+					var bValidPlatformsOnly = false;
+					var bSpecificModulesOnly = false;
+
+					// We need to be able to identify if one of the arguments is the platform...
+					// Leverage the existing parser function in UEBuildTarget to get this information.
+					UnrealTargetPlatform CheckPlatform;
+					UnrealTargetConfiguration CheckConfiguration;
+					UEBuildTarget.ParsePlatformAndConfiguration(Arguments, out CheckPlatform, out CheckConfiguration, false);
 
 					// @todo ubtmake: remove this when building with RPCUtility works
 					if (BuildHostPlatform.Current.Platform != CheckPlatform && (CheckPlatform == UnrealTargetPlatform.Mac || CheckPlatform == UnrealTargetPlatform.IOS))
@@ -1017,46 +1017,46 @@ namespace UnrealBuildTool
 						BuildConfiguration.bUseUBTMakefiles = false;
 					}
 
-                    // If we were asked to enable fast build iteration, we want the 'gather' phase to default to off (unless it is overridden below
-                    // using a command-line option.)
-                    if( BuildConfiguration.bUseUBTMakefiles )
-                    {
-                        bIsGatheringBuild_Unsafe = false;
-                    }
-
-                    foreach (var Arg in Arguments)
-                    {
-                        string LowercaseArg = Arg.ToLowerInvariant();
-                        const string OnlyPlatformSpecificForArg = "-onlyplatformspecificfor=";
-                        if (ParseRocketCommandlineArg(Arg, ref GameName))
-                        {
-                            // Already handled at startup. Calling now just to properly set the game name
-                            continue;
-                        }
-                        else if (LowercaseArg == "-skipnonhostplatforms")
-                        {
-                            bSkipNonHostPlatforms = true;
-                        }
-                        else if (LowercaseArg.StartsWith(OnlyPlatformSpecificForArg))
-                        {
-                            string OtherPlatform = LowercaseArg.Substring(OnlyPlatformSpecificForArg.Length);
-                            foreach (UnrealTargetPlatform platform in Enum.GetValues(typeof(UnrealTargetPlatform)))
-                            {
-                                if (OtherPlatform == platform.ToString().ToLowerInvariant())
-                                {
-                                    OnlyPlatformSpecificFor = platform;
-                                }
-                            }
-                            if (OnlyPlatformSpecificFor == UnrealTargetPlatform.Unknown)
-                            {
-                                throw new BuildException("Could not parse {0}", LowercaseArg);
-                            }
-                        }
-                        else if (LowercaseArg.StartsWith("-makefile"))
-                        {
-                            bGenerateMakefiles = true;
-                            ProjectFileGenerator.bGenerateProjectFiles = true;
-                        }
+					// If we were asked to enable fast build iteration, we want the 'gather' phase to default to off (unless it is overridden below
+					// using a command-line option.)
+					if (BuildConfiguration.bUseUBTMakefiles)
+					{
+						bIsGatheringBuild_Unsafe = false;
+					}
+
+					foreach (var Arg in Arguments)
+					{
+						string LowercaseArg = Arg.ToLowerInvariant();
+						const string OnlyPlatformSpecificForArg = "-onlyplatformspecificfor=";
+						if (ParseRocketCommandlineArg(Arg, ref GameName))
+						{
+							// Already handled at startup. Calling now just to properly set the game name
+							continue;
+						}
+						else if (LowercaseArg == "-skipnonhostplatforms")
+						{
+							bSkipNonHostPlatforms = true;
+						}
+						else if (LowercaseArg.StartsWith(OnlyPlatformSpecificForArg))
+						{
+							string OtherPlatform = LowercaseArg.Substring(OnlyPlatformSpecificForArg.Length);
+							foreach (UnrealTargetPlatform platform in Enum.GetValues(typeof(UnrealTargetPlatform)))
+							{
+								if (OtherPlatform == platform.ToString().ToLowerInvariant())
+								{
+									OnlyPlatformSpecificFor = platform;
+								}
+							}
+							if (OnlyPlatformSpecificFor == UnrealTargetPlatform.Unknown)
+							{
+								throw new BuildException("Could not parse {0}", LowercaseArg);
+							}
+						}
+						else if (LowercaseArg.StartsWith("-makefile"))
+						{
+							bGenerateMakefiles = true;
+							ProjectFileGenerator.bGenerateProjectFiles = true;
+						}
 						else if (LowercaseArg.StartsWith("-cmakefile"))
 						{
 							bGenerateCMakefiles = true;
@@ -1077,117 +1077,117 @@ namespace UnrealBuildTool
 							bGenerateCodeLiteFiles = true;
 							ProjectFileGenerator.bGenerateProjectFiles = true;
 						}
-                        else if (LowercaseArg.StartsWith("-projectfile"))
-                        {
-                            bGenerateVCProjectFiles = true;
-                            ProjectFileGenerator.bGenerateProjectFiles = true;
-                        }
-                        else if (LowercaseArg.StartsWith("-xcodeprojectfile"))
-                        {
-                            bGenerateXcodeProjectFiles = true;
-                            ProjectFileGenerator.bGenerateProjectFiles = true;
-                        }
-                        else if (LowercaseArg == "-validplatformsonly")
-                        {
-                            bValidPlatformsOnly = true;
-                        }
-                        else if (LowercaseArg == "development" || LowercaseArg == "debug" || LowercaseArg == "shipping" || LowercaseArg == "test" || LowercaseArg == "debuggame")
-                        {
-                            //ConfigName = Arg;
-                        }
-                        else if (LowercaseArg == "-modulewithsuffix")
-                        {
-                            bSpecificModulesOnly = true;
-                            continue;
-                        }
-                        else if (LowercaseArg == "-copyrightverification")
-                        {
-                            bRunCopyrightVerification = true;
-                        }
-                        else if (LowercaseArg == "-dumptofile")
-                        {
-                            bDumpToFile = true;
-                        }
-                        else if (LowercaseArg == "-check3rdpartyheaders")
-                        {
-                            bCheckThirdPartyHeaders = true;
-                        }
+						else if (LowercaseArg.StartsWith("-projectfile"))
+						{
+							bGenerateVCProjectFiles = true;
+							ProjectFileGenerator.bGenerateProjectFiles = true;
+						}
+						else if (LowercaseArg.StartsWith("-xcodeprojectfile"))
+						{
+							bGenerateXcodeProjectFiles = true;
+							ProjectFileGenerator.bGenerateProjectFiles = true;
+						}
+						else if (LowercaseArg == "-validplatformsonly")
+						{
+							bValidPlatformsOnly = true;
+						}
+						else if (LowercaseArg == "development" || LowercaseArg == "debug" || LowercaseArg == "shipping" || LowercaseArg == "test" || LowercaseArg == "debuggame")
+						{
+							//ConfigName = Arg;
+						}
+						else if (LowercaseArg == "-modulewithsuffix")
+						{
+							bSpecificModulesOnly = true;
+							continue;
+						}
+						else if (LowercaseArg == "-copyrightverification")
+						{
+							bRunCopyrightVerification = true;
+						}
+						else if (LowercaseArg == "-dumptofile")
+						{
+							bDumpToFile = true;
+						}
+						else if (LowercaseArg == "-check3rdpartyheaders")
+						{
+							bCheckThirdPartyHeaders = true;
+						}
 						else if (LowercaseArg == "-invalidatemakefilesonly")
 						{
 							UnrealBuildTool.bIsInvalidatingMakefilesOnly = true;
 						}
-                        else if (LowercaseArg == "-gather")
-                        {
-                            UnrealBuildTool.bIsGatheringBuild_Unsafe = true;
-                        }
-                        else if (LowercaseArg == "-nogather")
-                        {
-                            UnrealBuildTool.bIsGatheringBuild_Unsafe = false;
-                        }
-                        else if (LowercaseArg == "-gatheronly")
-                        {
-                            UnrealBuildTool.bIsGatheringBuild_Unsafe = true;
-                            UnrealBuildTool.bIsAssemblingBuild_Unsafe = false;
-                        }
-                        else if (LowercaseArg == "-assemble")
-                        {
-                            UnrealBuildTool.bIsAssemblingBuild_Unsafe = true;
-                        }
-                        else if (LowercaseArg == "-noassemble")
-                        {
-                            UnrealBuildTool.bIsAssemblingBuild_Unsafe = false;
-                        }
-                        else if (LowercaseArg == "-assembleonly")
-                        {
-                            UnrealBuildTool.bIsGatheringBuild_Unsafe = false;
-                            UnrealBuildTool.bIsAssemblingBuild_Unsafe = true;
-                        }
-                        else if (LowercaseArg == "-ignorejunk")
-                        {
-                            bIgnoreJunk = true;
-                        }
-                        else if (LowercaseArg == "-define")
-                        {
-                            // Skip -define
-                        }
-                        else if (LowercaseArg == "-progress")
-                        {
-                            ProgressWriter.bWriteMarkup = true;
-                        }
-                        else if (CheckPlatform.ToString().ToLowerInvariant() == LowercaseArg)
-                        {
-                            // It's the platform set...
-                            //PlatformName = Arg;
-                        }						
-                        else
-                        {
-                            // This arg may be a game name. Check for the existence of a game folder with this name.
-                            // "Engine" is not a valid game name.
-                            if (LowercaseArg != "engine" && Arg.IndexOfAny(Path.GetInvalidPathChars()) == -1 &&
-                                Directory.Exists(Path.Combine(ProjectFileGenerator.RootRelativePath, Arg, "Config")))
-                            {
-                                GameName = Arg;
-                                Log.TraceVerbose("CommandLine: Found game name '{0}'", GameName);
-                            }
-                            else if (LowercaseArg == "rocket")
-                            {
-                                GameName = Arg;
-                                Log.TraceVerbose("CommandLine: Found game name '{0}'", GameName);
-                            }
-                        }
-                    }
-
-                    if( !bIsGatheringBuild_Unsafe && !bIsAssemblingBuild_Unsafe )
-                    {
-                        throw new BuildException( "UnrealBuildTool: At least one of either IsGatheringBuild or IsAssemblingBuild must be true.  Did you pass '-NoGather' with '-NoAssemble'?" );
-                    }
+						else if (LowercaseArg == "-gather")
+						{
+							UnrealBuildTool.bIsGatheringBuild_Unsafe = true;
+						}
+						else if (LowercaseArg == "-nogather")
+						{
+							UnrealBuildTool.bIsGatheringBuild_Unsafe = false;
+						}
+						else if (LowercaseArg == "-gatheronly")
+						{
+							UnrealBuildTool.bIsGatheringBuild_Unsafe = true;
+							UnrealBuildTool.bIsAssemblingBuild_Unsafe = false;
+						}
+						else if (LowercaseArg == "-assemble")
+						{
+							UnrealBuildTool.bIsAssemblingBuild_Unsafe = true;
+						}
+						else if (LowercaseArg == "-noassemble")
+						{
+							UnrealBuildTool.bIsAssemblingBuild_Unsafe = false;
+						}
+						else if (LowercaseArg == "-assembleonly")
+						{
+							UnrealBuildTool.bIsGatheringBuild_Unsafe = false;
+							UnrealBuildTool.bIsAssemblingBuild_Unsafe = true;
+						}
+						else if (LowercaseArg == "-ignorejunk")
+						{
+							bIgnoreJunk = true;
+						}
+						else if (LowercaseArg == "-define")
+						{
+							// Skip -define
+						}
+						else if (LowercaseArg == "-progress")
+						{
+							ProgressWriter.bWriteMarkup = true;
+						}
+						else if (CheckPlatform.ToString().ToLowerInvariant() == LowercaseArg)
+						{
+							// It's the platform set...
+							//PlatformName = Arg;
+						}
+						else
+						{
+							// This arg may be a game name. Check for the existence of a game folder with this name.
+							// "Engine" is not a valid game name.
+							if (LowercaseArg != "engine" && Arg.IndexOfAny(Path.GetInvalidPathChars()) == -1 &&
+								Directory.Exists(Path.Combine(ProjectFileGenerator.RootRelativePath, Arg, "Config")))
+							{
+								GameName = Arg;
+								Log.TraceVerbose("CommandLine: Found game name '{0}'", GameName);
+							}
+							else if (LowercaseArg == "rocket")
+							{
+								GameName = Arg;
+								Log.TraceVerbose("CommandLine: Found game name '{0}'", GameName);
+							}
+						}
+					}
+
+					if (!bIsGatheringBuild_Unsafe && !bIsAssemblingBuild_Unsafe)
+					{
+						throw new BuildException("UnrealBuildTool: At least one of either IsGatheringBuild or IsAssemblingBuild must be true.  Did you pass '-NoGather' with '-NoAssemble'?");
+					}
 
 					if (BuildConfiguration.bRunUnrealCodeAnalyzer)
 					{
 						if (BuildConfiguration.UCAModuleToAnalyze == null || BuildConfiguration.UCAModuleToAnalyze.Length == 0)
-					{
-						throw new BuildException("When running UnrealCodeAnalyzer, please specify module to analyze in UCAModuleToAnalyze field in BuildConfiguration.xml");
-					}
+						{
+							throw new BuildException("When running UnrealCodeAnalyzer, please specify module to analyze in UCAModuleToAnalyze field in BuildConfiguration.xml");
+						}
 
 						if ((BuildConfiguration.bUCACheckPCHFiles && BuildConfiguration.bUCACheckUObjectThreadSafety)
 							|| (!BuildConfiguration.bUCACheckPCHFiles && !BuildConfiguration.bUCACheckUObjectThreadSafety))
@@ -1198,88 +1198,88 @@ namespace UnrealBuildTool
 
 					// Get the build version
 					BuildVersion Version;
-					if(!BuildVersion.TryRead("../Build/Build.version", out Version))
+					if (!BuildVersion.TryRead("../Build/Build.version", out Version))
 					{
 						Log.TraceWarning("Could not read build.version");
 					}
 
-                    // Send an event with basic usage dimensions
-                    // [RCL] 2014-11-03 FIXME: this is incorrect since we can have more than one action
-                    Telemetry.SendEvent("CommonAttributes.2",
-                        // @todo No platform independent way to do processor speed and count. Defer for now. Mono actually has ways to do this but needs to be tested.
-                        "ProcessorCount", Environment.ProcessorCount.ToString(),
-                        "ComputerName", Environment.MachineName,
-                        "User", Environment.UserName,
-                        "Domain", Environment.UserDomainName,
-                        "CommandLine", string.Join("|", Arguments),
-                        "UBT Action", bGenerateVCProjectFiles 
-                            ? "GenerateVCProjectFiles" 
-                            : bGenerateXcodeProjectFiles 
-                            ? "GenerateXcodeProjectFiles" 
-                            : bRunCopyrightVerification
-                            ? "RunCopyrightVerification"
-                            : bGenerateMakefiles
-                            ? "GenerateMakefiles"
-                            : bGenerateCMakefiles
-                            ? "GenerateCMakeFiles"
-                            : bGenerateQMakefiles
-                            ? "GenerateQMakefiles"
-                            : bGenerateKDevelopFiles
-                            ? "GenerateKDevelopFiles"
-                            : bGenerateCodeLiteFiles
-                            ? "GenerateCodeLiteFiles"
-                            : bValidatePlatforms 
-                            ? "ValidatePlatforms"
-                            : "Build",
-                        "Platform", CheckPlatform.ToString(),
-                        "Configuration", CheckConfiguration.ToString(),
-                        "EngineVersion", (Version == null)? "0" : Version.Changelist.ToString(),
-                        "Branch", (Version == null)? "" : Version.BranchName
-                        );
-
-
-
-                    if (bValidPlatformsOnly == true)
-                    {
-                        // We turned the flag on if generating projects so that the
-                        // build platforms would know we are doing so... In this case,
-                        // turn it off to only generate projects for valid platforms.
-                        ProjectFileGenerator.bGenerateProjectFiles = false;
-                    }
-
-                    // Find and register all tool chains, build platforms, etc. that are present
-                    RegisterAllUBTClasses();
-                    ProjectFileGenerator.bGenerateProjectFiles = false;
-
-					if( BuildConfiguration.bPrintPerformanceInfo )
-					{ 
+					// Send an event with basic usage dimensions
+					// [RCL] 2014-11-03 FIXME: this is incorrect since we can have more than one action
+					Telemetry.SendEvent("CommonAttributes.2",
+						// @todo No platform independent way to do processor speed and count. Defer for now. Mono actually has ways to do this but needs to be tested.
+						"ProcessorCount", Environment.ProcessorCount.ToString(),
+						"ComputerName", Environment.MachineName,
+						"User", Environment.UserName,
+						"Domain", Environment.UserDomainName,
+						"CommandLine", string.Join("|", Arguments),
+						"UBT Action", bGenerateVCProjectFiles
+							? "GenerateVCProjectFiles"
+							: bGenerateXcodeProjectFiles
+							? "GenerateXcodeProjectFiles"
+							: bRunCopyrightVerification
+							? "RunCopyrightVerification"
+							: bGenerateMakefiles
+							? "GenerateMakefiles"
+							: bGenerateCMakefiles
+							? "GenerateCMakeFiles"
+							: bGenerateQMakefiles
+							? "GenerateQMakefiles"
+							: bGenerateKDevelopFiles
+							? "GenerateKDevelopFiles"
+							: bGenerateCodeLiteFiles
+							? "GenerateCodeLiteFiles"
+							: bValidatePlatforms
+							? "ValidatePlatforms"
+							: "Build",
+						"Platform", CheckPlatform.ToString(),
+						"Configuration", CheckConfiguration.ToString(),
+						"EngineVersion", (Version == null) ? "0" : Version.Changelist.ToString(),
+						"Branch", (Version == null) ? "" : Version.BranchName
+						);
+
+
+
+					if (bValidPlatformsOnly == true)
+					{
+						// We turned the flag on if generating projects so that the
+						// build platforms would know we are doing so... In this case,
+						// turn it off to only generate projects for valid platforms.
+						ProjectFileGenerator.bGenerateProjectFiles = false;
+					}
+
+					// Find and register all tool chains, build platforms, etc. that are present
+					RegisterAllUBTClasses();
+					ProjectFileGenerator.bGenerateProjectFiles = false;
+
+					if (BuildConfiguration.bPrintPerformanceInfo)
+					{
 						var BasicInitTime = (DateTime.UtcNow - BasicInitStartTime).TotalSeconds;
-						Log.TraceInformation( "Basic UBT initialization took " + BasicInitTime + "s" );
+						Log.TraceInformation("Basic UBT initialization took " + BasicInitTime + "s");
 					}
 
-                    // now that we know the available platforms, we can delete other platforms' junk. if we're only building specific modules from the editor, don't touch anything else (it may be in use).
-                    if (!bSpecificModulesOnly && !bIgnoreJunk)
-                    {
-                        JunkDeleter.DeleteJunk();
-                    }
+					// now that we know the available platforms, we can delete other platforms' junk. if we're only building specific modules from the editor, don't touch anything else (it may be in use).
+					if (!bSpecificModulesOnly && !bIgnoreJunk)
+					{
+						JunkDeleter.DeleteJunk();
+					}
 
 					if (bGenerateCodeLiteFiles || bGenerateVCProjectFiles || bGenerateXcodeProjectFiles || bGenerateMakefiles || bGenerateCMakefiles || bGenerateQMakefiles || bGenerateKDevelopFiles)
-                    {
+					{
 						FileReference UProjectFile = GetUProjectFile();
 
-                        bool bGenerationSuccess = true;
-                        if (bGenerateVCProjectFiles)
-                        {
-                            bGenerationSuccess &= GenerateProjectFiles(new VCProjectFileGenerator(UProjectFile), Arguments);
-                        }
-                        if (bGenerateXcodeProjectFiles)
-                        {
-                            bGenerationSuccess &= GenerateProjectFiles(new XcodeProjectFileGenerator(UProjectFile), Arguments);
-                        }
-                        if (bGenerateMakefiles)
-                        {
-                            bGenerationSuccess &= GenerateProjectFiles(new MakefileGenerator(UProjectFile), Arguments);
-                        }
+						bool bGenerationSuccess = true;
+						if (bGenerateVCProjectFiles)
+						{
+							bGenerationSuccess &= GenerateProjectFiles(new VCProjectFileGenerator(UProjectFile), Arguments);
+						}
+						if (bGenerateXcodeProjectFiles)
+						{
+							bGenerationSuccess &= GenerateProjectFiles(new XcodeProjectFileGenerator(UProjectFile), Arguments);
+						}
+						if (bGenerateMakefiles)
+						{
+							bGenerationSuccess &= GenerateProjectFiles(new MakefileGenerator(UProjectFile), Arguments);
+						}
 						if (bGenerateCMakefiles)
 						{
 							bGenerationSuccess &= GenerateProjectFiles(new CMakefileGenerator(UProjectFile), Arguments);
@@ -1296,337 +1296,337 @@ namespace UnrealBuildTool
 						{
 							bGenerationSuccess &= GenerateProjectFiles(new CodeLiteGenerator(UProjectFile), Arguments);
 						}
-                        if(!bGenerationSuccess)
-                        {
-                            Result = ECompilationResult.OtherCompilationError;
-                        }
-                    }
-                    else if (bRunCopyrightVerification)
-                    {
-                        // We ensure we are in the engine source directory during startup.
-                        CopyrightVerify.RunCopyrightVerification(Directory.GetCurrentDirectory(), GameName, bDumpToFile);
-                        Log.TraceInformation("Completed... exiting.");
-                    }
-                    else if (bAutoSDKOnly)
-                    {
-                        // RegisterAllUBTClasses has already done all the SDK setup.
-                        Result = ECompilationResult.Succeeded;
-                    }
-                    else if (bValidatePlatforms)
-                    {
-                        ValidatePlatforms(Arguments);
-                    }
-                    else
-                    {
-                        // Check if any third party headers are included from public engine headers.
-                        if (bCheckThirdPartyHeaders)
-                        {
-                            ThirdPartyHeaderFinder.FindThirdPartyIncludes(CheckPlatform, CheckConfiguration);
-                        }
-
-                        // Build our project
-                        if (Result == ECompilationResult.Succeeded)
-                        {
-                            if (UEBuildConfiguration.bPrepForDeployment == false)
-                            {
+						if (!bGenerationSuccess)
+						{
+							Result = ECompilationResult.OtherCompilationError;
+						}
+					}
+					else if (bRunCopyrightVerification)
+					{
+						// We ensure we are in the engine source directory during startup.
+						CopyrightVerify.RunCopyrightVerification(Directory.GetCurrentDirectory(), GameName, bDumpToFile);
+						Log.TraceInformation("Completed... exiting.");
+					}
+					else if (bAutoSDKOnly)
+					{
+						// RegisterAllUBTClasses has already done all the SDK setup.
+						Result = ECompilationResult.Succeeded;
+					}
+					else if (bValidatePlatforms)
+					{
+						ValidatePlatforms(Arguments);
+					}
+					else
+					{
+						// Check if any third party headers are included from public engine headers.
+						if (bCheckThirdPartyHeaders)
+						{
+							ThirdPartyHeaderFinder.FindThirdPartyIncludes(CheckPlatform, CheckConfiguration);
+						}
+
+						// Build our project
+						if (Result == ECompilationResult.Succeeded)
+						{
+							if (UEBuildConfiguration.bPrepForDeployment == false)
+							{
 								// If we are only prepping for deployment, assume the build already occurred.
 								Result = RunUBT(Arguments);
-                            }
-                            else
-                            {
-                                var BuildPlatform = UEBuildPlatform.GetBuildPlatform(CheckPlatform, true);
-                                if (BuildPlatform != null)
-                                {
-                                    // Setup environment wasn't called, so set the flag
-                                    BuildConfiguration.bDeployAfterCompile = BuildPlatform.RequiresDeployPrepAfterCompile();
-                                    BuildConfiguration.PlatformIntermediateFolder = Path.Combine(BuildConfiguration.BaseIntermediateFolder, CheckPlatform.ToString(), BuildPlatform.GetActiveArchitecture());
-                                }
-                            }
-
-                            // If we build w/ bXGEExport true, we didn't REALLY build at this point, 
-                            // so don't bother with doing the PrepTargetForDeployment call. 
-                            if ((Result == ECompilationResult.Succeeded) && (BuildConfiguration.bDeployAfterCompile == true) && (BuildConfiguration.bXGEExport == false) &&
-                                (UEBuildConfiguration.bGenerateManifest == false) && (UEBuildConfiguration.bGenerateExternalFileList == false) && (UEBuildConfiguration.bCleanProject == false))
-                            {
-                                var DeployHandler = UEBuildDeploy.GetBuildDeploy(CheckPlatform);
-                                if (DeployHandler != null)
-                                {
-                                    // We need to be able to identify the Target.Type we can derive it from the Arguments.
-                                    BuildConfiguration.bFlushBuildDirOnRemoteMac = false;
-									var TargetDescs = UEBuildTarget.ParseTargetCommandLine( Arguments, GetUProjectFile() );
-			                        UEBuildTarget CheckTarget = UEBuildTarget.CreateTarget( TargetDescs[0] );	// @todo ubtmake: This may not work in assembler only mode.  We don't want to be loading target rules assemblies here either.
-	                                IUEBuildPlatform BuildPlatform = UEBuildPlatform.GetBuildPlatform(TargetDescs[0].Platform);
+							}
+							else
+							{
+								var BuildPlatform = UEBuildPlatform.GetBuildPlatform(CheckPlatform, true);
+								if (BuildPlatform != null)
+								{
+									// Setup environment wasn't called, so set the flag
+									BuildConfiguration.bDeployAfterCompile = BuildPlatform.RequiresDeployPrepAfterCompile();
+									BuildConfiguration.PlatformIntermediateFolder = Path.Combine(BuildConfiguration.BaseIntermediateFolder, CheckPlatform.ToString(), BuildPlatform.GetActiveArchitecture());
+								}
+							}
+
+							// If we build w/ bXGEExport true, we didn't REALLY build at this point, 
+							// so don't bother with doing the PrepTargetForDeployment call. 
+							if ((Result == ECompilationResult.Succeeded) && (BuildConfiguration.bDeployAfterCompile == true) && (BuildConfiguration.bXGEExport == false) &&
+								(UEBuildConfiguration.bGenerateManifest == false) && (UEBuildConfiguration.bGenerateExternalFileList == false) && (UEBuildConfiguration.bCleanProject == false))
+							{
+								var DeployHandler = UEBuildDeploy.GetBuildDeploy(CheckPlatform);
+								if (DeployHandler != null)
+								{
+									// We need to be able to identify the Target.Type we can derive it from the Arguments.
+									BuildConfiguration.bFlushBuildDirOnRemoteMac = false;
+									var TargetDescs = UEBuildTarget.ParseTargetCommandLine(Arguments, GetUProjectFile());
+									UEBuildTarget CheckTarget = UEBuildTarget.CreateTarget(TargetDescs[0]);	// @todo ubtmake: This may not work in assembler only mode.  We don't want to be loading target rules assemblies here either.
+									IUEBuildPlatform BuildPlatform = UEBuildPlatform.GetBuildPlatform(TargetDescs[0].Platform);
 									IUEToolChain ToolChain = UEToolChain.GetPlatformToolChain(BuildPlatform.GetCPPTargetPlatform(TargetDescs[0].Platform));
-                                    CheckTarget.SetupGlobalEnvironment(ToolChain);
-                                    if ((CheckTarget.TargetType == TargetRules.TargetType.Game) ||
-                                        (CheckTarget.TargetType == TargetRules.TargetType.Server) ||
-                                        (CheckTarget.TargetType == TargetRules.TargetType.Client))
-                                    {
-                                        CheckTarget.AppName = CheckTarget.TargetName;
-                                    }
-                                    DeployHandler.PrepTargetForDeployment(CheckTarget);
-                                }
-                            }
-                        }
-                    }
-                    // Print some performance info
-                    var BuildDuration = (DateTime.UtcNow - StartTime).TotalSeconds;
-                    if (BuildConfiguration.bPrintPerformanceInfo)
-                    {
-                        Log.TraceInformation("GetIncludes time: " + CPPEnvironment.TotalTimeSpentGettingIncludes + "s (" + CPPEnvironment.TotalIncludesRequested + " includes)" );
-                        Log.TraceInformation("DirectIncludes cache miss time: " + CPPEnvironment.DirectIncludeCacheMissesTotalTime + "s (" + CPPEnvironment.TotalDirectIncludeCacheMisses + " misses)" );
-                        Log.TraceInformation("FindIncludePaths calls: " + CPPEnvironment.TotalFindIncludedFileCalls + " (" + CPPEnvironment.IncludePathSearchAttempts + " searches)" );
-                        Log.TraceInformation("PCH gen time: " + UEBuildModuleCPP.TotalPCHGenTime + "s" );
-                        Log.TraceInformation("PCH cache time: " + UEBuildModuleCPP.TotalPCHCacheTime + "s" );
-                        Log.TraceInformation("Deep C++ include scan time: " + UnrealBuildTool.TotalDeepIncludeScanTime + "s" );
-                        Log.TraceInformation("Include Resolves: {0} ({1} misses, {2:0.00}%)", CPPEnvironment.TotalDirectIncludeResolves, CPPEnvironment.TotalDirectIncludeResolveCacheMisses, (float)CPPEnvironment.TotalDirectIncludeResolveCacheMisses / (float)CPPEnvironment.TotalDirectIncludeResolves * 100);
-                        Log.TraceInformation("Total FileItems: {0} ({1} missing)", FileItem.TotalFileItemCount, FileItem.MissingFileItemCount);
-
-                        Log.TraceInformation("Execution time: {0}s", BuildDuration);
-                    }
-
-                    //Telemetry.SendEvent("PerformanceInfo.2",
-                    //    "TotalExecutionTimeSec", BuildDuration.ToString("0.00"),
-                    //    "TotalTimeSpentGettingIncludesSec", CPPEnvironment.TotalTimeSpentGettingIncludes.ToString("0.00"),
-                    //    "TotalIncludesRequested", CPPEnvironment.TotalIncludesRequested.ToString(),
-                    //    "DirectIncludeCacheMissesTotalTimeSec", CPPEnvironment.DirectIncludeCacheMissesTotalTime.ToString("0.00"),
-                    //    "TotalDirectIncludeCacheMisses", CPPEnvironment.TotalDirectIncludeCacheMisses.ToString(),
-                    //    "TotalFindIncludedFileCalls", CPPEnvironment.TotalFindIncludedFileCalls.ToString(),
-                    //    "IncludePathSearchAttempts", CPPEnvironment.IncludePathSearchAttempts.ToString(),
-                    //    "TotalFileItemCount", FileItem.TotalFileItemCount.ToString(),
-                    //    "MissingFileItemCount", FileItem.MissingFileItemCount.ToString()
-                    //    );
-                }
-                catch (Exception Exception)
-                {
-                    Log.TraceError("UnrealBuildTool Exception: " + Exception);
-                    Result = ECompilationResult.OtherCompilationError;
-                }
+									CheckTarget.SetupGlobalEnvironment(ToolChain);
+									if ((CheckTarget.TargetType == TargetRules.TargetType.Game) ||
+										(CheckTarget.TargetType == TargetRules.TargetType.Server) ||
+										(CheckTarget.TargetType == TargetRules.TargetType.Client))
+									{
+										CheckTarget.AppName = CheckTarget.TargetName;
+									}
+									DeployHandler.PrepTargetForDeployment(CheckTarget);
+								}
+							}
+						}
+					}
+					// Print some performance info
+					var BuildDuration = (DateTime.UtcNow - StartTime).TotalSeconds;
+					if (BuildConfiguration.bPrintPerformanceInfo)
+					{
+						Log.TraceInformation("GetIncludes time: " + CPPEnvironment.TotalTimeSpentGettingIncludes + "s (" + CPPEnvironment.TotalIncludesRequested + " includes)");
+						Log.TraceInformation("DirectIncludes cache miss time: " + CPPEnvironment.DirectIncludeCacheMissesTotalTime + "s (" + CPPEnvironment.TotalDirectIncludeCacheMisses + " misses)");
+						Log.TraceInformation("FindIncludePaths calls: " + CPPEnvironment.TotalFindIncludedFileCalls + " (" + CPPEnvironment.IncludePathSearchAttempts + " searches)");
+						Log.TraceInformation("PCH gen time: " + UEBuildModuleCPP.TotalPCHGenTime + "s");
+						Log.TraceInformation("PCH cache time: " + UEBuildModuleCPP.TotalPCHCacheTime + "s");
+						Log.TraceInformation("Deep C++ include scan time: " + UnrealBuildTool.TotalDeepIncludeScanTime + "s");
+						Log.TraceInformation("Include Resolves: {0} ({1} misses, {2:0.00}%)", CPPEnvironment.TotalDirectIncludeResolves, CPPEnvironment.TotalDirectIncludeResolveCacheMisses, (float)CPPEnvironment.TotalDirectIncludeResolveCacheMisses / (float)CPPEnvironment.TotalDirectIncludeResolves * 100);
+						Log.TraceInformation("Total FileItems: {0} ({1} missing)", FileItem.TotalFileItemCount, FileItem.MissingFileItemCount);
+
+						Log.TraceInformation("Execution time: {0}s", BuildDuration);
+					}
 
-                if (bUseMutex)
-                {
-                    // Release the mutex to avoid the abandoned mutex timeout.
-                    SingleInstanceMutex.ReleaseMutex();
-                    SingleInstanceMutex.Dispose();
-                    SingleInstanceMutex = null;
-                }
-            }
+					//Telemetry.SendEvent("PerformanceInfo.2",
+					//    "TotalExecutionTimeSec", BuildDuration.ToString("0.00"),
+					//    "TotalTimeSpentGettingIncludesSec", CPPEnvironment.TotalTimeSpentGettingIncludes.ToString("0.00"),
+					//    "TotalIncludesRequested", CPPEnvironment.TotalIncludesRequested.ToString(),
+					//    "DirectIncludeCacheMissesTotalTimeSec", CPPEnvironment.DirectIncludeCacheMissesTotalTime.ToString("0.00"),
+					//    "TotalDirectIncludeCacheMisses", CPPEnvironment.TotalDirectIncludeCacheMisses.ToString(),
+					//    "TotalFindIncludedFileCalls", CPPEnvironment.TotalFindIncludedFileCalls.ToString(),
+					//    "IncludePathSearchAttempts", CPPEnvironment.IncludePathSearchAttempts.ToString(),
+					//    "TotalFileItemCount", FileItem.TotalFileItemCount.ToString(),
+					//    "MissingFileItemCount", FileItem.MissingFileItemCount.ToString()
+					//    );
+				}
+				catch (Exception Exception)
+				{
+					Log.TraceError("UnrealBuildTool Exception: " + Exception);
+					Result = ECompilationResult.OtherCompilationError;
+				}
+
+				if (bUseMutex)
+				{
+					// Release the mutex to avoid the abandoned mutex timeout.
+					SingleInstanceMutex.ReleaseMutex();
+					SingleInstanceMutex.Dispose();
+					SingleInstanceMutex = null;
+				}
+			}
 
-            Telemetry.Shutdown();
+			Telemetry.Shutdown();
 
-            // Print some performance info
-            Log.TraceVerbose("Execution time: {0}", (DateTime.UtcNow - StartTime).TotalSeconds);
+			// Print some performance info
+			Log.TraceVerbose("Execution time: {0}", (DateTime.UtcNow - StartTime).TotalSeconds);
 
 			if (ExtendedErrorCode != 0)
 			{
 				return ExtendedErrorCode;
 			}
-            return (int)Result; 
-        }
+			return (int)Result;
+		}
 
 		public static int ExtendedErrorCode = 0;
-        private static int Main(string[] Arguments)
-        {
-            // make sure we catch any exceptions and return an appropriate error code.
-            // Some inner code already does this (to ensure the Mutex is released),
-            // but we need something to cover all outer code as well.
-            try
-            {
-                // Make it more explicit what startup code should not be accessing config data.
-                // The Configuration classes will also assert, but this clarifies where the line is.
-                DoStartupStuffThatCannotAccessConfigurationClasses(Arguments);
-
-                // Now we can load the configuration files.
-                InitConfigurationClasses(Arguments);
-
-                // Finally, we do execute code that relies on config classes.
-                return DoPostStartupStuffThatCanAccessConfigs(Arguments);
-            }
-            catch (Exception Exception)
-            {
-                if (Log.IsInitialized())
-                {
-                    Log.TraceError("UnrealBuildTool Exception: " + Exception.ToString());
-                }
+		private static int Main(string[] Arguments)
+		{
+			// make sure we catch any exceptions and return an appropriate error code.
+			// Some inner code already does this (to ensure the Mutex is released),
+			// but we need something to cover all outer code as well.
+			try
+			{
+				// Make it more explicit what startup code should not be accessing config data.
+				// The Configuration classes will also assert, but this clarifies where the line is.
+				DoStartupStuffThatCannotAccessConfigurationClasses(Arguments);
+
+				// Now we can load the configuration files.
+				InitConfigurationClasses(Arguments);
+
+				// Finally, we do execute code that relies on config classes.
+				return DoPostStartupStuffThatCanAccessConfigs(Arguments);
+			}
+			catch (Exception Exception)
+			{
+				if (Log.IsInitialized())
+				{
+					Log.TraceError("UnrealBuildTool Exception: " + Exception.ToString());
+				}
 				if (ExtendedErrorCode != 0)
 				{
 					return ExtendedErrorCode;
 				}
-                return (int)ECompilationResult.OtherCompilationError;
-            }
-        }
-
-        /// <summary>
-        /// Generates project files.  Can also be used to generate projects "in memory", to allow for building
-        /// targets without even having project files at all.
-        /// </summary>
-        /// <param name="Generator">Project generator to use</param>
-        /// <param name="Arguments">Command-line arguments</param>
-        /// <returns>True if successful</returns>
-        public static bool GenerateProjectFiles(ProjectFileGenerator Generator, string[] Arguments)
-        {
-            bool bSuccess;
-            ProjectFileGenerator.bGenerateProjectFiles = true;
-            Generator.GenerateProjectFiles(Arguments, out bSuccess);
-            ProjectFileGenerator.bGenerateProjectFiles = false;
-            return bSuccess;
-        }
-
-
-
-        /// <summary>
-        /// Validates the various platforms to determine if they are ready for building
-        /// </summary>
-        public static void ValidatePlatforms(string[] Arguments)
-        {
-            List<UnrealTargetPlatform> Platforms = new List<UnrealTargetPlatform>(); 
-            foreach (var CurArgument in Arguments)
-            {
-                if (CurArgument.StartsWith("-"))
-                {
-                    if (CurArgument.StartsWith("-Platforms=", StringComparison.InvariantCultureIgnoreCase))
-                    {
-                        // Parse the list... will be in Foo+Bar+New format
-                        string PlatformList = CurArgument.Substring(11);
-                        while (PlatformList.Length > 0)
-                        {
-                            string PlatformString = PlatformList;
-                            Int32 PlusIdx = PlatformList.IndexOf("+");
-                            if (PlusIdx != -1)
-                            {
-                                PlatformString = PlatformList.Substring(0, PlusIdx);
-                                PlatformList = PlatformList.Substring(PlusIdx + 1);
-                            }
-                            else
-                            {
-                                // We are on the last platform... clear the list to exit the loop
-                                PlatformList = "";
-                            }
-
-                            // Is the string a valid platform? If so, add it to the list
-                            UnrealTargetPlatform SpecifiedPlatform = UnrealTargetPlatform.Unknown;
-                            foreach (UnrealTargetPlatform PlatformParam in Enum.GetValues(typeof(UnrealTargetPlatform)))
-                            {
-                                if (PlatformString.Equals(PlatformParam.ToString(), StringComparison.InvariantCultureIgnoreCase))
-                                {
-                                    SpecifiedPlatform = PlatformParam;
-                                    break;
-                                }
-                            }
-
-                            if (SpecifiedPlatform != UnrealTargetPlatform.Unknown)
-                            {
-                                if (Platforms.Contains(SpecifiedPlatform) == false)
-                                {
-                                    Platforms.Add(SpecifiedPlatform);
-                                }
-                            }
-                            else
-                            {
-                                Log.TraceWarning("ValidatePlatforms invalid platform specified: {0}", PlatformString);
-                            }
-                        }
-                    }
-                    else switch (CurArgument.ToUpperInvariant())
-                    {
-                        case "-ALLPLATFORMS":
-                            foreach (UnrealTargetPlatform platform in Enum.GetValues(typeof(UnrealTargetPlatform)))
-                            {
-                                if (platform != UnrealTargetPlatform.Unknown)
-                                {
-                                    if (Platforms.Contains(platform) == false)
-                                    {
-                                        Platforms.Add(platform);
-                                    }
-                                }
-                            }
-                            break;
-                    }
-                }
-            }
-
-            foreach (UnrealTargetPlatform platform in Platforms)
-            {
-                var BuildPlatform = UEBuildPlatform.GetBuildPlatform(platform, true);
-                if (BuildPlatform != null && BuildPlatform.HasRequiredSDKsInstalled() == SDKStatus.Valid)
-                {
-                    Console.WriteLine("##PlatformValidate: {0} VALID", platform.ToString());
-                }
-                else
-                {
-                    Console.WriteLine("##PlatformValidate: {0} INVALID", platform.ToString());
-                }
-            }
-        }
+				return (int)ECompilationResult.OtherCompilationError;
+			}
+		}
+
+		/// <summary>
+		/// Generates project files.  Can also be used to generate projects "in memory", to allow for building
+		/// targets without even having project files at all.
+		/// </summary>
+		/// <param name="Generator">Project generator to use</param>
+		/// <param name="Arguments">Command-line arguments</param>
+		/// <returns>True if successful</returns>
+		public static bool GenerateProjectFiles(ProjectFileGenerator Generator, string[] Arguments)
+		{
+			bool bSuccess;
+			ProjectFileGenerator.bGenerateProjectFiles = true;
+			Generator.GenerateProjectFiles(Arguments, out bSuccess);
+			ProjectFileGenerator.bGenerateProjectFiles = false;
+			return bSuccess;
+		}
+
+
+
+		/// <summary>
+		/// Validates the various platforms to determine if they are ready for building
+		/// </summary>
+		public static void ValidatePlatforms(string[] Arguments)
+		{
+			List<UnrealTargetPlatform> Platforms = new List<UnrealTargetPlatform>();
+			foreach (var CurArgument in Arguments)
+			{
+				if (CurArgument.StartsWith("-"))
+				{
+					if (CurArgument.StartsWith("-Platforms=", StringComparison.InvariantCultureIgnoreCase))
+					{
+						// Parse the list... will be in Foo+Bar+New format
+						string PlatformList = CurArgument.Substring(11);
+						while (PlatformList.Length > 0)
+						{
+							string PlatformString = PlatformList;
+							Int32 PlusIdx = PlatformList.IndexOf("+");
+							if (PlusIdx != -1)
+							{
+								PlatformString = PlatformList.Substring(0, PlusIdx);
+								PlatformList = PlatformList.Substring(PlusIdx + 1);
+							}
+							else
+							{
+								// We are on the last platform... clear the list to exit the loop
+								PlatformList = "";
+							}
+
+							// Is the string a valid platform? If so, add it to the list
+							UnrealTargetPlatform SpecifiedPlatform = UnrealTargetPlatform.Unknown;
+							foreach (UnrealTargetPlatform PlatformParam in Enum.GetValues(typeof(UnrealTargetPlatform)))
+							{
+								if (PlatformString.Equals(PlatformParam.ToString(), StringComparison.InvariantCultureIgnoreCase))
+								{
+									SpecifiedPlatform = PlatformParam;
+									break;
+								}
+							}
+
+							if (SpecifiedPlatform != UnrealTargetPlatform.Unknown)
+							{
+								if (Platforms.Contains(SpecifiedPlatform) == false)
+								{
+									Platforms.Add(SpecifiedPlatform);
+								}
+							}
+							else
+							{
+								Log.TraceWarning("ValidatePlatforms invalid platform specified: {0}", PlatformString);
+							}
+						}
+					}
+					else switch (CurArgument.ToUpperInvariant())
+						{
+							case "-ALLPLATFORMS":
+								foreach (UnrealTargetPlatform platform in Enum.GetValues(typeof(UnrealTargetPlatform)))
+								{
+									if (platform != UnrealTargetPlatform.Unknown)
+									{
+										if (Platforms.Contains(platform) == false)
+										{
+											Platforms.Add(platform);
+										}
+									}
+								}
+								break;
+						}
+				}
+			}
+
+			foreach (UnrealTargetPlatform platform in Platforms)
+			{
+				var BuildPlatform = UEBuildPlatform.GetBuildPlatform(platform, true);
+				if (BuildPlatform != null && BuildPlatform.HasRequiredSDKsInstalled() == SDKStatus.Valid)
+				{
+					Console.WriteLine("##PlatformValidate: {0} VALID", platform.ToString());
+				}
+				else
+				{
+					Console.WriteLine("##PlatformValidate: {0} INVALID", platform.ToString());
+				}
+			}
+		}
+
+
+		public static ECompilationResult RunUBT(string[] Arguments)
+		{
+			bool bSuccess = true;
 
-        
-        public static ECompilationResult RunUBT(string[] Arguments)
-        {
-            bool bSuccess = true;
 
-			
 			var RunUBTInitStartTime = DateTime.UtcNow;
 
 
-            // Reset global configurations
-            ActionGraph.ResetAllActions();
+			// Reset global configurations
+			ActionGraph.ResetAllActions();
 
-            // We need to allow the target platform to perform the 'reset' as well...
-            UnrealTargetPlatform ResetPlatform = UnrealTargetPlatform.Unknown;
-            UnrealTargetConfiguration ResetConfiguration;
-            UEBuildTarget.ParsePlatformAndConfiguration(Arguments, out ResetPlatform, out ResetConfiguration);
+			// We need to allow the target platform to perform the 'reset' as well...
+			UnrealTargetPlatform ResetPlatform = UnrealTargetPlatform.Unknown;
+			UnrealTargetConfiguration ResetConfiguration;
+			UEBuildTarget.ParsePlatformAndConfiguration(Arguments, out ResetPlatform, out ResetConfiguration);
 			var BuildPlatform = UEBuildPlatform.GetBuildPlatform(ResetPlatform);
 			BuildPlatform.ResetBuildConfiguration(ResetPlatform, ResetConfiguration);
 
-            // now that we have the platform, we can set the intermediate path to include the platform/architecture name
-            BuildConfiguration.PlatformIntermediateFolder = Path.Combine(BuildConfiguration.BaseIntermediateFolder, ResetPlatform.ToString(), BuildPlatform.GetActiveArchitecture());
+			// now that we have the platform, we can set the intermediate path to include the platform/architecture name
+			BuildConfiguration.PlatformIntermediateFolder = Path.Combine(BuildConfiguration.BaseIntermediateFolder, ResetPlatform.ToString(), BuildPlatform.GetActiveArchitecture());
 
-            string ExecutorName = "Unknown";
-            ECompilationResult BuildResult = ECompilationResult.Succeeded;
+			string ExecutorName = "Unknown";
+			ECompilationResult BuildResult = ECompilationResult.Succeeded;
 
-            var ToolChain = UEToolChain.GetPlatformToolChain(BuildPlatform.GetCPPTargetPlatform(ResetPlatform));
+			var ToolChain = UEToolChain.GetPlatformToolChain(BuildPlatform.GetCPPTargetPlatform(ResetPlatform));
 
-            string EULAViolationWarning = null;
+			string EULAViolationWarning = null;
 			Thread CPPIncludesThread = null;
 
-            try
-            {
-                List<string[]> TargetSettings = ParseCommandLineFlags(Arguments);
+			try
+			{
+				List<string[]> TargetSettings = ParseCommandLineFlags(Arguments);
 
-                int ArgumentIndex;
-                // action graph implies using the dependency resolve cache
-                bool GeneratingActionGraph = Utils.ParseCommandLineFlag(Arguments, "-graph", out ArgumentIndex );
-                if (GeneratingActionGraph)
-                {
-                    BuildConfiguration.bUseIncludeDependencyResolveCache = true;
-                }
+				int ArgumentIndex;
+				// action graph implies using the dependency resolve cache
+				bool GeneratingActionGraph = Utils.ParseCommandLineFlag(Arguments, "-graph", out ArgumentIndex);
+				if (GeneratingActionGraph)
+				{
+					BuildConfiguration.bUseIncludeDependencyResolveCache = true;
+				}
 
-                bool CreateStub = Utils.ParseCommandLineFlag(Arguments, "-nocreatestub", out ArgumentIndex);
-                if (CreateStub || (String.IsNullOrEmpty(Environment.GetEnvironmentVariable("uebp_LOCAL_ROOT")) && BuildHostPlatform.Current.Platform == UnrealTargetPlatform.Mac))
-                {
-                    BuildConfiguration.bCreateStubIPA = false;
-                }
+				bool CreateStub = Utils.ParseCommandLineFlag(Arguments, "-nocreatestub", out ArgumentIndex);
+				if (CreateStub || (String.IsNullOrEmpty(Environment.GetEnvironmentVariable("uebp_LOCAL_ROOT")) && BuildHostPlatform.Current.Platform == UnrealTargetPlatform.Mac))
+				{
+					BuildConfiguration.bCreateStubIPA = false;
+				}
 
 
-				if( BuildConfiguration.bPrintPerformanceInfo )
-				{ 
+				if (BuildConfiguration.bPrintPerformanceInfo)
+				{
 					var RunUBTInitTime = (DateTime.UtcNow - RunUBTInitStartTime).TotalSeconds;
-					Log.TraceInformation( "RunUBT initialization took " + RunUBTInitTime + "s" );
+					Log.TraceInformation("RunUBT initialization took " + RunUBTInitTime + "s");
 				}
 
-	            var TargetDescs = new List<TargetDescriptor>();
-				{ 
+				var TargetDescs = new List<TargetDescriptor>();
+				{
 					var TargetDescstStartTime = DateTime.UtcNow;
 
 					foreach (string[] TargetSetting in TargetSettings)
 					{
-						TargetDescs.AddRange( UEBuildTarget.ParseTargetCommandLine( TargetSetting, GetUProjectFile() ) );
+						TargetDescs.AddRange(UEBuildTarget.ParseTargetCommandLine(TargetSetting, GetUProjectFile()));
 					}
 
-					if( BuildConfiguration.bPrintPerformanceInfo )
-					{ 
+					if (BuildConfiguration.bPrintPerformanceInfo)
+					{
 						var TargetDescsTime = (DateTime.UtcNow - TargetDescstStartTime).TotalSeconds;
-						Log.TraceInformation( "Target descriptors took " + TargetDescsTime + "s" );
+						Log.TraceInformation("Target descriptors took " + TargetDescsTime + "s");
 					}
 				}
 
@@ -1644,107 +1644,107 @@ namespace UnrealBuildTool
 				}
 
 				UEBuildConfiguration.bHotReloadFromIDE = UEBuildConfiguration.bAllowHotReloadFromIDE && TargetDescs.Count == 1 && !TargetDescs[0].bIsEditorRecompile && ShouldDoHotReloadFromIDE(TargetDescs[0]);
-				bool bIsHotReload = UEBuildConfiguration.bHotReloadFromIDE || ( TargetDescs.Count == 1 && TargetDescs[0].OnlyModules.Count > 0 && TargetDescs[0].ForeignPlugins.Count == 0 );
+				bool bIsHotReload = UEBuildConfiguration.bHotReloadFromIDE || (TargetDescs.Count == 1 && TargetDescs[0].OnlyModules.Count > 0 && TargetDescs[0].ForeignPlugins.Count == 0);
 				TargetDescriptor HotReloadTargetDesc = bIsHotReload ? TargetDescs[0] : null;
 
 				if (ProjectFileGenerator.bGenerateProjectFiles)
 				{
 					// Create empty timestamp file to record when was the last time we regenerated projects.
-					Directory.CreateDirectory( Path.GetDirectoryName( ProjectFileGenerator.ProjectTimestampFile ) );
+					Directory.CreateDirectory(Path.GetDirectoryName(ProjectFileGenerator.ProjectTimestampFile));
 					File.Create(ProjectFileGenerator.ProjectTimestampFile).Dispose();
 				}
 
-                if( !ProjectFileGenerator.bGenerateProjectFiles )
-                {
-					if( BuildConfiguration.bUseUBTMakefiles )
-                    {
-                        // If we're building UHT without a Mutex, we'll need to assume that we're building the same targets and that no caches
-                        // should be invalidated for this run.  This is important when UBT is invoked from within UBT in order to compile
-                        // UHT.  In that (very common) case we definitely don't want to have to rebuild our cache from scratch.
-                        bool bMustAssumeSameTargets = false;
-
-                        if( TargetDescs[0].TargetName.Equals( "UnrealHeaderTool", StringComparison.InvariantCultureIgnoreCase ) )
-                        {
-                            int NoMutexArgumentIndex;
-                            if (Utils.ParseCommandLineFlag(Arguments, "-NoMutex", out NoMutexArgumentIndex))
-                            {
-                                bMustAssumeSameTargets = true;
-                            }
-                        }
-
-	                    bool bIsBuildingSameTargetsAsLastTime = false;
-                        if( bMustAssumeSameTargets )
-                        {
-                            bIsBuildingSameTargetsAsLastTime = true;
-                        }
-                        else
-                        {
-                            string TargetCollectionName = MakeTargetCollectionName( TargetDescs );
+				if (!ProjectFileGenerator.bGenerateProjectFiles)
+				{
+					if (BuildConfiguration.bUseUBTMakefiles)
+					{
+						// If we're building UHT without a Mutex, we'll need to assume that we're building the same targets and that no caches
+						// should be invalidated for this run.  This is important when UBT is invoked from within UBT in order to compile
+						// UHT.  In that (very common) case we definitely don't want to have to rebuild our cache from scratch.
+						bool bMustAssumeSameTargets = false;
+
+						if (TargetDescs[0].TargetName.Equals("UnrealHeaderTool", StringComparison.InvariantCultureIgnoreCase))
+						{
+							int NoMutexArgumentIndex;
+							if (Utils.ParseCommandLineFlag(Arguments, "-NoMutex", out NoMutexArgumentIndex))
+							{
+								bMustAssumeSameTargets = true;
+							}
+						}
+
+						bool bIsBuildingSameTargetsAsLastTime = false;
+						if (bMustAssumeSameTargets)
+						{
+							bIsBuildingSameTargetsAsLastTime = true;
+						}
+						else
+						{
+							string TargetCollectionName = MakeTargetCollectionName(TargetDescs);
 
 							string LastBuiltTargetsFileName = bIsHotReload ? "HotReloadLastBuiltTargets.txt" : "LastBuiltTargets.txt";
-                            string LastBuiltTargetsFilePath = Path.Combine( BuildConfiguration.BaseIntermediatePath, LastBuiltTargetsFileName );
-                            if( File.Exists( LastBuiltTargetsFilePath ) && Utils.ReadAllText( LastBuiltTargetsFilePath ) == TargetCollectionName )
-                            {
+							string LastBuiltTargetsFilePath = Path.Combine(BuildConfiguration.BaseIntermediatePath, LastBuiltTargetsFileName);
+							if (File.Exists(LastBuiltTargetsFilePath) && Utils.ReadAllText(LastBuiltTargetsFilePath) == TargetCollectionName)
+							{
 								// @todo ubtmake: Because we're using separate files for hot reload vs. full compiles, it's actually possible that includes will
 								// become out of date without us knowing if the developer ping-pongs between hot reloading target A and building target B normally.
 								// To fix this we can not use a different file name for last built targets, but the downside is slower performance when
 								// performing the first hot reload after compiling normally (forces full include dependency scan)
-                                bIsBuildingSameTargetsAsLastTime = true;
-                            }
-
-                            // Save out the name of the targets we're building
-                            if( !bIsBuildingSameTargetsAsLastTime )
-                            {
-                                Directory.CreateDirectory( Path.GetDirectoryName( LastBuiltTargetsFilePath ) );
-                                File.WriteAllText( LastBuiltTargetsFilePath, TargetCollectionName, Encoding.UTF8 );
-                            }
-
-
-                            if( !bIsBuildingSameTargetsAsLastTime )
-                            {
-                                // Can't use super fast include checking unless we're building the same set of targets as last time, because
-                                // we might not know about all of the C++ include dependencies for already-up-to-date shared build products
-                                // between the targets
-                                bNeedsFullCPPIncludeRescan = true;
-								Log.TraceInformation( "Performing full C++ include scan ({0} a new target)", bIsHotReload ? "hot reloading" : "building" );
-                            }
-                        }
-                    }
-                }
-            
-
-
-                UBTMakefile UBTMakefile = null;
-                { 
-                    // If we're generating project files, then go ahead and wipe out the existing UBTMakefile for every target, to make sure that
-                    // it gets a full dependency scan next time.
-                    // NOTE: This is just a safeguard and doesn't have to be perfect.  We also check for newer project file timestamps in LoadUBTMakefile()
-                    FileReference UBTMakefilePath = UnrealBuildTool.GetUBTMakefilePath( TargetDescs );
-                    if( ProjectFileGenerator.bGenerateProjectFiles )	// @todo ubtmake: This is only hit when generating IntelliSense for project files.  Probably should be done right inside ProjectFileGenerator.bat
-                    {													// @todo ubtmake: Won't catch multi-target cases as GPF always builds one target at a time for Intellisense
-                        // Delete the UBTMakefile
-                        if (UBTMakefilePath.Exists())
-                        {
-                            UEBuildTarget.CleanFile(UBTMakefilePath.FullName);
-                        }
-                    }
-
-                    // Make sure the gather phase is executed if we're not actually building anything
-                    if( ProjectFileGenerator.bGenerateProjectFiles || UEBuildConfiguration.bGenerateManifest || UEBuildConfiguration.bCleanProject || BuildConfiguration.bXGEExport || UEBuildConfiguration.bGenerateExternalFileList || GeneratingActionGraph)
-                    {
-                        UnrealBuildTool.bIsGatheringBuild_Unsafe = true;						
-                    }
-
-                    // Were we asked to run in 'assembler only' mode?  If so, let's check to see if that's even possible by seeing if
-                    // we have a valid UBTMakefile already saved to disk, ready for us to load.
-                    if( UnrealBuildTool.bIsAssemblingBuild_Unsafe && !UnrealBuildTool.bIsGatheringBuild_Unsafe )
-                    {
-                        // @todo ubtmake: Mildly terrified of BuildConfiguration/UEBuildConfiguration globals that were set during the Prepare phase but are not set now.  We may need to save/load all of these, otherwise
-                        //		we'll need to call SetupGlobalEnvironment on all of the targets (maybe other stuff, too.  See PreBuildStep())
-
-                        // Try to load the UBTMakefile.  It will only be loaded if it has valid content and is not determined to be out of date.    
+								bIsBuildingSameTargetsAsLastTime = true;
+							}
+
+							// Save out the name of the targets we're building
+							if (!bIsBuildingSameTargetsAsLastTime)
+							{
+								Directory.CreateDirectory(Path.GetDirectoryName(LastBuiltTargetsFilePath));
+								File.WriteAllText(LastBuiltTargetsFilePath, TargetCollectionName, Encoding.UTF8);
+							}
+
+
+							if (!bIsBuildingSameTargetsAsLastTime)
+							{
+								// Can't use super fast include checking unless we're building the same set of targets as last time, because
+								// we might not know about all of the C++ include dependencies for already-up-to-date shared build products
+								// between the targets
+								bNeedsFullCPPIncludeRescan = true;
+								Log.TraceInformation("Performing full C++ include scan ({0} a new target)", bIsHotReload ? "hot reloading" : "building");
+							}
+						}
+					}
+				}
+
+
+
+				UBTMakefile UBTMakefile = null;
+				{
+					// If we're generating project files, then go ahead and wipe out the existing UBTMakefile for every target, to make sure that
+					// it gets a full dependency scan next time.
+					// NOTE: This is just a safeguard and doesn't have to be perfect.  We also check for newer project file timestamps in LoadUBTMakefile()
+					FileReference UBTMakefilePath = UnrealBuildTool.GetUBTMakefilePath(TargetDescs);
+					if (ProjectFileGenerator.bGenerateProjectFiles)	// @todo ubtmake: This is only hit when generating IntelliSense for project files.  Probably should be done right inside ProjectFileGenerator.bat
+					{													// @todo ubtmake: Won't catch multi-target cases as GPF always builds one target at a time for Intellisense
+						// Delete the UBTMakefile
+						if (UBTMakefilePath.Exists())
+						{
+							UEBuildTarget.CleanFile(UBTMakefilePath.FullName);
+						}
+					}
+
+					// Make sure the gather phase is executed if we're not actually building anything
+					if (ProjectFileGenerator.bGenerateProjectFiles || UEBuildConfiguration.bGenerateManifest || UEBuildConfiguration.bCleanProject || BuildConfiguration.bXGEExport || UEBuildConfiguration.bGenerateExternalFileList || GeneratingActionGraph)
+					{
+						UnrealBuildTool.bIsGatheringBuild_Unsafe = true;
+					}
+
+					// Were we asked to run in 'assembler only' mode?  If so, let's check to see if that's even possible by seeing if
+					// we have a valid UBTMakefile already saved to disk, ready for us to load.
+					if (UnrealBuildTool.bIsAssemblingBuild_Unsafe && !UnrealBuildTool.bIsGatheringBuild_Unsafe)
+					{
+						// @todo ubtmake: Mildly terrified of BuildConfiguration/UEBuildConfiguration globals that were set during the Prepare phase but are not set now.  We may need to save/load all of these, otherwise
+						//		we'll need to call SetupGlobalEnvironment on all of the targets (maybe other stuff, too.  See PreBuildStep())
+
+						// Try to load the UBTMakefile.  It will only be loaded if it has valid content and is not determined to be out of date.    
 						string ReasonNotLoaded;
-                        UBTMakefile = LoadUBTMakefile( UBTMakefilePath, out ReasonNotLoaded );
+						UBTMakefile = LoadUBTMakefile(UBTMakefilePath, out ReasonNotLoaded);
 
 						// Invalid makefile if only modules have changed
 						if (UBTMakefile != null && !TargetDescs.SelectMany(x => x.OnlyModules)
@@ -1755,46 +1755,46 @@ namespace UnrealBuildTool
 								)
 							)
 						{
-							UBTMakefile     = null;
+							UBTMakefile = null;
 							ReasonNotLoaded = "modules to compile have changed";
 						}
 
-                        if( UBTMakefile == null )
-                        { 
-                            // If the Makefile couldn't be loaded, then we're not going to be able to continue in "assembler only" mode.  We'll do both
-                            // a 'gather' and 'assemble' in the same run.  This will take a while longer, but subsequent runs will be fast!
-                            UnrealBuildTool.bIsGatheringBuild_Unsafe = true;
+						if (UBTMakefile == null)
+						{
+							// If the Makefile couldn't be loaded, then we're not going to be able to continue in "assembler only" mode.  We'll do both
+							// a 'gather' and 'assemble' in the same run.  This will take a while longer, but subsequent runs will be fast!
+							UnrealBuildTool.bIsGatheringBuild_Unsafe = true;
 
 							FileItem.ClearCaches();
 
-                            Log.TraceInformation( "Creating makefile for {0}{1}{2} ({3})", 
-								bIsHotReload ? "hot reloading " : "", 
-								TargetDescs[0].TargetName, 
-								TargetDescs.Count > 1 ? ( " (and " + ( TargetDescs.Count - 1 ).ToString() + " more)" ) : "", 
-								ReasonNotLoaded );
-                        }
-                    }
+							Log.TraceInformation("Creating makefile for {0}{1}{2} ({3})",
+								bIsHotReload ? "hot reloading " : "",
+								TargetDescs[0].TargetName,
+								TargetDescs.Count > 1 ? (" (and " + (TargetDescs.Count - 1).ToString() + " more)") : "",
+								ReasonNotLoaded);
+						}
+					}
 
-                    // OK, after this point it is safe to access the UnrealBuildTool.IsGatheringBuild and UnrealBuildTool.IsAssemblingBuild properties.
-                    // These properties will not be changing again during this session/
-                    bIsSafeToCheckIfGatheringOrAssemblingBuild = true;
-                }
+					// OK, after this point it is safe to access the UnrealBuildTool.IsGatheringBuild and UnrealBuildTool.IsAssemblingBuild properties.
+					// These properties will not be changing again during this session/
+					bIsSafeToCheckIfGatheringOrAssemblingBuild = true;
+				}
 
 
 				List<UEBuildTarget> Targets;
-				if( UBTMakefile != null && !IsGatheringBuild && IsAssemblingBuild )
+				if (UBTMakefile != null && !IsGatheringBuild && IsAssemblingBuild)
 				{
 					// If we've loaded a makefile, then we can fill target information from this file!
 					Targets = UBTMakefile.Targets;
 				}
 				else
-				{ 
+				{
 					var TargetInitStartTime = DateTime.UtcNow;
 
 					Targets = new List<UEBuildTarget>();
-					foreach( var TargetDesc in TargetDescs )
+					foreach (var TargetDesc in TargetDescs)
 					{
-						var Target = UEBuildTarget.CreateTarget( TargetDesc );
+						var Target = UEBuildTarget.CreateTarget(TargetDesc);
 						if ((Target == null) && (UEBuildConfiguration.bCleanProject))
 						{
 							continue;
@@ -1802,194 +1802,194 @@ namespace UnrealBuildTool
 						Targets.Add(Target);
 					}
 
-					if( BuildConfiguration.bPrintPerformanceInfo )
-					{ 
+					if (BuildConfiguration.bPrintPerformanceInfo)
+					{
 						var TargetInitTime = (DateTime.UtcNow - TargetInitStartTime).TotalSeconds;
-						Log.TraceInformation( "Target init took " + TargetInitTime + "s" );
+						Log.TraceInformation("Target init took " + TargetInitTime + "s");
 					}
 				}
 
-                // Build action lists for all passed in targets.
-                var OutputItemsForAllTargets = new List<FileItem>();
-                var TargetNameToUObjectModules = new Dictionary<string, List<UHTModuleInfo>>( StringComparer.InvariantCultureIgnoreCase );
-                foreach (var Target in Targets)
-                {
-                    var TargetStartTime = DateTime.UtcNow;
+				// Build action lists for all passed in targets.
+				var OutputItemsForAllTargets = new List<FileItem>();
+				var TargetNameToUObjectModules = new Dictionary<string, List<UHTModuleInfo>>(StringComparer.InvariantCultureIgnoreCase);
+				foreach (var Target in Targets)
+				{
+					var TargetStartTime = DateTime.UtcNow;
 
-                    if (bIsHotReload)
-                    {
+					if (bIsHotReload)
+					{
 						// Don't produce new DLLs if there's been no code changes
 						UEBuildConfiguration.bSkipLinkingWhenNothingToCompile = true;
-                        Log.TraceInformation("Compiling game modules for hot reload");
-                    }
-
-                    // When in 'assembler only' mode, we'll load this cache later on a worker thread.  It takes a long time to load!
-                    if( !( !UnrealBuildTool.IsGatheringBuild && UnrealBuildTool.IsAssemblingBuild ) )
-                    {	
-                        // Load the direct include dependency cache.
-                        CPPEnvironment.IncludeDependencyCache.Add( Target, DependencyCache.Create( DependencyCache.GetDependencyCachePathForTarget(Target) ) );
-                    }
-
-                    // We don't need this dependency cache in 'gather only' mode
-                    if( BuildConfiguration.bUseUBTMakefiles &&
-                        !( UnrealBuildTool.IsGatheringBuild && !UnrealBuildTool.IsAssemblingBuild ) )
-                    { 
-                        // Load the cache that contains the list of flattened resolved includes for resolved source files
-                        // @todo ubtmake: Ideally load this asynchronously at startup and only block when it is first needed and not finished loading
-                        CPPEnvironment.FlatCPPIncludeDependencyCache.Add( Target, FlatCPPIncludeDependencyCache.Create( Target ) );
-                    }
-
-                    if( UnrealBuildTool.IsGatheringBuild )
-                    { 
-                        List<FileItem> TargetOutputItems;
-                        List<UHTModuleInfo> TargetUObjectModules;
-                        BuildResult = Target.Build(ToolChain, out TargetOutputItems, out TargetUObjectModules, out EULAViolationWarning);
-                        if(BuildResult != ECompilationResult.Succeeded)
-                        {
-                            break;
-                        }
-
-                        OutputItemsForAllTargets.AddRange( TargetOutputItems );
-
-                        // Update mapping of the target name to the list of UObject modules in this target
-                        TargetNameToUObjectModules[ Target.GetTargetName() ] = TargetUObjectModules;
-
-                        if ( (BuildConfiguration.bXGEExport && UEBuildConfiguration.bGenerateManifest) || (!ProjectFileGenerator.bGenerateProjectFiles && !UEBuildConfiguration.bGenerateManifest && !UEBuildConfiguration.bCleanProject))
-                        {
-                            // Generate an action graph if we were asked to do that.  The graph generation needs access to the include dependency cache, so
-                            // we generate it before saving and cleaning that up.
-                            if( GeneratingActionGraph )
-                            {
-                                // The graph generation feature currently only works with a single target at a time.  This is because of how we need access
-                                // to include dependencies for the target, but those aren't kept around as we process multiple targets
-                                if( TargetSettings.Count != 1 )
-                                {
-                                    throw new BuildException( "ERROR: The '-graph' option only works with a single target at a time" );
-                                }
-                                ActionGraph.FinalizeActionGraph();
-                                var ActionsToExecute = ActionGraph.AllActions;
-
-                                var VisualizationType = ActionGraph.ActionGraphVisualizationType.OnlyCPlusPlusFilesAndHeaders;
-                                ActionGraph.SaveActionGraphVisualization( Target, Path.Combine( BuildConfiguration.BaseIntermediatePath, Target.GetTargetName() + ".gexf" ), Target.GetTargetName(), VisualizationType, ActionsToExecute );
-                            }
-                        }
-
-                        var TargetBuildTime = (DateTime.UtcNow - TargetStartTime).TotalSeconds;
-
-                        // Send out telemetry for this target
-                        Telemetry.SendEvent("TargetBuildStats.2",
-                            "AppName", Target.AppName,
-                            "GameName", Target.TargetName,
-                            "Platform", Target.Platform.ToString(),
-                            "Configuration", Target.Configuration.ToString(),
-                            "CleanTarget", UEBuildConfiguration.bCleanProject.ToString(),
-                            "Monolithic", Target.ShouldCompileMonolithic().ToString(),
-                            "CreateDebugInfo", Target.IsCreatingDebugInfo().ToString(),
-                            "TargetType", Target.TargetType.ToString(),
-                            "TargetCreateTimeSec", TargetBuildTime.ToString("0.00")
-                            );
-                    }
-                }
-
-                if (BuildResult == ECompilationResult.Succeeded &&
-                    (
-                        (BuildConfiguration.bXGEExport && UEBuildConfiguration.bGenerateManifest) ||
-                        (!GeneratingActionGraph && !ProjectFileGenerator.bGenerateProjectFiles && !UEBuildConfiguration.bGenerateManifest && !UEBuildConfiguration.bCleanProject && !UEBuildConfiguration.bGenerateExternalFileList)
-                    ))
-                {
-                    if( UnrealBuildTool.IsGatheringBuild )
-                    {
-                        ActionGraph.FinalizeActionGraph();
-
-                        UBTMakefile = new UBTMakefile();
-                        UBTMakefile.AllActions = ActionGraph.AllActions;
-
-                        // For now simply treat all object files as the root target.
-                        { 
-                            var PrerequisiteActionsSet = new HashSet<Action>();
-                            foreach (var OutputItem in OutputItemsForAllTargets)
-                            {
-                                ActionGraph.GatherPrerequisiteActions(OutputItem, ref PrerequisiteActionsSet);
-                            }
-                            UBTMakefile.PrerequisiteActions = PrerequisiteActionsSet.ToArray();
-                        }			
-
-                        foreach( System.Collections.DictionaryEntry EnvironmentVariable in Environment.GetEnvironmentVariables() )
-                        {
-                            UBTMakefile.EnvironmentVariables.Add( Tuple.Create( (string)EnvironmentVariable.Key, (string)EnvironmentVariable.Value ) );
-                        }
-                        UBTMakefile.TargetNameToUObjectModules = TargetNameToUObjectModules;
+						Log.TraceInformation("Compiling game modules for hot reload");
+					}
+
+					// When in 'assembler only' mode, we'll load this cache later on a worker thread.  It takes a long time to load!
+					if (!(!UnrealBuildTool.IsGatheringBuild && UnrealBuildTool.IsAssemblingBuild))
+					{
+						// Load the direct include dependency cache.
+						CPPEnvironment.IncludeDependencyCache.Add(Target, DependencyCache.Create(DependencyCache.GetDependencyCachePathForTarget(Target)));
+					}
+
+					// We don't need this dependency cache in 'gather only' mode
+					if (BuildConfiguration.bUseUBTMakefiles &&
+						!(UnrealBuildTool.IsGatheringBuild && !UnrealBuildTool.IsAssemblingBuild))
+					{
+						// Load the cache that contains the list of flattened resolved includes for resolved source files
+						// @todo ubtmake: Ideally load this asynchronously at startup and only block when it is first needed and not finished loading
+						CPPEnvironment.FlatCPPIncludeDependencyCache.Add(Target, FlatCPPIncludeDependencyCache.Create(Target));
+					}
+
+					if (UnrealBuildTool.IsGatheringBuild)
+					{
+						List<FileItem> TargetOutputItems;
+						List<UHTModuleInfo> TargetUObjectModules;
+						BuildResult = Target.Build(ToolChain, out TargetOutputItems, out TargetUObjectModules, out EULAViolationWarning);
+						if (BuildResult != ECompilationResult.Succeeded)
+						{
+							break;
+						}
+
+						OutputItemsForAllTargets.AddRange(TargetOutputItems);
+
+						// Update mapping of the target name to the list of UObject modules in this target
+						TargetNameToUObjectModules[Target.GetTargetName()] = TargetUObjectModules;
+
+						if ((BuildConfiguration.bXGEExport && UEBuildConfiguration.bGenerateManifest) || (!ProjectFileGenerator.bGenerateProjectFiles && !UEBuildConfiguration.bGenerateManifest && !UEBuildConfiguration.bCleanProject))
+						{
+							// Generate an action graph if we were asked to do that.  The graph generation needs access to the include dependency cache, so
+							// we generate it before saving and cleaning that up.
+							if (GeneratingActionGraph)
+							{
+								// The graph generation feature currently only works with a single target at a time.  This is because of how we need access
+								// to include dependencies for the target, but those aren't kept around as we process multiple targets
+								if (TargetSettings.Count != 1)
+								{
+									throw new BuildException("ERROR: The '-graph' option only works with a single target at a time");
+								}
+								ActionGraph.FinalizeActionGraph();
+								var ActionsToExecute = ActionGraph.AllActions;
+
+								var VisualizationType = ActionGraph.ActionGraphVisualizationType.OnlyCPlusPlusFilesAndHeaders;
+								ActionGraph.SaveActionGraphVisualization(Target, Path.Combine(BuildConfiguration.BaseIntermediatePath, Target.GetTargetName() + ".gexf"), Target.GetTargetName(), VisualizationType, ActionsToExecute);
+							}
+						}
+
+						var TargetBuildTime = (DateTime.UtcNow - TargetStartTime).TotalSeconds;
+
+						// Send out telemetry for this target
+						Telemetry.SendEvent("TargetBuildStats.2",
+							"AppName", Target.AppName,
+							"GameName", Target.TargetName,
+							"Platform", Target.Platform.ToString(),
+							"Configuration", Target.Configuration.ToString(),
+							"CleanTarget", UEBuildConfiguration.bCleanProject.ToString(),
+							"Monolithic", Target.ShouldCompileMonolithic().ToString(),
+							"CreateDebugInfo", Target.IsCreatingDebugInfo().ToString(),
+							"TargetType", Target.TargetType.ToString(),
+							"TargetCreateTimeSec", TargetBuildTime.ToString("0.00")
+							);
+					}
+				}
+
+				if (BuildResult == ECompilationResult.Succeeded &&
+					(
+						(BuildConfiguration.bXGEExport && UEBuildConfiguration.bGenerateManifest) ||
+						(!GeneratingActionGraph && !ProjectFileGenerator.bGenerateProjectFiles && !UEBuildConfiguration.bGenerateManifest && !UEBuildConfiguration.bCleanProject && !UEBuildConfiguration.bGenerateExternalFileList)
+					))
+				{
+					if (UnrealBuildTool.IsGatheringBuild)
+					{
+						ActionGraph.FinalizeActionGraph();
+
+						UBTMakefile = new UBTMakefile();
+						UBTMakefile.AllActions = ActionGraph.AllActions;
+
+						// For now simply treat all object files as the root target.
+						{
+							var PrerequisiteActionsSet = new HashSet<Action>();
+							foreach (var OutputItem in OutputItemsForAllTargets)
+							{
+								ActionGraph.GatherPrerequisiteActions(OutputItem, ref PrerequisiteActionsSet);
+							}
+							UBTMakefile.PrerequisiteActions = PrerequisiteActionsSet.ToArray();
+						}
+
+						foreach (System.Collections.DictionaryEntry EnvironmentVariable in Environment.GetEnvironmentVariables())
+						{
+							UBTMakefile.EnvironmentVariables.Add(Tuple.Create((string)EnvironmentVariable.Key, (string)EnvironmentVariable.Value));
+						}
+						UBTMakefile.TargetNameToUObjectModules = TargetNameToUObjectModules;
 						UBTMakefile.Targets = Targets;
 						UBTMakefile.SourceFileWorkingSet = Unity.SourceFileWorkingSet;
 
-						if( BuildConfiguration.bUseUBTMakefiles )
-						{ 
+						if (BuildConfiguration.bUseUBTMakefiles)
+						{
 							// We've been told to prepare to build, so let's go ahead and save out our action graph so that we can use in a later invocation 
 							// to assemble the build.  Even if we are configured to assemble the build in this same invocation, we want to save out the
 							// Makefile so that it can be used on subsequent 'assemble only' runs, for the fastest possible iteration times
 							// @todo ubtmake: Optimization: We could make 'gather + assemble' mode slightly faster by saving this while busy compiling (on our worker thread)
-							SaveUBTMakefile( TargetDescs, UBTMakefile );
+							SaveUBTMakefile(TargetDescs, UBTMakefile);
 						}
-                    }
+					}
 
-                    if( UnrealBuildTool.IsAssemblingBuild )
-                    {
-                        // If we didn't build the graph in this session, then we'll need to load a cached one
-                        if( !UnrealBuildTool.IsGatheringBuild )
-                        {
-                            ActionGraph.AllActions = UBTMakefile.AllActions;
+					if (UnrealBuildTool.IsAssemblingBuild)
+					{
+						// If we didn't build the graph in this session, then we'll need to load a cached one
+						if (!UnrealBuildTool.IsGatheringBuild)
+						{
+							ActionGraph.AllActions = UBTMakefile.AllActions;
 
 							// Patch action history for hot reload when running in assembler mode.  In assembler mode, the suffix on the output file will be
 							// the same for every invocation on that makefile, but we need a new suffix each time.
-							if( bIsHotReload )
+							if (bIsHotReload)
+							{
+								PatchActionHistoryForHotReloadAssembling(HotReloadTargetDesc.OnlyModules);
+							}
+
+
+							foreach (var EnvironmentVariable in UBTMakefile.EnvironmentVariables)
+							{
+								// @todo ubtmake: There may be some variables we do NOT want to clobber.
+								Environment.SetEnvironmentVariable(EnvironmentVariable.Item1, EnvironmentVariable.Item2);
+							}
+
+							// If any of the targets need UHT to be run, we'll go ahead and do that now
+							foreach (var Target in Targets)
 							{
-								PatchActionHistoryForHotReloadAssembling( HotReloadTargetDesc.OnlyModules );
+								List<UHTModuleInfo> TargetUObjectModules;
+								if (UBTMakefile.TargetNameToUObjectModules.TryGetValue(Target.GetTargetName(), out TargetUObjectModules))
+								{
+									if (TargetUObjectModules.Count > 0)
+									{
+										// Execute the header tool
+										FileReference ModuleInfoFileName = FileReference.Combine(Target.ProjectIntermediateDirectory, "UnrealHeaderTool.manifest");
+										ECompilationResult UHTResult = ECompilationResult.OtherCompilationError;
+										if (!ExternalExecution.ExecuteHeaderToolIfNecessary(ToolChain, Target, GlobalCompileEnvironment: null, UObjectModules: TargetUObjectModules, ModuleInfoFileName: ModuleInfoFileName, UHTResult: ref UHTResult))
+										{
+											Log.TraceInformation("UnrealHeaderTool failed for target '" + Target.GetTargetName() + "' (platform: " + Target.Platform.ToString() + ", module info: " + ModuleInfoFileName + ").");
+											BuildResult = UHTResult;
+											break;
+										}
+									}
+								}
 							}
+						}
+
+						if (BuildResult.Succeeded())
+						{
+							// Make sure any old DLL files from in-engine recompiles aren't lying around.  Must be called after the action graph is finalized.
+							ActionGraph.DeleteStaleHotReloadDLLs();
 
+							// Plan the actions to execute for the build.
+							Dictionary<UEBuildTarget, List<FileItem>> TargetToOutdatedPrerequisitesMap;
+							List<Action> ActionsToExecute = ActionGraph.GetActionsToExecute(UBTMakefile.PrerequisiteActions, Targets, out TargetToOutdatedPrerequisitesMap);
 
-                            foreach( var EnvironmentVariable in UBTMakefile.EnvironmentVariables )
-                            {
-                                // @todo ubtmake: There may be some variables we do NOT want to clobber.
-                                Environment.SetEnvironmentVariable( EnvironmentVariable.Item1, EnvironmentVariable.Item2 );
-                            }
-
-                            // If any of the targets need UHT to be run, we'll go ahead and do that now
-                            foreach( var Target in Targets )
-                            {
-                                List<UHTModuleInfo> TargetUObjectModules;
-                                if( UBTMakefile.TargetNameToUObjectModules.TryGetValue( Target.GetTargetName(), out TargetUObjectModules ) )
-                                {
-                                    if( TargetUObjectModules.Count > 0 )
-                                    {
-                                        // Execute the header tool
-                                        FileReference ModuleInfoFileName = FileReference.Combine(Target.ProjectIntermediateDirectory, "UnrealHeaderTool.manifest");
-                                        ECompilationResult UHTResult = ECompilationResult.OtherCompilationError;
-                                        if (!ExternalExecution.ExecuteHeaderToolIfNecessary(ToolChain, Target, GlobalCompileEnvironment:null, UObjectModules:TargetUObjectModules, ModuleInfoFileName:ModuleInfoFileName, UHTResult:ref UHTResult))
-                                        {
-                                            Log.TraceInformation("UnrealHeaderTool failed for target '" + Target.GetTargetName() + "' (platform: " + Target.Platform.ToString() + ", module info: " + ModuleInfoFileName + ").");
-                                            BuildResult = UHTResult;
-                                            break;
-                                        }
-                                    }
-                                }
-                            }
-                        }
-
-                        if( BuildResult.Succeeded() )
-                        { 
-                            // Make sure any old DLL files from in-engine recompiles aren't lying around.  Must be called after the action graph is finalized.
-                            ActionGraph.DeleteStaleHotReloadDLLs();
-
-                            // Plan the actions to execute for the build.
-                            Dictionary<UEBuildTarget,List<FileItem>> TargetToOutdatedPrerequisitesMap;
-                            List<Action> ActionsToExecute = ActionGraph.GetActionsToExecute(UBTMakefile.PrerequisiteActions, Targets, out TargetToOutdatedPrerequisitesMap);
-
-                            // Display some stats to the user.
-                            Log.TraceVerbose(
-                                    "{0} actions, {1} outdated and requested actions",
-                                    ActionGraph.AllActions.Count,
-                                    ActionsToExecute.Count
-                                    );
+							// Display some stats to the user.
+							Log.TraceVerbose(
+									"{0} actions, {1} outdated and requested actions",
+									ActionGraph.AllActions.Count,
+									ActionsToExecute.Count
+									);
 
 							if (!bIsHotReload)
 							{
@@ -1999,136 +1999,136 @@ namespace UnrealBuildTool
 									Target.CleanStaleModules();
 								}
 							}
-							
+
 							ToolChain.PreBuildSync();
 
-                    
-                            // Cache indirect includes for all outdated C++ files.  We kick this off as a background thread so that it can
-                            // perform the scan while we're compiling.  It usually only takes up to a few seconds, but we don't want to hurt
-                            // our best case UBT iteration times for this task which can easily be performed asynchronously
-                            if( BuildConfiguration.bUseUBTMakefiles && TargetToOutdatedPrerequisitesMap.Count > 0 )
-                            {
-                                CPPIncludesThread = CreateThreadForCachingCPPIncludes( TargetToOutdatedPrerequisitesMap );
-                                CPPIncludesThread.Start();
-                            }
-
-                            // Execute the actions.
-                            bSuccess = ActionGraph.ExecuteActions(ActionsToExecute, out ExecutorName);
-
-                            // if the build succeeded, write the receipts and do any needed syncing
-                            if (bSuccess)
-                            {
-                                foreach (UEBuildTarget Target in Targets)
-                                {
+
+							// Cache indirect includes for all outdated C++ files.  We kick this off as a background thread so that it can
+							// perform the scan while we're compiling.  It usually only takes up to a few seconds, but we don't want to hurt
+							// our best case UBT iteration times for this task which can easily be performed asynchronously
+							if (BuildConfiguration.bUseUBTMakefiles && TargetToOutdatedPrerequisitesMap.Count > 0)
+							{
+								CPPIncludesThread = CreateThreadForCachingCPPIncludes(TargetToOutdatedPrerequisitesMap);
+								CPPIncludesThread.Start();
+							}
+
+							// Execute the actions.
+							bSuccess = ActionGraph.ExecuteActions(ActionsToExecute, out ExecutorName);
+
+							// if the build succeeded, write the receipts and do any needed syncing
+							if (bSuccess)
+							{
+								foreach (UEBuildTarget Target in Targets)
+								{
 									Target.WriteReceipts();
-                                    ToolChain.PostBuildSync(Target);
-                                }
+									ToolChain.PostBuildSync(Target);
+								}
 								if (ActionsToExecute.Count == 0 && UEBuildConfiguration.bSkipLinkingWhenNothingToCompile)
 								{
 									BuildResult = ECompilationResult.UpToDate;
 								}
-                            }
-                            else
-                            {
-                                BuildResult = ECompilationResult.OtherCompilationError;
-                            }
-                        }
-                    }
-                }
-            }
-            catch (BuildException Exception)
-            {
-                // Output the message only, without the call stack
-                Log.TraceInformation(Exception.Message);
-                BuildResult = ECompilationResult.OtherCompilationError;
-            }
-            catch (Exception Exception)
-            {
-                Log.TraceInformation("ERROR: {0}", Exception);
-                BuildResult = ECompilationResult.OtherCompilationError;
-            }
+							}
+							else
+							{
+								BuildResult = ECompilationResult.OtherCompilationError;
+							}
+						}
+					}
+				}
+			}
+			catch (BuildException Exception)
+			{
+				// Output the message only, without the call stack
+				Log.TraceInformation(Exception.Message);
+				BuildResult = ECompilationResult.OtherCompilationError;
+			}
+			catch (Exception Exception)
+			{
+				Log.TraceInformation("ERROR: {0}", Exception);
+				BuildResult = ECompilationResult.OtherCompilationError;
+			}
 
 			// Wait until our CPPIncludes dependency scanner thread has finished
-			if( CPPIncludesThread != null )
-			{ 
-			    CPPIncludesThread.Join();
+			if (CPPIncludesThread != null)
+			{
+				CPPIncludesThread.Join();
 			}
 
-            // Save the include dependency cache.
-            { 
+			// Save the include dependency cache.
+			{
 				// NOTE: It's very important that we save the include cache, even if a build exception was thrown (compile error, etc), because we need to make sure that
 				//    any C++ include dependencies that we computed for out of date source files are saved.  Remember, the build may fail *after* some build products
 				//    are successfully built.  If we didn't save our dependency cache after build failures, source files for those build products that were successsfully
 				//    built before the failure would not be considered out of date on the next run, so this is our only chance to cache C++ includes for those files!
 
-                foreach( var DependencyCache in CPPEnvironment.IncludeDependencyCache.Values )
-                { 
-                    DependencyCache.Save();
-                }
-                CPPEnvironment.IncludeDependencyCache.Clear();
-
-                foreach( var FlatCPPIncludeDependencyCache in CPPEnvironment.FlatCPPIncludeDependencyCache.Values )
-                { 
-                    FlatCPPIncludeDependencyCache.Save();
-                }
-                CPPEnvironment.FlatCPPIncludeDependencyCache.Clear();
-            }
-
-            if(EULAViolationWarning != null)
-            {
-                Log.TraceWarning("WARNING: {0}", EULAViolationWarning);
-            }
-
-            // Figure out how long we took to execute.
-            double BuildDuration = (DateTime.UtcNow - StartTime).TotalSeconds;
-            if (ExecutorName == "Local" || ExecutorName == "Distcc" || ExecutorName == "SNDBS")
-            {
-                Log.WriteLineIf(BuildConfiguration.bLogDetailedActionStats || BuildConfiguration.bPrintDebugInfo,
-					LogEventType.Console, 
+				foreach (var DependencyCache in CPPEnvironment.IncludeDependencyCache.Values)
+				{
+					DependencyCache.Save();
+				}
+				CPPEnvironment.IncludeDependencyCache.Clear();
+
+				foreach (var FlatCPPIncludeDependencyCache in CPPEnvironment.FlatCPPIncludeDependencyCache.Values)
+				{
+					FlatCPPIncludeDependencyCache.Save();
+				}
+				CPPEnvironment.FlatCPPIncludeDependencyCache.Clear();
+			}
+
+			if (EULAViolationWarning != null)
+			{
+				Log.TraceWarning("WARNING: {0}", EULAViolationWarning);
+			}
+
+			// Figure out how long we took to execute.
+			double BuildDuration = (DateTime.UtcNow - StartTime).TotalSeconds;
+			if (ExecutorName == "Local" || ExecutorName == "Distcc" || ExecutorName == "SNDBS")
+			{
+				Log.WriteLineIf(BuildConfiguration.bLogDetailedActionStats || BuildConfiguration.bPrintDebugInfo,
+					LogEventType.Console,
 					"Cumulative action seconds ({0} processors): {1:0.00} building projects, {2:0.00} compiling, {3:0.00} creating app bundles, {4:0.00} generating debug info, {5:0.00} linking, {6:0.00} other",
-                    System.Environment.ProcessorCount,
-                    TotalBuildProjectTime,
-                    TotalCompileTime,
-                    TotalCreateAppBundleTime,
-                    TotalGenerateDebugInfoTime,
-                    TotalLinkTime,
-                    TotalOtherActionsTime
-                );
-                //Telemetry.SendEvent("BuildStatsTotal.2",
-                //    "ExecutorName", ExecutorName,
-                //    "TotalUBTWallClockTimeSec", BuildDuration.ToString("0.00"),
-                //    "TotalBuildProjectThreadTimeSec", TotalBuildProjectTime.ToString("0.00"),
-                //    "TotalCompileThreadTimeSec", TotalCompileTime.ToString("0.00"),
-                //    "TotalCreateAppBundleThreadTimeSec", TotalCreateAppBundleTime.ToString("0.00"),
-                //    "TotalGenerateDebugInfoThreadTimeSec", TotalGenerateDebugInfoTime.ToString("0.00"),
-                //    "TotalLinkThreadTimeSec", TotalLinkTime.ToString("0.00"),
-                //    "TotalOtherActionsThreadTimeSec", TotalOtherActionsTime.ToString("0.00")
-                //    );
-
-                Log.TraceInformation("Total build time: {0:0.00} seconds", BuildDuration);
-
-                // reset statistics
-                TotalBuildProjectTime = 0;
-                TotalCompileTime = 0;
-                TotalCreateAppBundleTime = 0;
-                TotalGenerateDebugInfoTime = 0;
-                TotalLinkTime = 0;
-                TotalOtherActionsTime = 0;
-            }
-            else
-            {
-                if (ExecutorName == "XGE")
-                {
-                    Log.TraceInformation("XGE execution time: {0:0.00} seconds", BuildDuration);
-                }
-                //Telemetry.SendEvent("BuildStatsTotal.2",
-                //    "ExecutorName", ExecutorName,
-                //    "TotalUBTWallClockTimeSec", BuildDuration.ToString("0.00")
-                //    );
-            }
+					System.Environment.ProcessorCount,
+					TotalBuildProjectTime,
+					TotalCompileTime,
+					TotalCreateAppBundleTime,
+					TotalGenerateDebugInfoTime,
+					TotalLinkTime,
+					TotalOtherActionsTime
+				);
+				//Telemetry.SendEvent("BuildStatsTotal.2",
+				//    "ExecutorName", ExecutorName,
+				//    "TotalUBTWallClockTimeSec", BuildDuration.ToString("0.00"),
+				//    "TotalBuildProjectThreadTimeSec", TotalBuildProjectTime.ToString("0.00"),
+				//    "TotalCompileThreadTimeSec", TotalCompileTime.ToString("0.00"),
+				//    "TotalCreateAppBundleThreadTimeSec", TotalCreateAppBundleTime.ToString("0.00"),
+				//    "TotalGenerateDebugInfoThreadTimeSec", TotalGenerateDebugInfoTime.ToString("0.00"),
+				//    "TotalLinkThreadTimeSec", TotalLinkTime.ToString("0.00"),
+				//    "TotalOtherActionsThreadTimeSec", TotalOtherActionsTime.ToString("0.00")
+				//    );
+
+				Log.TraceInformation("Total build time: {0:0.00} seconds", BuildDuration);
+
+				// reset statistics
+				TotalBuildProjectTime = 0;
+				TotalCompileTime = 0;
+				TotalCreateAppBundleTime = 0;
+				TotalGenerateDebugInfoTime = 0;
+				TotalLinkTime = 0;
+				TotalOtherActionsTime = 0;
+			}
+			else
+			{
+				if (ExecutorName == "XGE")
+				{
+					Log.TraceInformation("XGE execution time: {0:0.00} seconds", BuildDuration);
+				}
+				//Telemetry.SendEvent("BuildStatsTotal.2",
+				//    "ExecutorName", ExecutorName,
+				//    "TotalUBTWallClockTimeSec", BuildDuration.ToString("0.00")
+				//    );
+			}
 
-            return BuildResult;
-        }
+			return BuildResult;
+		}
 
 		/// <summary>
 		/// Invalidates makefiles for given target.
@@ -2164,12 +2164,12 @@ namespace UnrealBuildTool
 			}
 
 			bool bIsRunning = false;
-			
+
 			// @todo ubtmake: Kind of cheating here to figure out if an editor target.  At this point we don't have access to the actual target description, and
 			// this code must be able to execute before we create or load module rules DLLs so that hot reload can work with bUseUBTMakefiles
-			bool bIsEditorTarget = TargetDesc.TargetName.EndsWith( "Editor", StringComparison.InvariantCultureIgnoreCase ) || TargetDesc.bIsEditorRecompile;	
+			bool bIsEditorTarget = TargetDesc.TargetName.EndsWith("Editor", StringComparison.InvariantCultureIgnoreCase) || TargetDesc.bIsEditorRecompile;
 
-			if (!ProjectFileGenerator.bGenerateProjectFiles && !UEBuildConfiguration.bGenerateManifest && bIsEditorTarget )
+			if (!ProjectFileGenerator.bGenerateProjectFiles && !UEBuildConfiguration.bGenerateManifest && bIsEditorTarget)
 			{
 				var EditorProcessFilenames = UEBuildTarget.MakeExecutablePaths(UnrealBuildTool.EngineDirectory, "UE4Editor", TargetDesc.Platform, TargetDesc.Configuration, UnrealTargetConfiguration.Development, false, null);
 				if (EditorProcessFilenames.Count != 1)
@@ -2181,21 +2181,21 @@ namespace UnrealBuildTool
 				var EditorProcessName = Path.GetFileNameWithoutExtension(EditorProcessFilename);
 				var EditorProcesses = BuildHostPlatform.Current.GetProcessesByName(EditorProcessName);
 				var BinariesPath = Path.GetFullPath(Path.GetDirectoryName(EditorProcessFilename));
-				var PerProjectBinariesPath = (TargetDesc.ProjectFile != null)? Path.Combine(TargetDesc.ProjectFile.Directory.FullName, BinariesPath.Substring(BinariesPath.LastIndexOf("Binaries"))) : "";
+				var PerProjectBinariesPath = (TargetDesc.ProjectFile != null) ? Path.Combine(TargetDesc.ProjectFile.Directory.FullName, BinariesPath.Substring(BinariesPath.LastIndexOf("Binaries"))) : "";
 				bIsRunning = EditorProcesses.FirstOrDefault(EditorProc =>
 					{
-						if(!Path.GetFullPath(EditorProc.Filename).StartsWith(BinariesPath, StringComparison.InvariantCultureIgnoreCase))
+						if (!Path.GetFullPath(EditorProc.Filename).StartsWith(BinariesPath, StringComparison.InvariantCultureIgnoreCase))
 						{
 							return false;
 						}
 
-						if(PerProjectBinariesPath.Equals(BinariesPath, StringComparison.InvariantCultureIgnoreCase))
+						if (PerProjectBinariesPath.Equals(BinariesPath, StringComparison.InvariantCultureIgnoreCase))
 						{
 							return false;
 						}
 
 						var Modules = BuildHostPlatform.Current.GetProcessModules(EditorProc.PID, EditorProc.Filename);
-						if(!Modules.Any(Module => Module.StartsWith(PerProjectBinariesPath)))
+						if (!Modules.Any(Module => Module.StartsWith(PerProjectBinariesPath)))
 						{
 							return false;
 						}
@@ -2206,95 +2206,95 @@ namespace UnrealBuildTool
 			return bIsRunning;
 		}
 
-        /**
-         * Parses the passed in command line for build configuration overrides.
-         * 
-         * @param	Arguments	List of arguments to parse
-         * @return	List of build target settings
-         */
-        private static List<string[]> ParseCommandLineFlags(string[] Arguments)
-        {
-            var TargetSettings = new List<string[]>();
-            int ArgumentIndex = 0;
-
-            if (Utils.ParseCommandLineFlag(Arguments, "-targets", out ArgumentIndex))
-            {
-                if (ArgumentIndex + 1 >= Arguments.Length)
-                {
-                    throw new BuildException("Expected filename after -targets argument, but found nothing.");
-                }
-                // Parse lines from the referenced file into target settings.
-                var Lines = File.ReadAllLines(Arguments[ArgumentIndex + 1]);
-                foreach (string Line in Lines)
-                {
-                    if (Line != "" && Line[0] != ';')
-                    {
-                        TargetSettings.Add(Line.Split(' '));
-                    }
-                }
-            }
-            // Simply use full command line arguments as target setting if not otherwise overridden.
-            else
-            {
-                TargetSettings.Add(Arguments);
-            }
-
-            return TargetSettings;
-        }
-
-
-        /// <summary>
-        /// Returns a Thread object that can be kicked off to update C++ include dependency cache
-        /// </summary>
-        /// <param name="TargetToOutdatedPrerequisitesMap">Maps each target to a list of outdated C++ files that need indirect dependencies cached</param>
-        /// <returns>The thread object</returns>
-        private static Thread CreateThreadForCachingCPPIncludes( Dictionary<UEBuildTarget, List<FileItem>> TargetToOutdatedPrerequisitesMap)
-        {
-            return new Thread( new ThreadStart( () =>
-                {
-                    // @todo ubtmake: This thread will access data structures that are also used on the main UBT thread, but during this time UBT
-                    // is only invoking the build executor, so should not be touching this stuff.  However, we need to at some guards to make sure.
-
-                    foreach( var TargetAndOutdatedPrerequisites in TargetToOutdatedPrerequisitesMap )
-                    {
-                        var Target = TargetAndOutdatedPrerequisites.Key;
-                        var OutdatedPrerequisites = TargetAndOutdatedPrerequisites.Value;
-
-                        foreach( var PrerequisiteItem in OutdatedPrerequisites )
-                        {
-                            // Invoke our deep include scanner to figure out whether any of the files included by this source file have
-                            // changed since the build product was built
-                            var FileBuildPlatform = UEBuildPlatform.GetBuildPlatform( Target.Platform );
-                            CPPEnvironment.FindAndCacheAllIncludedFiles( Target, PrerequisiteItem, FileBuildPlatform, PrerequisiteItem.CachedCPPIncludeInfo, bOnlyCachedDependencies:false );
-                        }
-                    }
-                } ) );
-        }
-
-
-
-        /// <summary>
-        /// A special Makefile that UBT is able to create in "-gather" mode, then load in "-assemble" mode to accelerate iterative compiling and linking
-        /// </summary>
-        [Serializable]
-        class UBTMakefile : ISerializable
-        {
+		/**
+		 * Parses the passed in command line for build configuration overrides.
+		 * 
+		 * @param	Arguments	List of arguments to parse
+		 * @return	List of build target settings
+		 */
+		private static List<string[]> ParseCommandLineFlags(string[] Arguments)
+		{
+			var TargetSettings = new List<string[]>();
+			int ArgumentIndex = 0;
+
+			if (Utils.ParseCommandLineFlag(Arguments, "-targets", out ArgumentIndex))
+			{
+				if (ArgumentIndex + 1 >= Arguments.Length)
+				{
+					throw new BuildException("Expected filename after -targets argument, but found nothing.");
+				}
+				// Parse lines from the referenced file into target settings.
+				var Lines = File.ReadAllLines(Arguments[ArgumentIndex + 1]);
+				foreach (string Line in Lines)
+				{
+					if (Line != "" && Line[0] != ';')
+					{
+						TargetSettings.Add(Line.Split(' '));
+					}
+				}
+			}
+			// Simply use full command line arguments as target setting if not otherwise overridden.
+			else
+			{
+				TargetSettings.Add(Arguments);
+			}
+
+			return TargetSettings;
+		}
+
+
+		/// <summary>
+		/// Returns a Thread object that can be kicked off to update C++ include dependency cache
+		/// </summary>
+		/// <param name="TargetToOutdatedPrerequisitesMap">Maps each target to a list of outdated C++ files that need indirect dependencies cached</param>
+		/// <returns>The thread object</returns>
+		private static Thread CreateThreadForCachingCPPIncludes(Dictionary<UEBuildTarget, List<FileItem>> TargetToOutdatedPrerequisitesMap)
+		{
+			return new Thread(new ThreadStart(() =>
+				{
+					// @todo ubtmake: This thread will access data structures that are also used on the main UBT thread, but during this time UBT
+					// is only invoking the build executor, so should not be touching this stuff.  However, we need to at some guards to make sure.
+
+					foreach (var TargetAndOutdatedPrerequisites in TargetToOutdatedPrerequisitesMap)
+					{
+						var Target = TargetAndOutdatedPrerequisites.Key;
+						var OutdatedPrerequisites = TargetAndOutdatedPrerequisites.Value;
+
+						foreach (var PrerequisiteItem in OutdatedPrerequisites)
+						{
+							// Invoke our deep include scanner to figure out whether any of the files included by this source file have
+							// changed since the build product was built
+							var FileBuildPlatform = UEBuildPlatform.GetBuildPlatform(Target.Platform);
+							CPPEnvironment.FindAndCacheAllIncludedFiles(Target, PrerequisiteItem, FileBuildPlatform, PrerequisiteItem.CachedCPPIncludeInfo, bOnlyCachedDependencies: false);
+						}
+					}
+				}));
+		}
+
+
+
+		/// <summary>
+		/// A special Makefile that UBT is able to create in "-gather" mode, then load in "-assemble" mode to accelerate iterative compiling and linking
+		/// </summary>
+		[Serializable]
+		class UBTMakefile : ISerializable
+		{
 			public const int CurrentVersion = 3;
 
 			public int Version;
 
-            /** Every action in the action graph */
-            public List<Action> AllActions;
+			/** Every action in the action graph */
+			public List<Action> AllActions;
 
-            /** List of the actions that need to be run in order to build the targets' final output items */
-            public Action[] PrerequisiteActions;
+			/** List of the actions that need to be run in order to build the targets' final output items */
+			public Action[] PrerequisiteActions;
 
-            /** Environment variables that we'll need in order to invoke the platform's compiler and linker */
+			/** Environment variables that we'll need in order to invoke the platform's compiler and linker */
 			// @todo ubtmake: Really we want to allow a different set of environment variables for every Action.  This would allow for targets on multiple platforms to be built in a single assembling phase.  We'd only have unique variables for each platform that has actions, so we'd want to make sure we only store the minimum set.
-            public readonly List<Tuple<string,string>> EnvironmentVariables = new List<Tuple<string,string>>();
+			public readonly List<Tuple<string, string>> EnvironmentVariables = new List<Tuple<string, string>>();
 
-            /** Maps each target to a list of UObject module info structures */
-            public Dictionary<string,List<UHTModuleInfo>> TargetNameToUObjectModules;
+			/** Maps each target to a list of UObject module info structures */
+			public Dictionary<string, List<UHTModuleInfo>> TargetNameToUObjectModules;
 
 			/** List of targets being built */
 			public List<UEBuildTarget> Targets;
@@ -2316,12 +2316,12 @@ namespace UnrealBuildTool
 					throw new Exception(string.Format("Makefile version does not match - found {0}, expected: {1}", Version, CurrentVersion));
 				}
 
-				AllActions                 = (List<Action>)Info.GetValue("ac", typeof(List<Action>));
-				PrerequisiteActions        = (Action[])Info.GetValue("pa", typeof(Action[]));
-				EnvironmentVariables       = ((string[])Info.GetValue("e1", typeof(string[]))).Zip((string[])Info.GetValue("e2", typeof(string[])), (i1, i2) => new Tuple<string, string>(i1, i2)).ToList();
-				TargetNameToUObjectModules = (Dictionary<string,List<UHTModuleInfo>>)Info.GetValue("nu", typeof(Dictionary<string,List<UHTModuleInfo>>));
-				Targets                    = (List<UEBuildTarget>)Info.GetValue("ta", typeof(List<UEBuildTarget>));
-				SourceFileWorkingSet       = (HashSet<FileItem>)Info.GetValue("ws", typeof(HashSet<FileItem>));
+				AllActions = (List<Action>)Info.GetValue("ac", typeof(List<Action>));
+				PrerequisiteActions = (Action[])Info.GetValue("pa", typeof(Action[]));
+				EnvironmentVariables = ((string[])Info.GetValue("e1", typeof(string[]))).Zip((string[])Info.GetValue("e2", typeof(string[])), (i1, i2) => new Tuple<string, string>(i1, i2)).ToList();
+				TargetNameToUObjectModules = (Dictionary<string, List<UHTModuleInfo>>)Info.GetValue("nu", typeof(Dictionary<string, List<UHTModuleInfo>>));
+				Targets = (List<UEBuildTarget>)Info.GetValue("ta", typeof(List<UEBuildTarget>));
+				SourceFileWorkingSet = (HashSet<FileItem>)Info.GetValue("ws", typeof(HashSet<FileItem>));
 			}
 
 			public void GetObjectData(SerializationInfo Info, StreamingContext Context)
@@ -2337,73 +2337,73 @@ namespace UnrealBuildTool
 			}
 
 
-            /** @return	 True if this makefile's contents look valid.  Called after loading the file to make sure it is legit. */
-            public bool IsValidMakefile()
-            {
-                return 
-                    AllActions != null && AllActions.Count > 0 &&
-                    PrerequisiteActions != null && PrerequisiteActions.Length > 0 &&
-                    EnvironmentVariables != null &&
-                    TargetNameToUObjectModules != null && TargetNameToUObjectModules.Count > 0 &&
+			/** @return	 True if this makefile's contents look valid.  Called after loading the file to make sure it is legit. */
+			public bool IsValidMakefile()
+			{
+				return
+					AllActions != null && AllActions.Count > 0 &&
+					PrerequisiteActions != null && PrerequisiteActions.Length > 0 &&
+					EnvironmentVariables != null &&
+					TargetNameToUObjectModules != null && TargetNameToUObjectModules.Count > 0 &&
 					Targets != null && Targets.Count > 0 &&
 					SourceFileWorkingSet != null;
-            }
-        }
+			}
+		}
 
 
-        /// <summary>
-        /// Saves a UBTMakefile to disk
-        /// </summary>
-        /// <param name="TargetDescs">List of targets.  Order is not important</param>
-        static void SaveUBTMakefile( List<TargetDescriptor> TargetDescs, UBTMakefile UBTMakefile )
-        {
-            if( !UBTMakefile.IsValidMakefile() )
-            {
-                throw new BuildException( "Can't save a makefile that has invalid contents.  See UBTMakefile.IsValidMakefile()" );
-            }
+		/// <summary>
+		/// Saves a UBTMakefile to disk
+		/// </summary>
+		/// <param name="TargetDescs">List of targets.  Order is not important</param>
+		static void SaveUBTMakefile(List<TargetDescriptor> TargetDescs, UBTMakefile UBTMakefile)
+		{
+			if (!UBTMakefile.IsValidMakefile())
+			{
+				throw new BuildException("Can't save a makefile that has invalid contents.  See UBTMakefile.IsValidMakefile()");
+			}
 
-            var TimerStartTime = DateTime.UtcNow;
+			var TimerStartTime = DateTime.UtcNow;
 
-            var UBTMakefileItem = FileItem.GetItemByFileReference( GetUBTMakefilePath( TargetDescs ) );
+			var UBTMakefileItem = FileItem.GetItemByFileReference(GetUBTMakefilePath(TargetDescs));
 
-            // @todo ubtmake: Optimization: The UBTMakefile saved for game projects is upwards of 9 MB.  We should try to shrink its content if possible
-            // @todo ubtmake: Optimization: C# Serialization may be too slow for these big Makefiles.  Loading these files often shows up as the slower part of the assembling phase.
+			// @todo ubtmake: Optimization: The UBTMakefile saved for game projects is upwards of 9 MB.  We should try to shrink its content if possible
+			// @todo ubtmake: Optimization: C# Serialization may be too slow for these big Makefiles.  Loading these files often shows up as the slower part of the assembling phase.
 
-            // Serialize the cache to disk.
-            try
-            {
-                Directory.CreateDirectory(Path.GetDirectoryName(UBTMakefileItem.AbsolutePath));
-                using (FileStream Stream = new FileStream(UBTMakefileItem.AbsolutePath, FileMode.Create, FileAccess.Write))
-                {
-                    BinaryFormatter Formatter = new BinaryFormatter();
-                    Formatter.Serialize(Stream, UBTMakefile);
-                }
-            }
-            catch (Exception Ex)
-            {
-                Console.Error.WriteLine("Failed to write makefile: {0}", Ex.Message);
-            }
-
-            if (BuildConfiguration.bPrintPerformanceInfo)
-            {
-                var TimerDuration = DateTime.UtcNow - TimerStartTime;
-                Log.TraceInformation("Saving makefile took " + TimerDuration.TotalSeconds + "s");
-            }
-        }
-
-
-        /// <summary>
-        /// Loads a UBTMakefile from disk
-        /// </summary>
-        /// <param name="MakefilePath">Path to the makefile to load</param>
+			// Serialize the cache to disk.
+			try
+			{
+				Directory.CreateDirectory(Path.GetDirectoryName(UBTMakefileItem.AbsolutePath));
+				using (FileStream Stream = new FileStream(UBTMakefileItem.AbsolutePath, FileMode.Create, FileAccess.Write))
+				{
+					BinaryFormatter Formatter = new BinaryFormatter();
+					Formatter.Serialize(Stream, UBTMakefile);
+				}
+			}
+			catch (Exception Ex)
+			{
+				Console.Error.WriteLine("Failed to write makefile: {0}", Ex.Message);
+			}
+
+			if (BuildConfiguration.bPrintPerformanceInfo)
+			{
+				var TimerDuration = DateTime.UtcNow - TimerStartTime;
+				Log.TraceInformation("Saving makefile took " + TimerDuration.TotalSeconds + "s");
+			}
+		}
+
+
+		/// <summary>
+		/// Loads a UBTMakefile from disk
+		/// </summary>
+		/// <param name="MakefilePath">Path to the makefile to load</param>
 		/// <param name="ReasonNotLoaded">If the function returns null, this string will contain the reason why</param>
 		/// <returns>The loaded makefile, or null if it failed for some reason.  On failure, the 'ReasonNotLoaded' variable will contain information about why</returns>
-		static UBTMakefile LoadUBTMakefile( FileReference MakefilePath, out string ReasonNotLoaded )
+		static UBTMakefile LoadUBTMakefile(FileReference MakefilePath, out string ReasonNotLoaded)
 		{
 			// Check the directory timestamp on the project files directory.  If the user has generated project files more
 			// recently than the UBTMakefile, then we need to consider the file to be out of date
-			var UBTMakefileInfo = new FileInfo( MakefilePath.FullName );
-			if( !UBTMakefileInfo.Exists )
+			var UBTMakefileInfo = new FileInfo(MakefilePath.FullName);
+			if (!UBTMakefileInfo.Exists)
 			{
 				// UBTMakefile doesn't even exist, so we won't bother loading it
 				ReasonNotLoaded = "no existing makefile";
@@ -2415,15 +2415,15 @@ namespace UnrealBuildTool
 			// @todo ubtmake: Some platforms may not save any files into this folder.  We should delete + generate a "touch" file to force the directory timestamp to be updated (or just check the timestamp file itself.  We could put it ANYWHERE, actually)
 
 			// Rocket doesn't need to check engine projects for outdatedness
-			if( !UnrealBuildTool.RunningRocket() )
+			if (!UnrealBuildTool.RunningRocket())
 			{
-				if( ProjectFileGenerator.IntermediateProjectFilesPath.Exists() )
+				if (ProjectFileGenerator.IntermediateProjectFilesPath.Exists())
 				{
 					var EngineProjectFilesLastUpdateTime = new FileInfo(ProjectFileGenerator.ProjectTimestampFile).LastWriteTime;
-					if( UBTMakefileInfo.LastWriteTime.CompareTo( EngineProjectFilesLastUpdateTime ) < 0 )
+					if (UBTMakefileInfo.LastWriteTime.CompareTo(EngineProjectFilesLastUpdateTime) < 0)
 					{
 						// Engine project files are newer than UBTMakefile
-						Log.TraceVerbose("Existing makefile is older than generated engine project files, ignoring it" );
+						Log.TraceVerbose("Existing makefile is older than generated engine project files, ignoring it");
 						ReasonNotLoaded = "project files are newer";
 						return null;
 					}
@@ -2431,27 +2431,27 @@ namespace UnrealBuildTool
 			}
 
 			// Check the game project directory too
-			if( UnrealBuildTool.HasUProjectFile() )
+			if (UnrealBuildTool.HasUProjectFile())
 			{
 				string ProjectFilename = UnrealBuildTool.GetUProjectFile().FullName;
 				var ProjectFileInfo = new FileInfo(ProjectFilename);
-				if( !ProjectFileInfo.Exists || UBTMakefileInfo.LastWriteTime.CompareTo(ProjectFileInfo.LastWriteTime ) < 0 )
+				if (!ProjectFileInfo.Exists || UBTMakefileInfo.LastWriteTime.CompareTo(ProjectFileInfo.LastWriteTime) < 0)
 				{
 					// .uproject file is newer than UBTMakefile
-					Log.TraceVerbose("Makefile is older than .uproject file, ignoring it" );
+					Log.TraceVerbose("Makefile is older than .uproject file, ignoring it");
 					ReasonNotLoaded = ".uproject file is newer";
 					return null;
 				}
 
 				var MasterProjectRelativePath = UnrealBuildTool.GetUProjectPath();
-				var GameIntermediateProjectFilesPath = Path.Combine( MasterProjectRelativePath.FullName, "Intermediate", "ProjectFiles" );
-				if( Directory.Exists( GameIntermediateProjectFilesPath ) )
+				var GameIntermediateProjectFilesPath = Path.Combine(MasterProjectRelativePath.FullName, "Intermediate", "ProjectFiles");
+				if (Directory.Exists(GameIntermediateProjectFilesPath))
 				{
-					var GameProjectFilesLastUpdateTime = new DirectoryInfo( GameIntermediateProjectFilesPath ).LastWriteTime;
-					if( UBTMakefileInfo.LastWriteTime.CompareTo( GameProjectFilesLastUpdateTime ) < 0 )
+					var GameProjectFilesLastUpdateTime = new DirectoryInfo(GameIntermediateProjectFilesPath).LastWriteTime;
+					if (UBTMakefileInfo.LastWriteTime.CompareTo(GameProjectFilesLastUpdateTime) < 0)
 					{
 						// Game project files are newer than UBTMakefile
-						Log.TraceVerbose("Makefile is older than generated game project files, ignoring it" );
+						Log.TraceVerbose("Makefile is older than generated game project files, ignoring it");
 						ReasonNotLoaded = "game project files are newer";
 						return null;
 					}
@@ -2459,11 +2459,11 @@ namespace UnrealBuildTool
 			}
 
 			// Check to see if UnrealBuildTool.exe was compiled more recently than the UBTMakefile
-			var UnrealBuildToolTimestamp = new FileInfo( Assembly.GetExecutingAssembly().Location ).LastWriteTime;
-			if( UBTMakefileInfo.LastWriteTime.CompareTo( UnrealBuildToolTimestamp ) < 0 )
+			var UnrealBuildToolTimestamp = new FileInfo(Assembly.GetExecutingAssembly().Location).LastWriteTime;
+			if (UBTMakefileInfo.LastWriteTime.CompareTo(UnrealBuildToolTimestamp) < 0)
 			{
 				// UnrealBuildTool.exe was compiled more recently than the UBTMakefile
-				Log.TraceVerbose("Makefile is older than UnrealBuildTool.exe, ignoring it" );
+				Log.TraceVerbose("Makefile is older than UnrealBuildTool.exe, ignoring it");
 				ReasonNotLoaded = "UnrealBuildTool.exe is newer";
 				return null;
 			}
@@ -2475,15 +2475,15 @@ namespace UnrealBuildTool
 				var LoadUBTMakefileStartTime = DateTime.UtcNow;
 
 				using (FileStream Stream = new FileStream(UBTMakefileInfo.FullName, FileMode.Open, FileAccess.Read))
-				{	
+				{
 					BinaryFormatter Formatter = new BinaryFormatter();
 					LoadedUBTMakefile = Formatter.Deserialize(Stream) as UBTMakefile;
 				}
 
-				if( BuildConfiguration.bPrintPerformanceInfo )
-				{ 
+				if (BuildConfiguration.bPrintPerformanceInfo)
+				{
 					var LoadUBTMakefileTime = (DateTime.UtcNow - LoadUBTMakefileStartTime).TotalSeconds;
-					Log.TraceInformation( "LoadUBTMakefile took " + LoadUBTMakefileTime + "s" );
+					Log.TraceInformation("LoadUBTMakefile took " + LoadUBTMakefileTime + "s");
 				}
 			}
 			catch (Exception Ex)
@@ -2493,9 +2493,9 @@ namespace UnrealBuildTool
 				return null;
 			}
 
-			if( !LoadedUBTMakefile.IsValidMakefile() )
+			if (!LoadedUBTMakefile.IsValidMakefile())
 			{
-				Log.TraceWarning("Loaded makefile appears to have invalid contents, ignoring it ({0})", UBTMakefileInfo.FullName );
+				Log.TraceWarning("Loaded makefile appears to have invalid contents, ignoring it ({0})", UBTMakefileInfo.FullName);
 				ReasonNotLoaded = "existing makefile appears to be invalid";
 				return null;
 			}
@@ -2556,7 +2556,7 @@ namespace UnrealBuildTool
 				foreach (string Filename in HFilesNewerThanMakefile)
 				{
 					bool bContainsUHTData = CPPEnvironment.DoesFileContainUObjects(Filename);
-					bool bWasProcessed    = AllUHTHeaders.Contains(Filename);
+					bool bWasProcessed = AllUHTHeaders.Contains(Filename);
 					if (bContainsUHTData != bWasProcessed)
 					{
 						Log.TraceVerbose("{0} {1} contain UHT types and now {2} , ignoring it ({3})", Filename, bWasProcessed ? "used to" : "didn't", bWasProcessed ? "doesn't" : "does", UBTMakefileInfo.FullName);
@@ -2571,30 +2571,30 @@ namespace UnrealBuildTool
 			// changed, then we'll force a new Makefile to be created so that we have fresh unity build blobs.  We always
 			// want to make sure that source files in the working set are excluded from those unity blobs (for fastest possible
 			// iteration times.)
-			if( BuildConfiguration.bUseAdaptiveUnityBuild )
+			if (BuildConfiguration.bUseAdaptiveUnityBuild)
 			{
 				// Get all .cpp files in processed modules newer than the makefile itself
 				foreach (var Target in LoadedUBTMakefile.Targets)
 				{
-					foreach( var FlatModule in Target.FlatModuleCsData.Values )
+					foreach (var FlatModule in Target.FlatModuleCsData.Values)
 					{
-						if( FlatModule != null && FlatModule.ModuleSourceFolder != null )
+						if (FlatModule != null && FlatModule.ModuleSourceFolder != null)
 						{
 							// Has the directory been touched?  Don't bother checking folders that haven't been written to since the
 							// Makefile was created. 
 							//if( Directory.GetLastWriteTimeUtc( FlatModule.ModuleSourceFolder ) > UBTMakefileInfo.LastWriteTimeUtc )
 							{
-								foreach( var SourceFilePath in FlatModule.ModuleSourceFolder.EnumerateFileReferences( "*.cpp", SearchOption.AllDirectories ) )
+								foreach (var SourceFilePath in FlatModule.ModuleSourceFolder.EnumerateFileReferences("*.cpp", SearchOption.AllDirectories))
 								{
 									// Was the file written to since the Makefile was created?  We're not able to do a fully exhaustive
 									// check for working sets when in assembler mode, so we might get it wrong occasionally.  This is
 									// why we only want to check newly-modified or created files.  Otherwise, we could potentially
 									// invalidate the Makefile every single run.
-									if( File.GetLastWriteTimeUtc( SourceFilePath.FullName ) > UBTMakefileInfo.LastWriteTimeUtc )
+									if (File.GetLastWriteTimeUtc(SourceFilePath.FullName) > UBTMakefileInfo.LastWriteTimeUtc)
 									{
-										bool bShouldBePartOfWorkingSet = UnrealBuildTool.ShouldSourceFileBePartOfWorkingSet( SourceFilePath.FullName );
-										bool bIsAlreadyPartOfWorkingSet = LoadedUBTMakefile.SourceFileWorkingSet.Contains( FileItem.GetItemByFileReference( SourceFilePath ) );
-										if( bShouldBePartOfWorkingSet != bIsAlreadyPartOfWorkingSet )
+										bool bShouldBePartOfWorkingSet = UnrealBuildTool.ShouldSourceFileBePartOfWorkingSet(SourceFilePath.FullName);
+										bool bIsAlreadyPartOfWorkingSet = LoadedUBTMakefile.SourceFileWorkingSet.Contains(FileItem.GetItemByFileReference(SourceFilePath));
+										if (bShouldBePartOfWorkingSet != bIsAlreadyPartOfWorkingSet)
 										{
 											Log.TraceVerbose("{0} {1} part of source working set and now {2}; invalidating makefile ({3})", SourceFilePath, bIsAlreadyPartOfWorkingSet ? "was" : "was not", bShouldBePartOfWorkingSet ? "is" : "is not", UBTMakefileInfo.FullName);
 											ReasonNotLoaded = string.Format("working set of source files changed");
@@ -2613,46 +2613,46 @@ namespace UnrealBuildTool
 		}
 
 
-        /// <summary>
-        /// Gets the file path for a UBTMakefile
-        /// </summary>
-        /// <param name="Targets">List of targets.  Order is not important</param>
-        /// <returns>UBTMakefile path</returns>
-        public static FileReference GetUBTMakefilePath( List<TargetDescriptor> TargetDescs )
-        {
-            FileReference UBTMakefilePath;
+		/// <summary>
+		/// Gets the file path for a UBTMakefile
+		/// </summary>
+		/// <param name="Targets">List of targets.  Order is not important</param>
+		/// <returns>UBTMakefile path</returns>
+		public static FileReference GetUBTMakefilePath(List<TargetDescriptor> TargetDescs)
+		{
+			FileReference UBTMakefilePath;
 
-            if( TargetDescs.Count == 1 )
-            {
+			if (TargetDescs.Count == 1)
+			{
 				var TargetDesc = TargetDescs[0];
 
 				bool bIsHotReload = (UEBuildConfiguration.bHotReloadFromIDE || (TargetDesc.OnlyModules != null && TargetDesc.OnlyModules.Count > 0));
 				string UBTMakefileName = bIsHotReload ? "HotReloadMakefile.ubt" : "Makefile.ubt";
 
 				UBTMakefilePath = FileReference.Combine(GetUBTMakefileDirectoryPathForSingleTarget(TargetDesc), UBTMakefileName);
-            }
-            else
-            {
-                // For Makefiles that contain multiple targets, we'll make up a file name that contains all of the targets, their
-                // configurations and platforms, and save it into the base intermediate folder
-                var TargetCollectionName = MakeTargetCollectionName( TargetDescs );
-
-                DirectoryReference ProjectIntermediatePath;
-                if (UnrealBuildTool.HasUProjectFile())
-                {
-                    ProjectIntermediatePath = DirectoryReference.Combine(UnrealBuildTool.GetUProjectPath(), BuildConfiguration.BaseIntermediateFolder);
-                }
+			}
+			else
+			{
+				// For Makefiles that contain multiple targets, we'll make up a file name that contains all of the targets, their
+				// configurations and platforms, and save it into the base intermediate folder
+				var TargetCollectionName = MakeTargetCollectionName(TargetDescs);
+
+				DirectoryReference ProjectIntermediatePath;
+				if (UnrealBuildTool.HasUProjectFile())
+				{
+					ProjectIntermediatePath = DirectoryReference.Combine(UnrealBuildTool.GetUProjectPath(), BuildConfiguration.BaseIntermediateFolder);
+				}
 				else
 				{
 					ProjectIntermediatePath = DirectoryReference.Combine(UnrealBuildTool.EngineDirectory, BuildConfiguration.BaseIntermediateFolder);
 				}
 
-                // @todo ubtmake: The TargetCollectionName string could be really long if there is more than one target!  Hash it?
-                UBTMakefilePath = FileReference.Combine(ProjectIntermediatePath, TargetCollectionName + ".ubt" );
-            }
+				// @todo ubtmake: The TargetCollectionName string could be really long if there is more than one target!  Hash it?
+				UBTMakefilePath = FileReference.Combine(ProjectIntermediatePath, TargetCollectionName + ".ubt");
+			}
 
-            return UBTMakefilePath;
-        }
+			return UBTMakefilePath;
+		}
 
 		/// <summary>
 		/// Gets the file path for a UBTMakefile for single target.
@@ -2686,47 +2686,47 @@ namespace UnrealBuildTool
 			return DirectoryReference.Combine(PlatformIntermediatePath, TargetName, Target.Configuration.ToString());
 		}
 
-        /// <summary>
-        /// Makes up a name for a set of targets that we can use for file or directory names
-        /// </summary>
-        /// <param name="TargetDescs">List of targets.  Order is not important</param>
-        /// <returns>The name to use</returns>
-        private static string MakeTargetCollectionName( List<TargetDescriptor> TargetDescs )
-        {
-            if( TargetDescs.Count == 0 )
-            {
-                throw new BuildException( "Expecting at least one Target to be passed to MakeTargetCollectionName" );
-            }
-
-            var SortedTargets = new List<TargetDescriptor>();
-            SortedTargets.AddRange( TargetDescs );
-            SortedTargets.Sort( ( x, y ) => { return x.TargetName.CompareTo( y.TargetName ); } );
-
-            // Figure out what to call our action graph based on everything we're building
-            var TargetCollectionName = new StringBuilder();
-            foreach( var Target in SortedTargets )
-            {
-                if( TargetCollectionName.Length > 0 )
-                {
-                    TargetCollectionName.Append( "_" );
-                }
+		/// <summary>
+		/// Makes up a name for a set of targets that we can use for file or directory names
+		/// </summary>
+		/// <param name="TargetDescs">List of targets.  Order is not important</param>
+		/// <returns>The name to use</returns>
+		private static string MakeTargetCollectionName(List<TargetDescriptor> TargetDescs)
+		{
+			if (TargetDescs.Count == 0)
+			{
+				throw new BuildException("Expecting at least one Target to be passed to MakeTargetCollectionName");
+			}
+
+			var SortedTargets = new List<TargetDescriptor>();
+			SortedTargets.AddRange(TargetDescs);
+			SortedTargets.Sort((x, y) => { return x.TargetName.CompareTo(y.TargetName); });
+
+			// Figure out what to call our action graph based on everything we're building
+			var TargetCollectionName = new StringBuilder();
+			foreach (var Target in SortedTargets)
+			{
+				if (TargetCollectionName.Length > 0)
+				{
+					TargetCollectionName.Append("_");
+				}
 
 				// @todo ubtmake: If this is a compile triggered from the editor it will have passed along the game's target name, not the editor target name.
 				// At this point in Unreal Build Tool, we can't possibly know what the actual editor target name is, but we can take a guess.
 				// Even if we get it wrong, this won't have any side effects aside from not being able to share the exact same cached UBT Makefile
 				// between hot reloads invoked from the editor and hot reloads invoked from within the IDE.
 				string TargetName = Target.TargetName;
-				if( !TargetName.EndsWith( "Editor", StringComparison.InvariantCultureIgnoreCase ) && Target.bIsEditorRecompile )
+				if (!TargetName.EndsWith("Editor", StringComparison.InvariantCultureIgnoreCase) && Target.bIsEditorRecompile)
 				{
 					TargetName += "Editor";
 				}
 
 				// @todo ubtmake: Should we also have the platform Architecture in this string?
-				TargetCollectionName.Append( TargetName + "-" + Target.Platform.ToString() + "-" + Target.Configuration.ToString() );
-            }
+				TargetCollectionName.Append(TargetName + "-" + Target.Platform.ToString() + "-" + Target.Configuration.ToString());
+			}
 
-            return TargetCollectionName.ToString();
-        }
+			return TargetCollectionName.ToString();
+		}
 
 		/// <summary>
 		/// Sets up UBT when running from UAT
@@ -2746,7 +2746,7 @@ namespace UnrealBuildTool
 		/// Patch action history for hot reload when running in assembler mode.  In assembler mode, the suffix on the output file will be
 		/// the same for every invocation on that makefile, but we need a new suffix each time.
 		/// </summary>
-		protected static void PatchActionHistoryForHotReloadAssembling( List<OnlyModule> OnlyModules )
+		protected static void PatchActionHistoryForHotReloadAssembling(List<OnlyModule> OnlyModules)
 		{
 			// Gather all of the response files for link actions.  We're going to need to patch 'em up after we figure out new
 			// names for all of the output files and import libraries
@@ -2756,18 +2756,18 @@ namespace UnrealBuildTool
 			var OriginalFileNameAndNewFileNameList_NoExtensions = new List<Tuple<string, string>>();
 
 			// Finally, we'll keep track of any file items that we had to create counterparts for change file names, so we can fix those up too
-			var AffectedOriginalFileItemAndNewFileItemMap = new Dictionary<FileItem,FileItem>();
+			var AffectedOriginalFileItemAndNewFileItemMap = new Dictionary<FileItem, FileItem>();
 
-			foreach( var Action in ActionGraph.AllActions )
+			foreach (var Action in ActionGraph.AllActions)
 			{
-				if( Action.ActionType == ActionType.Link )
+				if (Action.ActionType == ActionType.Link)
 				{
 					// Assume that the first produced item (with no extension) is our output file name
-					var OriginalFileNameWithoutExtension = Utils.GetFilenameWithoutAnyExtensions( Action.ProducedItems[0].AbsolutePath );
+					var OriginalFileNameWithoutExtension = Utils.GetFilenameWithoutAnyExtensions(Action.ProducedItems[0].AbsolutePath);
 
 					// Remove the numbered suffix from the end of the file
-					var IndexOfLastHyphen = OriginalFileNameWithoutExtension.LastIndexOf( '-' );
-					var OriginalNumberSuffix = OriginalFileNameWithoutExtension.Substring( IndexOfLastHyphen );
+					var IndexOfLastHyphen = OriginalFileNameWithoutExtension.LastIndexOf('-');
+					var OriginalNumberSuffix = OriginalFileNameWithoutExtension.Substring(IndexOfLastHyphen);
 					int NumberSuffix;
 					bool bHasNumberSuffix = int.TryParse(OriginalNumberSuffix, out NumberSuffix);
 
@@ -2795,23 +2795,23 @@ namespace UnrealBuildTool
 
 					// Figure out which suffix to use
 					string UniqueSuffix = null;
-					{ 
-						int FirstHyphenIndex = OriginalFileNameWithoutNumberSuffix.IndexOf( '-' );
-						if( FirstHyphenIndex == -1 )
+					{
+						int FirstHyphenIndex = OriginalFileNameWithoutNumberSuffix.IndexOf('-');
+						if (FirstHyphenIndex == -1)
 						{
-							throw new BuildException( "Expected produced item file name '{0}' to start with a prefix (such as 'UE4Editor-') when hot reloading", OriginalFileNameWithoutExtension );
+							throw new BuildException("Expected produced item file name '{0}' to start with a prefix (such as 'UE4Editor-') when hot reloading", OriginalFileNameWithoutExtension);
 						}
-						string OutputFileNamePrefix = OriginalFileNameWithoutNumberSuffix.Substring( 0, FirstHyphenIndex + 1 );
+						string OutputFileNamePrefix = OriginalFileNameWithoutNumberSuffix.Substring(0, FirstHyphenIndex + 1);
 
 						// Get the module name from the file name
-						string ModuleName = OriginalFileNameWithoutNumberSuffix.Substring( OutputFileNamePrefix.Length );
+						string ModuleName = OriginalFileNameWithoutNumberSuffix.Substring(OutputFileNamePrefix.Length);
 
 						// Add a new random suffix
-						if( OnlyModules != null )
-						{ 
-							foreach( var OnlyModule in OnlyModules )
+						if (OnlyModules != null)
+						{
+							foreach (var OnlyModule in OnlyModules)
 							{
-								if( OnlyModule.OnlyModuleName.Equals( ModuleName, StringComparison.InvariantCultureIgnoreCase ) )
+								if (OnlyModule.OnlyModuleName.Equals(ModuleName, StringComparison.InvariantCultureIgnoreCase))
 								{
 									// OK, we found the module!
 									UniqueSuffix = "-" + OnlyModule.OnlyModuleSuffix;
@@ -2819,7 +2819,7 @@ namespace UnrealBuildTool
 								}
 							}
 						}
-						if( UniqueSuffix == null )
+						if (UniqueSuffix == null)
 						{
 							UniqueSuffix = "-" + (new Random((int)(DateTime.Now.Ticks % Int32.MaxValue)).Next(10000)).ToString();
 						}
@@ -2829,47 +2829,47 @@ namespace UnrealBuildTool
 					// Find the response file in the command line.  We'll need to make a copy of it with our new file name.
 					{
 						string ResponseFileExtension = ".response";
-						var ResponseExtensionIndex = Action.CommandArguments.IndexOf( ResponseFileExtension, StringComparison.InvariantCultureIgnoreCase );
-						if( ResponseExtensionIndex != -1 )
-						{ 
-							var ResponseFilePathIndex = Action.CommandArguments.LastIndexOf( "@\"", ResponseExtensionIndex );
-							if( ResponseFilePathIndex == -1 )
+						var ResponseExtensionIndex = Action.CommandArguments.IndexOf(ResponseFileExtension, StringComparison.InvariantCultureIgnoreCase);
+						if (ResponseExtensionIndex != -1)
+						{
+							var ResponseFilePathIndex = Action.CommandArguments.LastIndexOf("@\"", ResponseExtensionIndex);
+							if (ResponseFilePathIndex == -1)
 							{
-								throw new BuildException( "Couldn't find response file path in action's command arguments when hot reloading" );
+								throw new BuildException("Couldn't find response file path in action's command arguments when hot reloading");
 							}
 
-							var OriginalResponseFilePathWithoutExtension = Action.CommandArguments.Substring( ResponseFilePathIndex + 2, ( ResponseExtensionIndex - ResponseFilePathIndex ) - 2 );
+							var OriginalResponseFilePathWithoutExtension = Action.CommandArguments.Substring(ResponseFilePathIndex + 2, (ResponseExtensionIndex - ResponseFilePathIndex) - 2);
 							var OriginalResponseFilePath = OriginalResponseFilePathWithoutExtension + ResponseFileExtension;
 
-							var NewResponseFilePath = OriginalResponseFilePath.Replace( OriginalFileNameWithoutExtension, NewFileNameWithoutExtension );
+							var NewResponseFilePath = OriginalResponseFilePath.Replace(OriginalFileNameWithoutExtension, NewFileNameWithoutExtension);
 
 							// Copy the old response file to the new path
-							File.Copy( OriginalResponseFilePath, NewResponseFilePath, overwrite: true);
+							File.Copy(OriginalResponseFilePath, NewResponseFilePath, overwrite: true);
 
 							// Keep track of the new response file name.  We'll have to do some edits afterwards.
-							ResponseFilePaths.Add( NewResponseFilePath );
+							ResponseFilePaths.Add(NewResponseFilePath);
 						}
 					}
 
 
 					// Go ahead and replace all occurrences of our file name in the command-line (ignoring extensions)
-					Action.CommandArguments = Action.CommandArguments.Replace( OriginalFileNameWithoutExtension, NewFileNameWithoutExtension );
+					Action.CommandArguments = Action.CommandArguments.Replace(OriginalFileNameWithoutExtension, NewFileNameWithoutExtension);
+
 
-					
 					// Update this action's list of produced items too
 					{
-						for( int ItemIndex = 0; ItemIndex < Action.ProducedItems.Count; ++ItemIndex )
+						for (int ItemIndex = 0; ItemIndex < Action.ProducedItems.Count; ++ItemIndex)
 						{
-							var OriginalProducedItem = Action.ProducedItems[ ItemIndex ];
+							var OriginalProducedItem = Action.ProducedItems[ItemIndex];
 
 							var OriginalProducedItemFilePath = OriginalProducedItem.AbsolutePath;
-							var NewProducedItemFilePath = OriginalProducedItemFilePath.Replace( OriginalFileNameWithoutExtension, NewFileNameWithoutExtension );
-						
-							if( OriginalProducedItemFilePath != NewProducedItemFilePath )
+							var NewProducedItemFilePath = OriginalProducedItemFilePath.Replace(OriginalFileNameWithoutExtension, NewFileNameWithoutExtension);
+
+							if (OriginalProducedItemFilePath != NewProducedItemFilePath)
 							{
 								// OK, the produced item's file name changed so we'll update it to point to our new file
-								var NewProducedItem = FileItem.GetItemByPath( NewProducedItemFilePath );
-								Action.ProducedItems[ ItemIndex ] = NewProducedItem;
+								var NewProducedItem = FileItem.GetItemByPath(NewProducedItemFilePath);
+								Action.ProducedItems[ItemIndex] = NewProducedItem;
 
 								// Copy the other important settings from the original file item
 								NewProducedItem.bNeedsHotReloadNumbersDLLCleanUp = OriginalProducedItem.bNeedsHotReloadNumbersDLLCleanUp;
@@ -2877,52 +2877,52 @@ namespace UnrealBuildTool
 								NewProducedItem.bIsRemoteFile = OriginalProducedItem.bIsRemoteFile;
 
 								// Keep track of it so we can fix up dependencies in a second pass afterwards
-								AffectedOriginalFileItemAndNewFileItemMap.Add( OriginalProducedItem, NewProducedItem );
+								AffectedOriginalFileItemAndNewFileItemMap.Add(OriginalProducedItem, NewProducedItem);
 							}
 						}
 					}
 
 					// The status description of the item has the file name, so we'll update it too
-					Action.StatusDescription = Action.StatusDescription.Replace( OriginalFileNameWithoutExtension, NewFileNameWithoutExtension );
+					Action.StatusDescription = Action.StatusDescription.Replace(OriginalFileNameWithoutExtension, NewFileNameWithoutExtension);
+
 
-		
 					// Keep track of the file names, so we can fix up response files afterwards.
-					OriginalFileNameAndNewFileNameList_NoExtensions.Add( Tuple.Create( OriginalFileNameWithoutExtension, NewFileNameWithoutExtension ) );
+					OriginalFileNameAndNewFileNameList_NoExtensions.Add(Tuple.Create(OriginalFileNameWithoutExtension, NewFileNameWithoutExtension));
 				}
 			}
 
-			
+
 			// Do another pass and update any actions that depended on the original file names that we changed
-			foreach( var Action in ActionGraph.AllActions )
+			foreach (var Action in ActionGraph.AllActions)
 			{
-				for( int ItemIndex = 0; ItemIndex < Action.PrerequisiteItems.Count; ++ItemIndex )
+				for (int ItemIndex = 0; ItemIndex < Action.PrerequisiteItems.Count; ++ItemIndex)
 				{
-					var OriginalFileItem = Action.PrerequisiteItems[ ItemIndex ];
+					var OriginalFileItem = Action.PrerequisiteItems[ItemIndex];
 
 					FileItem NewFileItem;
-					if( AffectedOriginalFileItemAndNewFileItemMap.TryGetValue( OriginalFileItem, out NewFileItem ) )
+					if (AffectedOriginalFileItemAndNewFileItemMap.TryGetValue(OriginalFileItem, out NewFileItem))
 					{
 						// OK, looks like we need to replace this file item because we've renamed the file
-						Action.PrerequisiteItems[ ItemIndex ] = NewFileItem;
+						Action.PrerequisiteItems[ItemIndex] = NewFileItem;
 					}
 				}
 			}
 
 
-			if( OriginalFileNameAndNewFileNameList_NoExtensions.Count > 0 )
-			{	
-				foreach( var ResponseFilePath in ResponseFilePaths )
+			if (OriginalFileNameAndNewFileNameList_NoExtensions.Count > 0)
+			{
+				foreach (var ResponseFilePath in ResponseFilePaths)
 				{
 					// Load the file up
-					StringBuilder FileContents = new StringBuilder( Utils.ReadAllText( ResponseFilePath ) );
+					StringBuilder FileContents = new StringBuilder(Utils.ReadAllText(ResponseFilePath));
 
 					// Replace all of the old file names with new ones
-					foreach( var FileNameTuple in OriginalFileNameAndNewFileNameList_NoExtensions )
+					foreach (var FileNameTuple in OriginalFileNameAndNewFileNameList_NoExtensions)
 					{
 						var OriginalFileNameWithoutExtension = FileNameTuple.Item1;
 						var NewFileNameWithoutExtension = FileNameTuple.Item2;
 
-						FileContents.Replace( OriginalFileNameWithoutExtension, NewFileNameWithoutExtension );
+						FileContents.Replace(OriginalFileNameWithoutExtension, NewFileNameWithoutExtension);
 					}
 
 					// Overwrite the original file
@@ -2939,13 +2939,13 @@ namespace UnrealBuildTool
 		/// </summary>
 		/// <param name="SourceFileAbsolutePath">Path to the source file on disk</param>
 		/// <returns>True if file is part of the working set</returns>
-		public static bool ShouldSourceFileBePartOfWorkingSet( string SourceFileAbsolutePath )
+		public static bool ShouldSourceFileBePartOfWorkingSet(string SourceFileAbsolutePath)
 		{
 			// @todo: This logic needs work.  Currently its designed for users working with Perforce, or a similar
 			// source control system that keeps source files read-only when not actively checked out for editing.
 			bool bShouldBePartOfWorkingSourceFileSet =
-				!System.IO.File.GetAttributes( SourceFileAbsolutePath ).HasFlag( FileAttributes.ReadOnly );
+				!System.IO.File.GetAttributes(SourceFileAbsolutePath).HasFlag(FileAttributes.ReadOnly);
 			return bShouldBePartOfWorkingSourceFileSet;
 		}
-    }
+	}
 }
