Auto-format UBT for consistent spacing. All default Visual Studio settings, but tabs instead of spaces.

diff --git a/Engine/Source/Programs/UnrealBuildTool/System/Utils.cs b/Engine/Source/Programs/UnrealBuildTool/System/Utils.cs
--- a/Engine/Source/Programs/UnrealBuildTool/System/Utils.cs
+++ b/Engine/Source/Programs/UnrealBuildTool/System/Utils.cs
@@ -38,7 +38,7 @@ namespace UnrealBuildTool
 		public static bool TryRead(string FileName, out BuildVersion Version)
 		{
 			JsonObject Object;
-			if(!JsonObject.TryRead(FileName, out Object))
+			if (!JsonObject.TryRead(FileName, out Object))
 			{
 				Version = null;
 				return false;
@@ -55,7 +55,7 @@ namespace UnrealBuildTool
 		public static bool TryParse(JsonObject Object, out BuildVersion Version)
 		{
 			BuildVersion NewVersion = new BuildVersion();
-			if(!Object.TryGetIntegerField("MajorVersion", out NewVersion.MajorVersion) || !Object.TryGetIntegerField("MinorVersion", out NewVersion.MinorVersion) || !Object.TryGetIntegerField("PatchVersion", out NewVersion.PatchVersion))
+			if (!Object.TryGetIntegerField("MajorVersion", out NewVersion.MajorVersion) || !Object.TryGetIntegerField("MinorVersion", out NewVersion.MinorVersion) || !Object.TryGetIntegerField("PatchVersion", out NewVersion.PatchVersion))
 			{
 				Version = null;
 				return false;
@@ -108,10 +108,10 @@ namespace UnrealBuildTool
 		/// <summary>
 		/// Whether we are currently running on Mono platform.  We cache this statically because it is a bit slow to check.
 		/// </summary>
-		public static readonly bool IsRunningOnMono = Type.GetType( "Mono.Runtime" ) != null;
+		public static readonly bool IsRunningOnMono = Type.GetType("Mono.Runtime") != null;
 
 		/** Searches for a flag in a set of command-line arguments. */
-		public static bool ParseCommandLineFlag(string[] Arguments, string FlagName, out int ArgumentIndex )
+		public static bool ParseCommandLineFlag(string[] Arguments, string FlagName, out int ArgumentIndex)
 		{
 			// Find an argument with the given name.
 			for (ArgumentIndex = 0; ArgumentIndex < Arguments.Length; ArgumentIndex++)
@@ -126,7 +126,7 @@ namespace UnrealBuildTool
 		}
 
 		/** Regular expression to match $(ENV) and/ or %ENV% environment variables. */
-		static Regex EnvironmentVariableRegex = new Regex( @"\$\((.*?)\)|\%(.*?)\%", RegexOptions.None );
+		static Regex EnvironmentVariableRegex = new Regex(@"\$\((.*?)\)|\%(.*?)\%", RegexOptions.None);
 
 		/**
 		 * Resolves $(ENV) and/ or %ENV% to the value of the environment variable in the passed in string.
@@ -134,35 +134,35 @@ namespace UnrealBuildTool
 		 * @param	InString	String to resolve environment variable in.
 		 * @return	String with environment variable expanded/ resolved.
 		 */
-		public static string ResolveEnvironmentVariable( string InString )
+		public static string ResolveEnvironmentVariable(string InString)
 		{
 			string Result = InString;
-		
+
 			// Try to find $(ENV) substring.
-			Match M = EnvironmentVariableRegex.Match( InString );
+			Match M = EnvironmentVariableRegex.Match(InString);
 
 			// Iterate over all matches, resolving the match to an environment variable.
-			while( M.Success )
+			while (M.Success)
 			{
 				// Convoluted way of stripping first and last character and '(' in the case of $(ENV) to get to ENV
 				string EnvironmentVariable = M.ToString();
-				if ( EnvironmentVariable.StartsWith("$") && EnvironmentVariable.EndsWith(")") )
+				if (EnvironmentVariable.StartsWith("$") && EnvironmentVariable.EndsWith(")"))
 				{
 					EnvironmentVariable = EnvironmentVariable.Substring(1, EnvironmentVariable.Length - 2).Replace("(", "");
 				}
 
-				if ( EnvironmentVariable.StartsWith("%") && EnvironmentVariable.EndsWith("%") )
+				if (EnvironmentVariable.StartsWith("%") && EnvironmentVariable.EndsWith("%"))
 				{
 					EnvironmentVariable = EnvironmentVariable.Substring(1, EnvironmentVariable.Length - 2);
 				}
-			
+
 				// Resolve environment variable.				
-				Result = Result.Replace( M.ToString(), Environment.GetEnvironmentVariable( EnvironmentVariable ) );
+				Result = Result.Replace(M.ToString(), Environment.GetEnvironmentVariable(EnvironmentVariable));
 
 				// Move on to next match. Multiple environment variables are handled correctly by regexp.
 				M = M.NextMatch();
 			}
-				
+
 			return Result;
 		}
 
@@ -176,11 +176,11 @@ namespace UnrealBuildTool
 		public static string ExpandVariables(string InputString, Dictionary<string, string> AdditionalVariables = null)
 		{
 			string Result = InputString;
-			for(int Idx = Result.IndexOf("$("); Idx != -1; Idx = Result.IndexOf("$(", Idx))
+			for (int Idx = Result.IndexOf("$("); Idx != -1; Idx = Result.IndexOf("$(", Idx))
 			{
 				// Find the end of the variable name
 				int EndIdx = Result.IndexOf(')', Idx + 2);
-				if(EndIdx == -1)
+				if (EndIdx == -1)
 				{
 					break;
 				}
@@ -190,10 +190,10 @@ namespace UnrealBuildTool
 
 				// Find the value for it, either from the dictionary or the environment block
 				string Value;
-				if(AdditionalVariables == null || !AdditionalVariables.TryGetValue(Name, out Value))
+				if (AdditionalVariables == null || !AdditionalVariables.TryGetValue(Name, out Value))
 				{
 					Value = Environment.GetEnvironmentVariable(Name);
-					if(Value == null)
+					if (Value == null)
 					{
 						Idx = EndIdx + 1;
 						continue;
@@ -205,7 +205,7 @@ namespace UnrealBuildTool
 			}
 			return Result;
 		}
-		
+
 		/**
 		 * This is a faster replacement of File.ReadAllText. Code snippet based on code 
 		 * and analysis by Sam Allen
@@ -215,9 +215,9 @@ namespace UnrealBuildTool
 		 * @param	SourceFile		Source file to fully read and convert to string
 		 * @return	Textual representation of file.
 		 */
-		public static string ReadAllText( string SourceFile )
+		public static string ReadAllText(string SourceFile)
 		{
-			using( StreamReader Reader = new StreamReader( SourceFile, System.Text.Encoding.UTF8 ) )
+			using (StreamReader Reader = new StreamReader(SourceFile, System.Text.Encoding.UTF8))
 			{
 				return Reader.ReadToEnd();
 			}
@@ -326,7 +326,7 @@ namespace UnrealBuildTool
 			int ExitCode = -1;
 
 			// release all process resources
-			using(LocalProcess)
+			using (LocalProcess)
 			{
 				LocalProcess.StartInfo.UseShellExecute = false;
 				LocalProcess.StartInfo.RedirectStandardOutput = true;
@@ -341,7 +341,7 @@ namespace UnrealBuildTool
 					LocalProcess.WaitForExit();
 					ExitCode = LocalProcess.ExitCode;
 				}
-				catch(Exception ex)
+				catch (Exception ex)
 				{
 					throw new BuildException(ex, "Failed to start local process for action (\"{0}\"): {1} {2}", ex.Message, LocalProcess.StartInfo.FileName, LocalProcess.StartInfo.Arguments);
 				}
@@ -357,8 +357,8 @@ namespace UnrealBuildTool
 		{
 			Process LocalProcess = new Process();
 			LocalProcess.StartInfo = StartInfo;
-			LocalProcess.OutputDataReceived += (Sender, Line) => { if(Line != null && Line.Data != null) Log.TraceInformation(Line.Data); };
-			LocalProcess.ErrorDataReceived += (Sender, Line) => { if(Line != null && Line.Data != null) Log.TraceError(Line.Data); };
+			LocalProcess.OutputDataReceived += (Sender, Line) => { if (Line != null && Line.Data != null) Log.TraceInformation(Line.Data); };
+			LocalProcess.ErrorDataReceived += (Sender, Line) => { if (Line != null && Line.Data != null) Log.TraceError(Line.Data); };
 			return RunLocalProcess(LocalProcess);
 		}
 
@@ -391,7 +391,7 @@ namespace UnrealBuildTool
 		/// </summary>
 		/// <param name="SupportedPlatforms">List of supported platforms</param>
 		/// <returns>List of unsupported platforms in string format</returns>
-		public static List<string> MakeListOfUnsupportedPlatforms( List<UnrealTargetPlatform> SupportedPlatforms )
+		public static List<string> MakeListOfUnsupportedPlatforms(List<UnrealTargetPlatform> SupportedPlatforms)
 		{
 			// Make a list of all platform name strings that we're *not* currently compiling, to speed
 			// up file path comparisons later on
@@ -426,28 +426,28 @@ namespace UnrealBuildTool
 					}
 				}
 
-				foreach( UnrealTargetPlatform CurPlatform in Enum.GetValues( typeof( UnrealTargetPlatform ) ) )
+				foreach (UnrealTargetPlatform CurPlatform in Enum.GetValues(typeof(UnrealTargetPlatform)))
 				{
 					if (CurPlatform != UnrealTargetPlatform.Unknown)
 					{
-                        bool ShouldConsider = true;
-
-                        // If we have a platform and a group with the same name, don't add the platform
-                        // to the other list if the same-named group is supported.  This is a lot of
-                        // lines because we need to do the comparisons as strings.
-                        string CurPlatformString = CurPlatform.ToString();
-                        foreach (UnrealPlatformGroup Group in Enum.GetValues(typeof(UnrealPlatformGroup)))
-                        {
-                            if (Group.ToString().Equals(CurPlatformString))
-                            {
-                                ShouldConsider = false;
-                                break;
-                            }
-                        }
+						bool ShouldConsider = true;
+
+						// If we have a platform and a group with the same name, don't add the platform
+						// to the other list if the same-named group is supported.  This is a lot of
+						// lines because we need to do the comparisons as strings.
+						string CurPlatformString = CurPlatform.ToString();
+						foreach (UnrealPlatformGroup Group in Enum.GetValues(typeof(UnrealPlatformGroup)))
+						{
+							if (Group.ToString().Equals(CurPlatformString))
+							{
+								ShouldConsider = false;
+								break;
+							}
+						}
 
 						// Don't add our current platform to the list of platform sub-directory names that
 						// we'll skip source files for
-						if ( ShouldConsider && !SupportedPlatforms.Contains( CurPlatform ) )
+						if (ShouldConsider && !SupportedPlatforms.Contains(CurPlatform))
 						{
 							OtherPlatformNameStrings.Add(CurPlatform.ToString());
 						}
@@ -464,7 +464,7 @@ namespace UnrealBuildTool
 		/// </summary>
 		/// <param name="FilePath">File path with potentially inconsistent slashes</param>
 		/// <returns>File path with consistent separators</returns>
-		public static string CleanDirectorySeparators( string FilePath, char UseDirectorySeparatorChar = '\0' )
+		public static string CleanDirectorySeparators(string FilePath, char UseDirectorySeparatorChar = '\0')
 		{
 			StringBuilder CleanPath = null;
 			if (UseDirectorySeparatorChar == '\0')
@@ -476,23 +476,23 @@ namespace UnrealBuildTool
 			bool bCanCheckDoubleSeparators = false;
 			for (int Index = 0; Index < FilePath.Length; ++Index)
 			{
-				char C = FilePath[Index];				
+				char C = FilePath[Index];
 				if (C == '/' || C == '\\')
 				{
-					if( C != UseDirectorySeparatorChar )
-					{ 
+					if (C != UseDirectorySeparatorChar)
+					{
 						C = UseDirectorySeparatorChar;
-						if( CleanPath == null )
-						{ 
-							CleanPath = new StringBuilder( FilePath.Substring( 0, Index ), FilePath.Length );
+						if (CleanPath == null)
+						{
+							CleanPath = new StringBuilder(FilePath.Substring(0, Index), FilePath.Length);
 						}
 					}
 
 					if (bCanCheckDoubleSeparators && C == PrevC)
 					{
-						if( CleanPath == null )
-						{ 
-							CleanPath = new StringBuilder( FilePath.Substring( 0, Index ), FilePath.Length );
+						if (CleanPath == null)
+						{
+							CleanPath = new StringBuilder(FilePath.Substring(0, Index), FilePath.Length);
 						}
 						continue;
 					}
@@ -503,8 +503,8 @@ namespace UnrealBuildTool
 					bCanCheckDoubleSeparators = true;
 				}
 
-				if( CleanPath != null )
-				{ 
+				if (CleanPath != null)
+				{
 					CleanPath.Append(C);
 				}
 				PrevC = C;
@@ -512,7 +512,7 @@ namespace UnrealBuildTool
 			return CleanPath != null ? CleanPath.ToString() : FilePath;
 		}
 
-	
+
 		/// <summary>
 		/// Given a file path and a directory, returns a file path that is relative to the specified directory
 		/// </summary>
@@ -520,51 +520,51 @@ namespace UnrealBuildTool
 		/// <param name="RelativeToDirectory">The directory that the source file path should be converted to be relative to.  If this path is not rooted, it will be assumed to be relative to the current working directory.</param>
 		/// <param name="AlwaysTreatSourceAsDirectory">True if we should treat the source path like a directory even if it doesn't end with a path separator</param>
 		/// <returns>Converted relative path</returns>
-		public static string MakePathRelativeTo( string SourcePath, string RelativeToDirectory, bool AlwaysTreatSourceAsDirectory = false )
+		public static string MakePathRelativeTo(string SourcePath, string RelativeToDirectory, bool AlwaysTreatSourceAsDirectory = false)
 		{
-			if( String.IsNullOrEmpty( RelativeToDirectory ) )
+			if (String.IsNullOrEmpty(RelativeToDirectory))
 			{
 				// Assume CWD
 				RelativeToDirectory = ".";
 			}
 
 			var AbsolutePath = SourcePath;
-			if( !Path.IsPathRooted( AbsolutePath ) )
+			if (!Path.IsPathRooted(AbsolutePath))
 			{
-				AbsolutePath = Path.GetFullPath( SourcePath );
+				AbsolutePath = Path.GetFullPath(SourcePath);
 			}
-			var SourcePathEndsWithDirectorySeparator = AbsolutePath.EndsWith( Path.DirectorySeparatorChar.ToString() ) || AbsolutePath.EndsWith( Path.AltDirectorySeparatorChar.ToString() );
-			if( AlwaysTreatSourceAsDirectory && !SourcePathEndsWithDirectorySeparator )
+			var SourcePathEndsWithDirectorySeparator = AbsolutePath.EndsWith(Path.DirectorySeparatorChar.ToString()) || AbsolutePath.EndsWith(Path.AltDirectorySeparatorChar.ToString());
+			if (AlwaysTreatSourceAsDirectory && !SourcePathEndsWithDirectorySeparator)
 			{
 				AbsolutePath += Path.DirectorySeparatorChar;
 			}
 
-			var AbsolutePathUri = new Uri( AbsolutePath );
+			var AbsolutePathUri = new Uri(AbsolutePath);
 
 			var AbsoluteRelativeDirectory = RelativeToDirectory;
-			if( !Path.IsPathRooted( AbsoluteRelativeDirectory ) )
+			if (!Path.IsPathRooted(AbsoluteRelativeDirectory))
 			{
-				AbsoluteRelativeDirectory = Path.GetFullPath( AbsoluteRelativeDirectory );
+				AbsoluteRelativeDirectory = Path.GetFullPath(AbsoluteRelativeDirectory);
 			}
 
 			// Make sure the directory has a trailing directory separator so that the relative directory that
 			// MakeRelativeUri creates doesn't include our directory -- only the directories beneath it!
-			if( !AbsoluteRelativeDirectory.EndsWith( Path.DirectorySeparatorChar.ToString() ) && !AbsoluteRelativeDirectory.EndsWith( Path.AltDirectorySeparatorChar.ToString() ) )
+			if (!AbsoluteRelativeDirectory.EndsWith(Path.DirectorySeparatorChar.ToString()) && !AbsoluteRelativeDirectory.EndsWith(Path.AltDirectorySeparatorChar.ToString()))
 			{
 				AbsoluteRelativeDirectory += Path.DirectorySeparatorChar;
 			}
 
 			// Convert to URI form which is where we can make the relative conversion happen
-			var AbsoluteRelativeDirectoryUri = new Uri( AbsoluteRelativeDirectory );
+			var AbsoluteRelativeDirectoryUri = new Uri(AbsoluteRelativeDirectory);
 
 			// Ask the URI system to convert to a nicely formed relative path, then convert it back to a regular path string
-			var UriRelativePath = AbsoluteRelativeDirectoryUri.MakeRelativeUri( AbsolutePathUri );
-			var RelativePath = Uri.UnescapeDataString( UriRelativePath.ToString() ).Replace( '/', Path.DirectorySeparatorChar );
+			var UriRelativePath = AbsoluteRelativeDirectoryUri.MakeRelativeUri(AbsolutePathUri);
+			var RelativePath = Uri.UnescapeDataString(UriRelativePath.ToString()).Replace('/', Path.DirectorySeparatorChar);
 
 			// If we added a directory separator character earlier on, remove it now
-			if( !SourcePathEndsWithDirectorySeparator && AlwaysTreatSourceAsDirectory && RelativePath.EndsWith( Path.DirectorySeparatorChar.ToString() ) )
+			if (!SourcePathEndsWithDirectorySeparator && AlwaysTreatSourceAsDirectory && RelativePath.EndsWith(Path.DirectorySeparatorChar.ToString()))
 			{
-				RelativePath = RelativePath.Substring( 0, RelativePath.Length - 1 );
+				RelativePath = RelativePath.Substring(0, RelativePath.Length - 1);
 			}
 
 			// Uri.MakeRelativeUri is broken in Mono 2.x and sometimes returns broken path
@@ -593,58 +593,58 @@ namespace UnrealBuildTool
 		/// <param name="Numerator">Progress numerator</param>
 		/// <param name="Denominator">Progress denominator</param>
 		/// <param name="NumCharsToBackspaceOver">Number of characters to backspace before writing the text.  This value will be updated with the length of the new progress string.  The first time progress is displayed, you should pass 0 for this value.</param>
-		public static void DisplayProgress( int Numerator, int Denominator, ref int NumCharsToBackspaceOver )
+		public static void DisplayProgress(int Numerator, int Denominator, ref int NumCharsToBackspaceOver)
 		{
 			// Backspace over previous progress value
-			while( NumCharsToBackspaceOver-- > 0 )
+			while (NumCharsToBackspaceOver-- > 0)
 			{
-				Console.Write( "\b" );
+				Console.Write("\b");
 			}
 
 			// Display updated progress string and keep track of how long it was
-			float ProgressValue = Denominator > 0 ? ( (float)Numerator / (float)Denominator ) : 1.0f;
-			var ProgressString = String.Format( "{0}%", Math.Round( ProgressValue * 100.0f ) );
+			float ProgressValue = Denominator > 0 ? ((float)Numerator / (float)Denominator) : 1.0f;
+			var ProgressString = String.Format("{0}%", Math.Round(ProgressValue * 100.0f));
 			NumCharsToBackspaceOver = ProgressString.Length;
-			Console.Write( ProgressString );
+			Console.Write(ProgressString);
 		}
 
 
 		/*
 		 * Read and write classes with xml specifiers
 		 */
-		static private void UnknownAttributeDelegate( object sender, XmlAttributeEventArgs e )
+		static private void UnknownAttributeDelegate(object sender, XmlAttributeEventArgs e)
 		{
 		}
 
-		static private void UnknownNodeDelegate( object sender, XmlNodeEventArgs e )
+		static private void UnknownNodeDelegate(object sender, XmlNodeEventArgs e)
 		{
 		}
 
-		static public T ReadClass<T>( string FileName ) where T : new()
+		static public T ReadClass<T>(string FileName) where T : new()
 		{
 			T Instance = new T();
 			StreamReader XmlStream = null;
 			try
 			{
 				// Get the XML data stream to read from
-				XmlStream = new StreamReader( FileName );
+				XmlStream = new StreamReader(FileName);
 
 				// Creates an instance of the XmlSerializer class so we can read the settings object
-				XmlSerializer Serialiser = new XmlSerializer( typeof( T ) );
+				XmlSerializer Serialiser = new XmlSerializer(typeof(T));
 				// Add our callbacks for unknown nodes and attributes
-				Serialiser.UnknownNode += new XmlNodeEventHandler( UnknownNodeDelegate );
-				Serialiser.UnknownAttribute += new XmlAttributeEventHandler( UnknownAttributeDelegate );
+				Serialiser.UnknownNode += new XmlNodeEventHandler(UnknownNodeDelegate);
+				Serialiser.UnknownAttribute += new XmlAttributeEventHandler(UnknownAttributeDelegate);
 
 				// Create an object graph from the XML data
-				Instance = ( T )Serialiser.Deserialize( XmlStream );
+				Instance = (T)Serialiser.Deserialize(XmlStream);
 			}
-			catch( Exception E )
+			catch (Exception E)
 			{
-				Log.TraceInformation( E.Message );
+				Log.TraceInformation(E.Message);
 			}
 			finally
 			{
-				if( XmlStream != null )
+				if (XmlStream != null)
 				{
 					// Done with the file so close it
 					XmlStream.Close();
@@ -654,14 +654,14 @@ namespace UnrealBuildTool
 			return Instance;
 		}
 
-		static public bool WriteClass<T>( T Data, string FileName, string DefaultNameSpace )
+		static public bool WriteClass<T>(T Data, string FileName, string DefaultNameSpace)
 		{
 			bool bSuccess = true;
 			StreamWriter XmlStream = null;
 			try
 			{
-				FileInfo Info = new FileInfo( FileName );
-				if( Info.Exists )
+				FileInfo Info = new FileInfo(FileName);
+				if (Info.Exists)
 				{
 					Info.IsReadOnly = false;
 				}
@@ -670,32 +670,32 @@ namespace UnrealBuildTool
 				Directory.CreateDirectory(Path.GetDirectoryName(FileName));
 
 				XmlSerializerNamespaces EmptyNameSpace = new XmlSerializerNamespaces();
-				EmptyNameSpace.Add( "", DefaultNameSpace );
+				EmptyNameSpace.Add("", DefaultNameSpace);
 
-				XmlStream = new StreamWriter( FileName, false, Encoding.Unicode );
-				XmlSerializer Serialiser = new XmlSerializer( typeof( T ) );
+				XmlStream = new StreamWriter(FileName, false, Encoding.Unicode);
+				XmlSerializer Serialiser = new XmlSerializer(typeof(T));
 
 				// Add our callbacks for unknown nodes and attributes
-				Serialiser.UnknownNode += new XmlNodeEventHandler( UnknownNodeDelegate );
-				Serialiser.UnknownAttribute += new XmlAttributeEventHandler( UnknownAttributeDelegate );
+				Serialiser.UnknownNode += new XmlNodeEventHandler(UnknownNodeDelegate);
+				Serialiser.UnknownAttribute += new XmlAttributeEventHandler(UnknownAttributeDelegate);
 
-				Serialiser.Serialize( XmlStream, Data, EmptyNameSpace );
+				Serialiser.Serialize(XmlStream, Data, EmptyNameSpace);
 			}
-			catch( Exception E )
+			catch (Exception E)
 			{
-				Log.TraceInformation( E.Message );
+				Log.TraceInformation(E.Message);
 				bSuccess = false;
 			}
 			finally
 			{
-				if( XmlStream != null )
+				if (XmlStream != null)
 				{
 					// Done with the file so close it
 					XmlStream.Close();
 				}
 			}
 
-			return ( bSuccess );
+			return (bSuccess);
 		}
 
 		/// <summary>
@@ -706,13 +706,13 @@ namespace UnrealBuildTool
 		public static bool IsValidProcess(Process p)
 		{
 			// null objects are always invalid
-			if(p == null)
+			if (p == null)
 				return false;
 			// due to multithreading on Windows, lock the object
-			lock(p)
+			lock (p)
 			{
 				// Mono has a specific requirement if testing for an alive process
-				if(IsRunningOnMono)
+				if (IsRunningOnMono)
 					return p.Handle != IntPtr.Zero; // native handle to the process
 				// on Windows, simply test the process ID to be non-zero. 
 				// note that this can fail and have a race condition in threads, but the framework throws an exception when this occurs.
@@ -769,14 +769,14 @@ namespace UnrealBuildTool
 		/// <param name="FilePath">The path to the file</param>
 		/// <param name="Directory">The directory to check to see if the file is located under (or any of this directory's subfolders)</param>
 		/// <returns></returns>
-		public static bool IsFileUnderDirectory( string FilePath, string Directory )
+		public static bool IsFileUnderDirectory(string FilePath, string Directory)
 		{
-			var DirectoryPathPlusSeparator = Path.GetFullPath( Directory );
-			if( !DirectoryPathPlusSeparator.EndsWith( Path.DirectorySeparatorChar.ToString() ) )
+			var DirectoryPathPlusSeparator = Path.GetFullPath(Directory);
+			if (!DirectoryPathPlusSeparator.EndsWith(Path.DirectorySeparatorChar.ToString()))
 			{
 				DirectoryPathPlusSeparator += Path.DirectorySeparatorChar;
 			}
-			return Path.GetFullPath( FilePath ).StartsWith( DirectoryPathPlusSeparator, StringComparison.InvariantCultureIgnoreCase );
+			return Path.GetFullPath(FilePath).StartsWith(DirectoryPathPlusSeparator, StringComparison.InvariantCultureIgnoreCase);
 		}
 
 
@@ -786,9 +786,9 @@ namespace UnrealBuildTool
 		/// <param name="FilePath">The full path</param>
 		/// <param name="BaseDirectory">The base directory, which must be the first part of the path</param>
 		/// <returns>The part of the path after the base directory</returns>
-		public static string StripBaseDirectory( string FilePath, string BaseDirectory )
+		public static string StripBaseDirectory(string FilePath, string BaseDirectory)
 		{
-			return CleanDirectorySeparators( FilePath ).Substring( CleanDirectorySeparators( BaseDirectory ).Length + 1 );
+			return CleanDirectorySeparators(FilePath).Substring(CleanDirectorySeparators(BaseDirectory).Length + 1);
 		}
 
 
@@ -799,10 +799,10 @@ namespace UnrealBuildTool
 		/// <param name="BaseDirectory"></param>
 		/// <param name="NewBaseDirectory"></param>
 		/// <returns></returns>
-		public static string MakeRerootedFilePath( string FilePath, string BaseDirectory, string NewBaseDirectory )
+		public static string MakeRerootedFilePath(string FilePath, string BaseDirectory, string NewBaseDirectory)
 		{
-			var RelativeFile = StripBaseDirectory( FilePath, BaseDirectory );
-			var DestFile = Path.Combine( NewBaseDirectory, RelativeFile );
+			var RelativeFile = StripBaseDirectory(FilePath, BaseDirectory);
+			var DestFile = Path.Combine(NewBaseDirectory, RelativeFile);
 			return DestFile;
 		}
 
@@ -823,10 +823,10 @@ namespace UnrealBuildTool
 				bHadBackSlashes = true;
 			}
 
-			string ParentDir       = "/..";
-			int    ParentDirLength = ParentDir.Length;
+			string ParentDir = "/..";
+			int ParentDirLength = ParentDir.Length;
 
-			for (;;)
+			for (; ; )
 			{
 				// An empty path is finished
 				if (string.IsNullOrEmpty(LocalString))
@@ -842,7 +842,7 @@ namespace UnrealBuildTool
 					break;
 
 				int PreviousSeparatorIndex = Index;
-				for (;;)
+				for (; ; )
 				{
 					// Find the previous slash
 					PreviousSeparatorIndex = Math.Max(0, LocalString.LastIndexOf("/", PreviousSeparatorIndex - 1));
@@ -978,133 +978,133 @@ namespace UnrealBuildTool
 		}
 	}
 
-    /// <summary>
-    /// UAT/UBT Custom log system.
-    /// 
-    /// This lets you use any TraceListeners you want, but you should only call the static 
-    /// methods below, not call Trace.XXX directly, as the static methods
-    /// This allows the system to enforce the formatting and filtering conventions we desire.
-    ///
-    /// For posterity, we cannot use the Trace or TraceSource class directly because of our special log requirements:
-    ///   1. We possibly capture the method name of the logging event. This cannot be done as a macro, so must be done at the top level so we know how many layers of the stack to peel off to get the real function.
-    ///   2. We have a verbose filter we would like to apply to all logs without having to have each listener filter individually, which would require our string formatting code to run every time.
-    ///   3. We possibly want to ensure severity prefixes are logged, but Trace.WriteXXX does not allow any severity info to be passed down.
-    /// </summary>
-    static public class Log
-    {
-        /// <summary>
-        /// Guard our initialization. Mainly used by top level exception handlers to ensure its safe to call a logging function.
-        /// In general user code should not concern itself checking for this.
-        /// </summary>
-        private static bool bIsInitialized = false;
-        /// <summary>
-        /// When true, verbose loggin is enabled.
-        /// </summary>
-        private static LogEventType LogLevel = LogEventType.Log;
-        /// <summary>
-        /// When true, warnings and errors will have a WARNING: or ERROR: prexifx, respectively.
-        /// </summary>
-        private static bool bLogSeverity = false;
-        /// <summary>
-        /// When true, logs will have the calling mehod prepended to the output as MethodName:
-        /// </summary>
-        private static bool bLogSources = false;
-        /// <summary>
-        /// When true, will detect warnings and errors and set the console output color to yellow and red.
-        /// </summary>
-        private static bool bColorConsoleOutput = false;
-        /// <summary>
-        /// When configured, this tracks time since initialization to prepend a timestamp to each log.
-        /// </summary>
-        private static Stopwatch Timer;
-
-        /// <summary>
-        /// Expose the log level. This is a hack for ProcessResult.LogOutput, which wants to bypass our normal formatting scheme.
-        /// </summary>
-        public static bool bIsVerbose { get { return LogLevel >= LogEventType.Verbose; } }
+	/// <summary>
+	/// UAT/UBT Custom log system.
+	/// 
+	/// This lets you use any TraceListeners you want, but you should only call the static 
+	/// methods below, not call Trace.XXX directly, as the static methods
+	/// This allows the system to enforce the formatting and filtering conventions we desire.
+	///
+	/// For posterity, we cannot use the Trace or TraceSource class directly because of our special log requirements:
+	///   1. We possibly capture the method name of the logging event. This cannot be done as a macro, so must be done at the top level so we know how many layers of the stack to peel off to get the real function.
+	///   2. We have a verbose filter we would like to apply to all logs without having to have each listener filter individually, which would require our string formatting code to run every time.
+	///   3. We possibly want to ensure severity prefixes are logged, but Trace.WriteXXX does not allow any severity info to be passed down.
+	/// </summary>
+	static public class Log
+	{
+		/// <summary>
+		/// Guard our initialization. Mainly used by top level exception handlers to ensure its safe to call a logging function.
+		/// In general user code should not concern itself checking for this.
+		/// </summary>
+		private static bool bIsInitialized = false;
+		/// <summary>
+		/// When true, verbose loggin is enabled.
+		/// </summary>
+		private static LogEventType LogLevel = LogEventType.Log;
+		/// <summary>
+		/// When true, warnings and errors will have a WARNING: or ERROR: prexifx, respectively.
+		/// </summary>
+		private static bool bLogSeverity = false;
+		/// <summary>
+		/// When true, logs will have the calling mehod prepended to the output as MethodName:
+		/// </summary>
+		private static bool bLogSources = false;
+		/// <summary>
+		/// When true, will detect warnings and errors and set the console output color to yellow and red.
+		/// </summary>
+		private static bool bColorConsoleOutput = false;
+		/// <summary>
+		/// When configured, this tracks time since initialization to prepend a timestamp to each log.
+		/// </summary>
+		private static Stopwatch Timer;
+
+		/// <summary>
+		/// Expose the log level. This is a hack for ProcessResult.LogOutput, which wants to bypass our normal formatting scheme.
+		/// </summary>
+		public static bool bIsVerbose { get { return LogLevel >= LogEventType.Verbose; } }
 
 		/// <summary>
 		/// A collection of strings that have been already written once
 		/// </summary>
 		private static List<string> WriteOnceSet = new List<string>();
 
-        /// <summary>
-        /// Allows code to check if the log system is ready yet.
-        /// End users should NOT need to use this. It pretty much exists
-        /// to work around startup issues since this is a global singleton.
-        /// </summary>
-        /// <returns></returns>
-        public static bool IsInitialized() 
-        { 
-            return bIsInitialized; 
-        }
-
-        /// <summary>
-        /// Allows us to change verbosity after initializing. This can happen since we initialize logging early, 
-        /// but then read the config and command line later, which could change this value.
-        /// </summary>
-        /// <param name="bLogVerbose">Whether to log verbose logs.</param>
-        public static void SetLoggingLevel(LogEventType InLogLevel)
-        {
-            Log.LogLevel = InLogLevel;
-        }
-
-        /// <summary>
-        /// This class allows InitLogging to be called more than once to work around chicken and eggs issues with logging and parsing command lines (see UBT startup code).
-        /// </summary>
-        /// <param name="bLogTimestamps">If true, the timestamp from Log init time will be prepended to all logs.</param>
-        /// <param name="bLogVerbose">If true, any Verbose log method method will not be ignored.</param>
-        /// <param name="bLogSeverity">If true, warnings and errors will have a WARNING: and ERROR: prefix to them. </param>
-        /// <param name="bLogSources">If true, logs will have the originating method name prepended to them.</param>
-        /// <param name="TraceListeners">Collection of trace listeners to attach to the Trace.Listeners, in addition to the Default listener. The existing listeners (except the Default listener) are cleared first.</param>
-        public static void InitLogging(bool bLogTimestamps, LogEventType InLogLevel, bool bLogSeverity, bool bLogSources, bool bColorConsoleOutput, IEnumerable<TraceListener> TraceListeners)
-        {
-            bIsInitialized = true;
-            Timer = (bLogTimestamps && Timer == null) ? Stopwatch.StartNew() : null;
-            Log.LogLevel = InLogLevel;
-            Log.bLogSeverity = bLogSeverity;
-            Log.bLogSources = bLogSources;
-            Log.bColorConsoleOutput = bColorConsoleOutput;
-
-            // ensure that if InitLogging is called more than once we don't stack listeners.
-            // but always leave the default listener around.
-            for (int ListenerNdx = 0; ListenerNdx < Trace.Listeners.Count; )
-			{
-                if (Trace.Listeners[ListenerNdx].GetType() != typeof(DefaultTraceListener))
+		/// <summary>
+		/// Allows code to check if the log system is ready yet.
+		/// End users should NOT need to use this. It pretty much exists
+		/// to work around startup issues since this is a global singleton.
+		/// </summary>
+		/// <returns></returns>
+		public static bool IsInitialized()
+		{
+			return bIsInitialized;
+		}
+
+		/// <summary>
+		/// Allows us to change verbosity after initializing. This can happen since we initialize logging early, 
+		/// but then read the config and command line later, which could change this value.
+		/// </summary>
+		/// <param name="bLogVerbose">Whether to log verbose logs.</param>
+		public static void SetLoggingLevel(LogEventType InLogLevel)
+		{
+			Log.LogLevel = InLogLevel;
+		}
+
+		/// <summary>
+		/// This class allows InitLogging to be called more than once to work around chicken and eggs issues with logging and parsing command lines (see UBT startup code).
+		/// </summary>
+		/// <param name="bLogTimestamps">If true, the timestamp from Log init time will be prepended to all logs.</param>
+		/// <param name="bLogVerbose">If true, any Verbose log method method will not be ignored.</param>
+		/// <param name="bLogSeverity">If true, warnings and errors will have a WARNING: and ERROR: prefix to them. </param>
+		/// <param name="bLogSources">If true, logs will have the originating method name prepended to them.</param>
+		/// <param name="TraceListeners">Collection of trace listeners to attach to the Trace.Listeners, in addition to the Default listener. The existing listeners (except the Default listener) are cleared first.</param>
+		public static void InitLogging(bool bLogTimestamps, LogEventType InLogLevel, bool bLogSeverity, bool bLogSources, bool bColorConsoleOutput, IEnumerable<TraceListener> TraceListeners)
+		{
+			bIsInitialized = true;
+			Timer = (bLogTimestamps && Timer == null) ? Stopwatch.StartNew() : null;
+			Log.LogLevel = InLogLevel;
+			Log.bLogSeverity = bLogSeverity;
+			Log.bLogSources = bLogSources;
+			Log.bColorConsoleOutput = bColorConsoleOutput;
+
+			// ensure that if InitLogging is called more than once we don't stack listeners.
+			// but always leave the default listener around.
+			for (int ListenerNdx = 0; ListenerNdx < Trace.Listeners.Count; )
+			{
+				if (Trace.Listeners[ListenerNdx].GetType() != typeof(DefaultTraceListener))
 				{
-                    Trace.Listeners.RemoveAt(ListenerNdx);
+					Trace.Listeners.RemoveAt(ListenerNdx);
 				}
 				else
 				{
-                    ++ListenerNdx;
+					++ListenerNdx;
 				}
 			}
-            // don't add any null listeners
-            Trace.Listeners.AddRange(TraceListeners.Where(l => l != null).ToArray());
-            Trace.AutoFlush = true;
-        }
-
-        /// <summary>
-        /// Gets the name of the Method N levels deep in the stack frame. Used to trap what method actually made the logging call.
-        /// Only used when bLogSources is true.
-        /// </summary>
-        /// <param name="StackFramesToSkip"></param>
-        /// <returns>ClassName.MethodName</returns>
-        [MethodImplAttribute(MethodImplOptions.NoInlining)]
-        private static string GetSource(int StackFramesToSkip)
-        {
-            StackFrame Frame = new StackFrame(2 + StackFramesToSkip);
-            System.Reflection.MethodBase Method = Frame.GetMethod();
-            return String.Format("{0}.{1}", Method.DeclaringType.Name, Method.Name);
-        }
-
-        /// <summary>
-        /// Converts a LogEventType into a log prefix. Only used when bLogSeverity is true.
-        /// </summary>
-        /// <param name="EventType"></param>
-        /// <returns></returns>
-        private static string GetSeverityPrefix(LogEventType Severity)
-        {
+			// don't add any null listeners
+			Trace.Listeners.AddRange(TraceListeners.Where(l => l != null).ToArray());
+			Trace.AutoFlush = true;
+		}
+
+		/// <summary>
+		/// Gets the name of the Method N levels deep in the stack frame. Used to trap what method actually made the logging call.
+		/// Only used when bLogSources is true.
+		/// </summary>
+		/// <param name="StackFramesToSkip"></param>
+		/// <returns>ClassName.MethodName</returns>
+		[MethodImplAttribute(MethodImplOptions.NoInlining)]
+		private static string GetSource(int StackFramesToSkip)
+		{
+			StackFrame Frame = new StackFrame(2 + StackFramesToSkip);
+			System.Reflection.MethodBase Method = Frame.GetMethod();
+			return String.Format("{0}.{1}", Method.DeclaringType.Name, Method.Name);
+		}
+
+		/// <summary>
+		/// Converts a LogEventType into a log prefix. Only used when bLogSeverity is true.
+		/// </summary>
+		/// <param name="EventType"></param>
+		/// <returns></returns>
+		private static string GetSeverityPrefix(LogEventType Severity)
+		{
 			switch (Severity)
 			{
 				case LogEventType.Fatal:
@@ -1122,7 +1122,7 @@ namespace UnrealBuildTool
 				default:
 					return " ";
 			}
-        }
+		}
 
 		/// <summary>
 		/// Converts a LogEventType into a message code
@@ -1134,42 +1134,42 @@ namespace UnrealBuildTool
 			return (int)Severity;
 		}
 
-        /// <summary>
-        /// Formats message for logging. Enforces the configured options.
-        /// </summary>
-        /// <param name="StackFramesToSkip">Number of frames to skip to get to the originator of the log request.</param>
-        /// <param name="CustomSource">Custom source string to use. Use the Class.Method string if null. Only used if bLogSources = true.</param>
-        /// <param name="Verbosity">Message verbosity level</param>
-        /// <param name="Format">Message text format string</param>
-        /// <param name="Args">Message text parameters</param>
-        /// <returns>Formatted message</returns>
-        [MethodImplAttribute(MethodImplOptions.NoInlining)]
+		/// <summary>
+		/// Formats message for logging. Enforces the configured options.
+		/// </summary>
+		/// <param name="StackFramesToSkip">Number of frames to skip to get to the originator of the log request.</param>
+		/// <param name="CustomSource">Custom source string to use. Use the Class.Method string if null. Only used if bLogSources = true.</param>
+		/// <param name="Verbosity">Message verbosity level</param>
+		/// <param name="Format">Message text format string</param>
+		/// <param name="Args">Message text parameters</param>
+		/// <returns>Formatted message</returns>
+		[MethodImplAttribute(MethodImplOptions.NoInlining)]
 		private static string FormatMessage(int StackFramesToSkip, string CustomSource, LogEventType Verbosity, string Format, params object[] Args)
-        {
-            return string.Format("{0}{1}{2}{3}",
-                    Timer != null ? String.Format("[{0:hh\\:mm\\:ss\\.fff}] ", Timer.Elapsed) : "",
-                    bLogSources ? string.Format("{0}: ", string.IsNullOrEmpty(CustomSource) ? GetSource(StackFramesToSkip) : CustomSource) : "",
-                    bLogSeverity ? GetSeverityPrefix(Verbosity) : "",
-                    // If there are no extra args, don't try to format the string, in case it has any format control characters in it (our LOCTEXT strings tend to).
-                    Args.Length > 0 ? string.Format(Format, Args) : Format);
-        }
-
-        /// <summary>
-        /// Writes a formatted message to the console. All other functions should boil down to calling this method.
-        /// </summary>
-        /// <param name="StackFramesToSkip">Number of frames to skip to get to the originator of the log request.</param>
-        /// <param name="CustomSource">Custom source string to use. Use the default if null.</param>
+		{
+			return string.Format("{0}{1}{2}{3}",
+					Timer != null ? String.Format("[{0:hh\\:mm\\:ss\\.fff}] ", Timer.Elapsed) : "",
+					bLogSources ? string.Format("{0}: ", string.IsNullOrEmpty(CustomSource) ? GetSource(StackFramesToSkip) : CustomSource) : "",
+					bLogSeverity ? GetSeverityPrefix(Verbosity) : "",
+				// If there are no extra args, don't try to format the string, in case it has any format control characters in it (our LOCTEXT strings tend to).
+					Args.Length > 0 ? string.Format(Format, Args) : Format);
+		}
+
+		/// <summary>
+		/// Writes a formatted message to the console. All other functions should boil down to calling this method.
+		/// </summary>
+		/// <param name="StackFramesToSkip">Number of frames to skip to get to the originator of the log request.</param>
+		/// <param name="CustomSource">Custom source string to use. Use the default if null.</param>
 		/// <param name="bWriteOnce">If true, this message will be written only once</param>
-        /// <param name="Verbosity">Message verbosity level. We only meaningfully use values up to Verbose</param>
-        /// <param name="Format">Message format string.</param>
-        /// <param name="Args">Optional arguments</param>
-        [MethodImplAttribute(MethodImplOptions.NoInlining)]
+		/// <param name="Verbosity">Message verbosity level. We only meaningfully use values up to Verbose</param>
+		/// <param name="Format">Message format string.</param>
+		/// <param name="Args">Optional arguments</param>
+		[MethodImplAttribute(MethodImplOptions.NoInlining)]
 		private static void WriteLinePrivate(int StackFramesToSkip, string CustomSource, bool bWriteOnce, LogEventType Verbosity, string Format, params object[] Args)
-        {
-            if (!bIsInitialized)
-            {
-                throw new BuildException("Tried to using Logging system before it was ready");
-            }
+		{
+			if (!bIsInitialized)
+			{
+				throw new BuildException("Tried to using Logging system before it was ready");
+			}
 
 			// if we want this message only written one time, check if it was already written out
 			if (bWriteOnce)
@@ -1184,159 +1184,159 @@ namespace UnrealBuildTool
 			}
 
 			if (Verbosity <= LogLevel)
-            {
-                // Do console color highlighting here.
-                ConsoleColor DefaultColor = ConsoleColor.Gray;
+			{
+				// Do console color highlighting here.
+				ConsoleColor DefaultColor = ConsoleColor.Gray;
 				bool bIsWarning = false;
 				bool bIsError = false;
-                // don't try to touch the console unless we are told to color the output.
-                if (bColorConsoleOutput)
-                {
-                    DefaultColor = Console.ForegroundColor;
+				// don't try to touch the console unless we are told to color the output.
+				if (bColorConsoleOutput)
+				{
+					DefaultColor = Console.ForegroundColor;
 					bIsWarning = Verbosity == LogEventType.Warning;
 					bIsError = Verbosity <= LogEventType.Error;
 					// @todo mono - mono doesn't seem to initialize the ForegroundColor properly, so we can't restore it properly.
 					// Avoid touching the console color unless we really need to.
 					if (bIsWarning || bIsError)
 					{
-                    	Console.ForegroundColor = bIsWarning ? ConsoleColor.Yellow : ConsoleColor.Red;
+						Console.ForegroundColor = bIsWarning ? ConsoleColor.Yellow : ConsoleColor.Red;
 					}
-                }
-                try
-                {
+				}
+				try
+				{
 					// @todo mono: mono has some kind of bug where calling mono recursively by spawning
 					// a new process causes Trace.WriteLine to stop functioning (it returns, but does nothing for some reason).
 					// work around this by simulating Trace.WriteLine on mono.
 					// We use UAT to spawn UBT instances recursively a lot, so this bug can effectively
 					// make all build output disappear outside of the top level UAT process.
-//					#if MONO
-					    lock (((System.Collections.ICollection)Trace.Listeners).SyncRoot)
+					//					#if MONO
+					lock (((System.Collections.ICollection)Trace.Listeners).SyncRoot)
+					{
+						foreach (TraceListener l in Trace.Listeners)
 						{
-						    foreach (TraceListener l in Trace.Listeners) 
-						    {
-								if (Verbosity != LogEventType.Log || (l as ConsoleTraceListener) == null || LogLevel >= LogEventType.Verbose)
-								{
-									l.WriteLine(FormatMessage(StackFramesToSkip + 1, CustomSource, Verbosity, Format, Args));
-								}
-							    l.Flush();
-						    }
+							if (Verbosity != LogEventType.Log || (l as ConsoleTraceListener) == null || LogLevel >= LogEventType.Verbose)
+							{
+								l.WriteLine(FormatMessage(StackFramesToSkip + 1, CustomSource, Verbosity, Format, Args));
+							}
+							l.Flush();
 						}
-//					#else
-                    	// Call Trace directly here. Trace ensures that our logging is threadsafe using the GlobalLock.
-//                    	Trace.WriteLine(FormatMessage(StackFramesToSkip + 1, CustomSource, Verbosity, Format, Args));
-//					#endif
-                }
-                finally
-                {
-                    // make sure we always put the console color back.
-                    if (bColorConsoleOutput && (bIsWarning || bIsError))
-                    {
-                        Console.ForegroundColor = DefaultColor;
-                    }
-                }
-            }
-        }
-
-        /// <summary>
-        /// Similar to Trace.WriteLineIf
-        /// </summary>
-        /// <param name="Condition"></param>
-        /// <param name="Verbosity"></param>
-        /// <param name="Format"></param>
-        /// <param name="Args"></param>
-        [MethodImplAttribute(MethodImplOptions.NoInlining)]
-        public static void WriteLineIf(bool Condition, LogEventType Verbosity, string Format, params object[] Args)
-        {
-            if (Condition)
-            {
+					}
+					//					#else
+					// Call Trace directly here. Trace ensures that our logging is threadsafe using the GlobalLock.
+					//                    	Trace.WriteLine(FormatMessage(StackFramesToSkip + 1, CustomSource, Verbosity, Format, Args));
+					//					#endif
+				}
+				finally
+				{
+					// make sure we always put the console color back.
+					if (bColorConsoleOutput && (bIsWarning || bIsError))
+					{
+						Console.ForegroundColor = DefaultColor;
+					}
+				}
+			}
+		}
+
+		/// <summary>
+		/// Similar to Trace.WriteLineIf
+		/// </summary>
+		/// <param name="Condition"></param>
+		/// <param name="Verbosity"></param>
+		/// <param name="Format"></param>
+		/// <param name="Args"></param>
+		[MethodImplAttribute(MethodImplOptions.NoInlining)]
+		public static void WriteLineIf(bool Condition, LogEventType Verbosity, string Format, params object[] Args)
+		{
+			if (Condition)
+			{
 				WriteLinePrivate(1, null, false, Verbosity, Format, Args);
-            }
-        }
-
-        /// <summary>
-        /// Mostly an internal function, but expose StackFramesToSkip to allow UAT to use existing wrapper functions and still get proper formatting.
-        /// </summary>
-        /// <param name="StackFramesToSkip"></param>
-        /// <param name="CustomSource">Custom source string to use. Use the default if null.</param>
-        /// <param name="Verbosity"></param>
-        /// <param name="Format"></param>
-        /// <param name="Args"></param>
-        [MethodImplAttribute(MethodImplOptions.NoInlining)]
+			}
+		}
+
+		/// <summary>
+		/// Mostly an internal function, but expose StackFramesToSkip to allow UAT to use existing wrapper functions and still get proper formatting.
+		/// </summary>
+		/// <param name="StackFramesToSkip"></param>
+		/// <param name="CustomSource">Custom source string to use. Use the default if null.</param>
+		/// <param name="Verbosity"></param>
+		/// <param name="Format"></param>
+		/// <param name="Args"></param>
+		[MethodImplAttribute(MethodImplOptions.NoInlining)]
 		public static void WriteLine(int StackFramesToSkip, LogEventType Verbosity, string Format, params object[] Args)
-        {
-            WriteLinePrivate(StackFramesToSkip + 1, null, false, Verbosity, Format, Args);
-        }
-
-        /// <summary>
-        /// Mostly an internal function, but expose StackFramesToSkip and a custom Source to alow UAT to use existing wrapper functions and still get proper formatting.
-        /// </summary>
-        /// <param name="StackFramesToSkip"></param>
-        /// <param name="CustomSource">Custom source string to use. Use the default if null.</param>
-        /// <param name="Verbosity"></param>
-        /// <param name="Format"></param>
-        /// <param name="Args"></param>
-        [MethodImplAttribute(MethodImplOptions.NoInlining)]
+		{
+			WriteLinePrivate(StackFramesToSkip + 1, null, false, Verbosity, Format, Args);
+		}
+
+		/// <summary>
+		/// Mostly an internal function, but expose StackFramesToSkip and a custom Source to alow UAT to use existing wrapper functions and still get proper formatting.
+		/// </summary>
+		/// <param name="StackFramesToSkip"></param>
+		/// <param name="CustomSource">Custom source string to use. Use the default if null.</param>
+		/// <param name="Verbosity"></param>
+		/// <param name="Format"></param>
+		/// <param name="Args"></param>
+		[MethodImplAttribute(MethodImplOptions.NoInlining)]
 		public static void WriteLine(int StackFramesToSkip, string CustomSource, LogEventType Verbosity, string Format, params object[] Args)
-        {
-            WriteLinePrivate(StackFramesToSkip + 1, CustomSource, false, Verbosity, Format, Args);
-        }
-
-        /// <summary>
-        /// Similar to Trace.WriteLin
-        /// </summary>
-        /// <param name="Verbosity"></param>
-        /// <param name="Format"></param>
-        /// <param name="Args"></param>
-        [MethodImplAttribute(MethodImplOptions.NoInlining)]
+		{
+			WriteLinePrivate(StackFramesToSkip + 1, CustomSource, false, Verbosity, Format, Args);
+		}
+
+		/// <summary>
+		/// Similar to Trace.WriteLin
+		/// </summary>
+		/// <param name="Verbosity"></param>
+		/// <param name="Format"></param>
+		/// <param name="Args"></param>
+		[MethodImplAttribute(MethodImplOptions.NoInlining)]
 		public static void WriteLine(LogEventType Verbosity, string Format, params object[] Args)
-        {
-            WriteLinePrivate(1, null, false, Verbosity, Format, Args);
-        }
-
-        /// <summary>
-        /// Writes an error message to the console.
-        /// </summary>
-        /// <param name="Format">Message format string</param>
-        /// <param name="Args">Optional arguments</param>
-        [MethodImplAttribute(MethodImplOptions.NoInlining)]
-        public static void TraceError(string Format, params object[] Args)
-        {
+		{
+			WriteLinePrivate(1, null, false, Verbosity, Format, Args);
+		}
+
+		/// <summary>
+		/// Writes an error message to the console.
+		/// </summary>
+		/// <param name="Format">Message format string</param>
+		/// <param name="Args">Optional arguments</param>
+		[MethodImplAttribute(MethodImplOptions.NoInlining)]
+		public static void TraceError(string Format, params object[] Args)
+		{
 			WriteLinePrivate(1, null, false, LogEventType.Error, Format, Args);
-        }
-
-        /// <summary>
-        /// Writes a verbose message to the console.
-        /// </summary>
-        /// <param name="Format">Message format string</param>
-        /// <param name="Args">Optional arguments</param>
-        [Conditional("TRACE")]
-        [MethodImplAttribute(MethodImplOptions.NoInlining)]
-        public static void TraceVerbose(string Format, params object[] Args)
-        {
+		}
+
+		/// <summary>
+		/// Writes a verbose message to the console.
+		/// </summary>
+		/// <param name="Format">Message format string</param>
+		/// <param name="Args">Optional arguments</param>
+		[Conditional("TRACE")]
+		[MethodImplAttribute(MethodImplOptions.NoInlining)]
+		public static void TraceVerbose(string Format, params object[] Args)
+		{
 			WriteLinePrivate(1, null, false, LogEventType.Verbose, Format, Args);
-        }
-
-        /// <summary>
-        /// Writes a message to the console.
-        /// </summary>
-        /// <param name="Format">Message format string</param>
-        /// <param name="Args">Optional arguments</param>
-        [MethodImplAttribute(MethodImplOptions.NoInlining)]
-        public static void TraceInformation(string Format, params object[] Args)
-        {
+		}
+
+		/// <summary>
+		/// Writes a message to the console.
+		/// </summary>
+		/// <param name="Format">Message format string</param>
+		/// <param name="Args">Optional arguments</param>
+		[MethodImplAttribute(MethodImplOptions.NoInlining)]
+		public static void TraceInformation(string Format, params object[] Args)
+		{
 			WriteLinePrivate(1, null, false, LogEventType.Console, Format, Args);
-        }
-
-        /// <summary>
-        /// Writes a warning message to the console.
-        /// </summary>
-        /// <param name="Format">Message format string</param>
-        /// <param name="Args">Optional arguments</param>
-        [MethodImplAttribute(MethodImplOptions.NoInlining)]
-        public static void TraceWarning(string Format, params object[] Args)
-        {
+		}
+
+		/// <summary>
+		/// Writes a warning message to the console.
+		/// </summary>
+		/// <param name="Format">Message format string</param>
+		/// <param name="Args">Optional arguments</param>
+		[MethodImplAttribute(MethodImplOptions.NoInlining)]
+		public static void TraceWarning(string Format, params object[] Args)
+		{
 			WriteLinePrivate(1, null, false, LogEventType.Warning, Format, Args);
-        }
+		}
 
 		/// <summary>
 		/// Writes a very verbose message to the console.
@@ -1444,193 +1444,193 @@ namespace UnrealBuildTool
 		}
 	}
 
-    #region StreamUtils
-    public static class StreamUtils
-    {
-        /// <summary>
-        /// Read a stream into another, buffering in 4K chunks.
-        /// </summary>
-        /// <param name="output">this</param>
-        /// <param name="input">the Stream to read from</param>
-        /// <returns>same stream for expression chaining.</returns>
-        public static Stream ReadFrom(this Stream output, Stream input)
-        {
-            long bytesRead;
-            return output.ReadFrom(input, out bytesRead);
-        }
-
-        /// <summary>
-        /// Read a stream into another, buffering in 4K chunks.
-        /// </summary>
-        /// <param name="output">this</param>
-        /// <param name="input">the Stream to read from</param>
-        /// <param name="totalBytesRead">returns bytes read</param>
-        /// <returns>same stream for expression chaining.</returns>
-        public static Stream ReadFrom(this Stream output, Stream input, out long totalBytesRead)
-        {
-            totalBytesRead = 0;
-            const int BytesToRead = 4096;
-            var buf = new byte[BytesToRead];
-            var bytesReadThisTime = 0;
-            do
-            {
-                bytesReadThisTime = input.Read(buf, 0, BytesToRead);
-                totalBytesRead += bytesReadThisTime;
-                output.Write(buf, 0, bytesReadThisTime);
-            } while (bytesReadThisTime != 0);
-            return output;
-        }
-
-        /// <summary>
-        /// Read stream into a new MemoryStream. Useful for chaining together expressions.
-        /// </summary>
-        /// <param name="input">Stream to read from.</param>
-        /// <returns>memory stream that contains the stream contents.</returns>
-        public static MemoryStream ReadIntoMemoryStream(this Stream input)
-        {
-            var data = new MemoryStream(4096);
-            data.ReadFrom(input);
-            return data;
-        }
-
-        /// <summary>
-        /// Writes the entire contents of a byte array to the stream.
-        /// </summary>
-        /// <param name="stream"></param>
-        /// <param name="arr"></param>
-        /// <returns></returns>
-        public static Stream Write(this Stream stream, byte[] arr)
-        {
-            stream.Write(arr, 0, arr.Length);
-            return stream;
-        }
-    }
-    #endregion
-
-    #region HashCodeHelper
-    /// <summary>
-    /// Helper for generating hashcodes for value types.
-    /// </summary>
-    public static class HashCodeHelper
-    {
-        private const int BasePrimeNumber = 691;
-        private const int FieldMutatorPrimeNumber = 397;
-
-        /// <summary>
-        /// Helper to compute a reasonable hash code from a group of objects, typically base types, but could be anything.
-        /// </summary>
-        /// <typeparam name="T1"></typeparam>
-        /// <param name="Value1"></param>
-        /// <returns></returns>
-        public static int GetHashCodeFromValues<T1>(T1 Value1)
-        {
-            var Result = BasePrimeNumber;
-            if (Value1 != null)
-            {
-                Result *= Value1.GetHashCode() + FieldMutatorPrimeNumber;
-            }
-            return Result;
-        }
-
-        /// <summary>
-        /// Helper to compute a reasonable hash code from a group of objects, typically base types, but could be anything.
-        /// </summary>
-        /// <typeparam name="T1"></typeparam>
-        /// <typeparam name="T2"></typeparam>
-        /// <param name="Value1"></param>
-        /// <param name="Value2"></param>
-        /// <returns></returns>
-        public static int GetHashCodeFromValues<T1, T2>(T1 Value1, T2 Value2)
-        {
-            var Result = BasePrimeNumber;
-            if (Value1 != null)
-            {
-                Result *= Value1.GetHashCode() + FieldMutatorPrimeNumber;
-            }
-            if (Value2 != null)
-            {
-                Result *= Value2.GetHashCode() + FieldMutatorPrimeNumber;
-            }
-            return Result;
-        }
-
-        /// <summary>
-        /// Helper to compute a reasonable hash code from a group of objects, typically base types, but could be anything.
-        /// </summary>
-        /// <typeparam name="T1"></typeparam>
-        /// <typeparam name="T2"></typeparam>
-        /// <typeparam name="T3"></typeparam>
-        /// <param name="Value1"></param>
-        /// <param name="Value2"></param>
-        /// <param name="Value3"></param>
-        /// <returns></returns>
-        public static int GetHashCodeFromValues<T1, T2, T3>(T1 Value1, T2 Value2, T3 Value3)
-        {
-            var Result = BasePrimeNumber;
-            if (Value1 != null)
-            {
-                Result *= Value1.GetHashCode() + FieldMutatorPrimeNumber;
-            }
-            if (Value2 != null)
-            {
-                Result *= Value2.GetHashCode() + FieldMutatorPrimeNumber;
-            }
-            if (Value3 != null)
-            {
-                Result *= Value3.GetHashCode() + FieldMutatorPrimeNumber;
-            }
-            return Result;
-        }
-
-        /// <summary>
-        /// Helper to compute a reasonable hash code from a group of objects, typically base types, but could be anything.
-        /// </summary>
-        /// <typeparam name="T1"></typeparam>
-        /// <typeparam name="T2"></typeparam>
-        /// <typeparam name="T3"></typeparam>
-        /// <typeparam name="T4"></typeparam>
-        /// <param name="Value1"></param>
-        /// <param name="Value2"></param>
-        /// <param name="Value3"></param>
-        /// <param name="Value4"></param>
-        /// <returns></returns>
-        public static int GetHashCodeFromValues<T1, T2, T3, T4>(T1 Value1, T2 Value2, T3 Value3, T4 Value4)
-        {
-            var Result = BasePrimeNumber;
-            if (Value1 != null)
-            {
-                Result *= Value1.GetHashCode() + FieldMutatorPrimeNumber;
-            }
-            if (Value2 != null)
-            {
-                Result *= Value2.GetHashCode() + FieldMutatorPrimeNumber;
-            }
-            if (Value3 != null)
-            {
-                Result *= Value3.GetHashCode() + FieldMutatorPrimeNumber;
-            }
-            if (Value4 != null)
-            {
-                Result *= Value4.GetHashCode() + FieldMutatorPrimeNumber;
-            }
-            return Result;
-        }
-
-        /// <summary>
-        /// Helper to compute a reasonable hash code from a group of objects, typically base types, but could be anything.
-        /// </summary>
-        /// <param name="values"></param>
-        /// <returns></returns>
-        public static int GetHashCodeFromValues(params object[] values)
-        {
-            // prevent arithmetic overflow exceptions
-            unchecked
-            {
-                return values.Where(elt => elt != null).Aggregate(BasePrimeNumber, (result, rhs) => result * (FieldMutatorPrimeNumber + rhs.GetHashCode()));
-            }
-        }
-    }
-    #endregion
+	#region StreamUtils
+	public static class StreamUtils
+	{
+		/// <summary>
+		/// Read a stream into another, buffering in 4K chunks.
+		/// </summary>
+		/// <param name="output">this</param>
+		/// <param name="input">the Stream to read from</param>
+		/// <returns>same stream for expression chaining.</returns>
+		public static Stream ReadFrom(this Stream output, Stream input)
+		{
+			long bytesRead;
+			return output.ReadFrom(input, out bytesRead);
+		}
+
+		/// <summary>
+		/// Read a stream into another, buffering in 4K chunks.
+		/// </summary>
+		/// <param name="output">this</param>
+		/// <param name="input">the Stream to read from</param>
+		/// <param name="totalBytesRead">returns bytes read</param>
+		/// <returns>same stream for expression chaining.</returns>
+		public static Stream ReadFrom(this Stream output, Stream input, out long totalBytesRead)
+		{
+			totalBytesRead = 0;
+			const int BytesToRead = 4096;
+			var buf = new byte[BytesToRead];
+			var bytesReadThisTime = 0;
+			do
+			{
+				bytesReadThisTime = input.Read(buf, 0, BytesToRead);
+				totalBytesRead += bytesReadThisTime;
+				output.Write(buf, 0, bytesReadThisTime);
+			} while (bytesReadThisTime != 0);
+			return output;
+		}
+
+		/// <summary>
+		/// Read stream into a new MemoryStream. Useful for chaining together expressions.
+		/// </summary>
+		/// <param name="input">Stream to read from.</param>
+		/// <returns>memory stream that contains the stream contents.</returns>
+		public static MemoryStream ReadIntoMemoryStream(this Stream input)
+		{
+			var data = new MemoryStream(4096);
+			data.ReadFrom(input);
+			return data;
+		}
+
+		/// <summary>
+		/// Writes the entire contents of a byte array to the stream.
+		/// </summary>
+		/// <param name="stream"></param>
+		/// <param name="arr"></param>
+		/// <returns></returns>
+		public static Stream Write(this Stream stream, byte[] arr)
+		{
+			stream.Write(arr, 0, arr.Length);
+			return stream;
+		}
+	}
+	#endregion
+
+	#region HashCodeHelper
+	/// <summary>
+	/// Helper for generating hashcodes for value types.
+	/// </summary>
+	public static class HashCodeHelper
+	{
+		private const int BasePrimeNumber = 691;
+		private const int FieldMutatorPrimeNumber = 397;
+
+		/// <summary>
+		/// Helper to compute a reasonable hash code from a group of objects, typically base types, but could be anything.
+		/// </summary>
+		/// <typeparam name="T1"></typeparam>
+		/// <param name="Value1"></param>
+		/// <returns></returns>
+		public static int GetHashCodeFromValues<T1>(T1 Value1)
+		{
+			var Result = BasePrimeNumber;
+			if (Value1 != null)
+			{
+				Result *= Value1.GetHashCode() + FieldMutatorPrimeNumber;
+			}
+			return Result;
+		}
+
+		/// <summary>
+		/// Helper to compute a reasonable hash code from a group of objects, typically base types, but could be anything.
+		/// </summary>
+		/// <typeparam name="T1"></typeparam>
+		/// <typeparam name="T2"></typeparam>
+		/// <param name="Value1"></param>
+		/// <param name="Value2"></param>
+		/// <returns></returns>
+		public static int GetHashCodeFromValues<T1, T2>(T1 Value1, T2 Value2)
+		{
+			var Result = BasePrimeNumber;
+			if (Value1 != null)
+			{
+				Result *= Value1.GetHashCode() + FieldMutatorPrimeNumber;
+			}
+			if (Value2 != null)
+			{
+				Result *= Value2.GetHashCode() + FieldMutatorPrimeNumber;
+			}
+			return Result;
+		}
+
+		/// <summary>
+		/// Helper to compute a reasonable hash code from a group of objects, typically base types, but could be anything.
+		/// </summary>
+		/// <typeparam name="T1"></typeparam>
+		/// <typeparam name="T2"></typeparam>
+		/// <typeparam name="T3"></typeparam>
+		/// <param name="Value1"></param>
+		/// <param name="Value2"></param>
+		/// <param name="Value3"></param>
+		/// <returns></returns>
+		public static int GetHashCodeFromValues<T1, T2, T3>(T1 Value1, T2 Value2, T3 Value3)
+		{
+			var Result = BasePrimeNumber;
+			if (Value1 != null)
+			{
+				Result *= Value1.GetHashCode() + FieldMutatorPrimeNumber;
+			}
+			if (Value2 != null)
+			{
+				Result *= Value2.GetHashCode() + FieldMutatorPrimeNumber;
+			}
+			if (Value3 != null)
+			{
+				Result *= Value3.GetHashCode() + FieldMutatorPrimeNumber;
+			}
+			return Result;
+		}
+
+		/// <summary>
+		/// Helper to compute a reasonable hash code from a group of objects, typically base types, but could be anything.
+		/// </summary>
+		/// <typeparam name="T1"></typeparam>
+		/// <typeparam name="T2"></typeparam>
+		/// <typeparam name="T3"></typeparam>
+		/// <typeparam name="T4"></typeparam>
+		/// <param name="Value1"></param>
+		/// <param name="Value2"></param>
+		/// <param name="Value3"></param>
+		/// <param name="Value4"></param>
+		/// <returns></returns>
+		public static int GetHashCodeFromValues<T1, T2, T3, T4>(T1 Value1, T2 Value2, T3 Value3, T4 Value4)
+		{
+			var Result = BasePrimeNumber;
+			if (Value1 != null)
+			{
+				Result *= Value1.GetHashCode() + FieldMutatorPrimeNumber;
+			}
+			if (Value2 != null)
+			{
+				Result *= Value2.GetHashCode() + FieldMutatorPrimeNumber;
+			}
+			if (Value3 != null)
+			{
+				Result *= Value3.GetHashCode() + FieldMutatorPrimeNumber;
+			}
+			if (Value4 != null)
+			{
+				Result *= Value4.GetHashCode() + FieldMutatorPrimeNumber;
+			}
+			return Result;
+		}
+
+		/// <summary>
+		/// Helper to compute a reasonable hash code from a group of objects, typically base types, but could be anything.
+		/// </summary>
+		/// <param name="values"></param>
+		/// <returns></returns>
+		public static int GetHashCodeFromValues(params object[] values)
+		{
+			// prevent arithmetic overflow exceptions
+			unchecked
+			{
+				return values.Where(elt => elt != null).Aggregate(BasePrimeNumber, (result, rhs) => result * (FieldMutatorPrimeNumber + rhs.GetHashCode()));
+			}
+		}
+	}
+	#endregion
 
 	#region Scoped Timers
 
@@ -1666,7 +1666,7 @@ namespace UnrealBuildTool
 			}
 			var IndentText = new StringBuilder(LogIndent * 2);
 			IndentText.Append(' ', LogIndent * 2);
-			
+
 			Log.WriteLine(Verbosity, "{0}{1} took {2}s", IndentText.ToString(), TimerName, TotalSeconds);
 		}
 	}
@@ -1684,7 +1684,7 @@ namespace UnrealBuildTool
 		DateTime StartTime;
 		Accumulator ScopedAccumulator;
 		string CounterName;
-		
+
 		static Dictionary<string, Accumulator> Accumulators = new Dictionary<string, Accumulator>();
 
 		public ScopedCounter(string Name, LogEventType InVerbosity = LogEventType.Verbose)
diff --git a/Engine/Source/Programs/UnrealBuildTool/System/XcodeProjectFileGenerator.cs b/Engine/Source/Programs/UnrealBuildTool/System/XcodeProjectFileGenerator.cs
--- a/Engine/Source/Programs/UnrealBuildTool/System/XcodeProjectFileGenerator.cs
+++ b/Engine/Source/Programs/UnrealBuildTool/System/XcodeProjectFileGenerator.cs
@@ -113,9 +113,9 @@ namespace UnrealBuildTool
 
 		public bool HasUI; // true if generates an .app instead of a console exe
 
-		public override string ToString ()
+		public override string ToString()
 		{
-			return string.Format ("[XcodeProjectTarget: {0}, {1}]", DisplayName, Type);
+			return string.Format("[XcodeProjectTarget: {0}, {1}]", DisplayName, Type);
 		}
 	}
 
@@ -167,10 +167,10 @@ namespace UnrealBuildTool
 			Path = InPath;
 			SourceTree = InSourceTree;
 		}
-	    public string Name;
-	    public string Path;
-	    public string SourceTree;
-	    public string Guid;
+		public string Name;
+		public string Path;
+		public string SourceTree;
+		public string Guid;
 	}
 
 	/// <summary>
@@ -192,7 +192,8 @@ namespace UnrealBuildTool
 		// always seed the random number the same, so multiple runs of the generator will generate the same project
 		static Random Rand = new Random(0);
 
-		public XcodeProjectFileGenerator(FileReference InOnlyGameProject) : base(InOnlyGameProject)
+		public XcodeProjectFileGenerator(FileReference InOnlyGameProject)
+			: base(InOnlyGameProject)
 		{
 		}
 
@@ -254,7 +255,7 @@ namespace UnrealBuildTool
 		/// </summary>
 		/// <param name="Contents">StringBuilder object to append groups string to</param>
 		/// <param name="Groups">Dictionary of all project groups</param>
-        /// <param name="Frameworks">Frameworks referenced</param>
+		/// <param name="Frameworks">Frameworks referenced</param>
 		private void AppendGroups(ref StringBuilder Contents, ref Dictionary<string, XcodeFileGroup> Groups, List<XcodeProjectTarget> Targets, List<XcodeFramework> Frameworks)
 		{
 			Contents.Append("/* Begin PBXGroup section */" + ProjectFileGenerator.NewLine);
@@ -282,9 +283,9 @@ namespace UnrealBuildTool
 			}
 
 			ProductRefGroupGuid = MakeXcodeGuid();
-            FrameworkGroupGuid = MakeXcodeGuid();
+			FrameworkGroupGuid = MakeXcodeGuid();
 			Contents.Append(string.Format("\t\t\t\t{0} /* Products */,{1}", ProductRefGroupGuid, ProjectFileGenerator.NewLine));
-            Contents.Append(string.Format("\t\t\t\t{0} /* Frameworks */,{1}", FrameworkGroupGuid, ProjectFileGenerator.NewLine));
+			Contents.Append(string.Format("\t\t\t\t{0} /* Frameworks */,{1}", FrameworkGroupGuid, ProjectFileGenerator.NewLine));
 			if (Groups.ContainsKey(""))
 			{
 				Groups[""].Append(ref Contents, bFilesOnly: true);
@@ -311,18 +312,18 @@ namespace UnrealBuildTool
 			Contents.Append("\t\t\tsourceTree = \"<group>\";" + ProjectFileGenerator.NewLine);
 			Contents.Append("\t\t};" + ProjectFileGenerator.NewLine);
 
-            // Add Frameworks group
-            Contents.Append(string.Format("\t\t{0} = {{{1}", FrameworkGroupGuid, ProjectFileGenerator.NewLine));
-            Contents.Append("\t\t\tisa = PBXGroup;" + ProjectFileGenerator.NewLine);
-            Contents.Append("\t\t\tchildren = (" + ProjectFileGenerator.NewLine);
-            foreach (XcodeFramework Framework in Frameworks)
-            {
-                Contents.Append("\t\t\t\t" + Framework.Guid + " /* " + Framework.Name + " */," + ProjectFileGenerator.NewLine);
-            }
-            Contents.Append("\t\t\t);" + ProjectFileGenerator.NewLine);
-            Contents.Append("\t\t\tname = Frameworks;" + ProjectFileGenerator.NewLine);
-            Contents.Append("\t\t\tsourceTree = \"<group>\";" + ProjectFileGenerator.NewLine);
-            Contents.Append("\t\t};" + ProjectFileGenerator.NewLine);
+			// Add Frameworks group
+			Contents.Append(string.Format("\t\t{0} = {{{1}", FrameworkGroupGuid, ProjectFileGenerator.NewLine));
+			Contents.Append("\t\t\tisa = PBXGroup;" + ProjectFileGenerator.NewLine);
+			Contents.Append("\t\t\tchildren = (" + ProjectFileGenerator.NewLine);
+			foreach (XcodeFramework Framework in Frameworks)
+			{
+				Contents.Append("\t\t\t\t" + Framework.Guid + " /* " + Framework.Name + " */," + ProjectFileGenerator.NewLine);
+			}
+			Contents.Append("\t\t\t);" + ProjectFileGenerator.NewLine);
+			Contents.Append("\t\t\tname = Frameworks;" + ProjectFileGenerator.NewLine);
+			Contents.Append("\t\t\tsourceTree = \"<group>\";" + ProjectFileGenerator.NewLine);
+			Contents.Append("\t\t};" + ProjectFileGenerator.NewLine);
 
 			foreach (XcodeFileGroup Group in Groups.Values)
 			{
@@ -617,7 +618,7 @@ namespace UnrealBuildTool
 				{
 					string ProjectPath = GamePath;
 					string GameName = TargetName;
-					if (string.IsNullOrEmpty (ProjectPath))
+					if (string.IsNullOrEmpty(ProjectPath))
 					{
 						ProjectPath = EngineRelative + "Engine";
 					}
@@ -738,7 +739,7 @@ namespace UnrealBuildTool
 
 			foreach (string GameFolder in GameFolders.Select(x => x.FullName))
 			{
-				if (File.Exists(Path.Combine(GameFolder, Target.TargetName+".uproject")))
+				if (File.Exists(Path.Combine(GameFolder, Target.TargetName + ".uproject")))
 				{
 					IsAGame = true;
 					GamePath = ConvertPath(GameFolder);
@@ -773,7 +774,7 @@ namespace UnrealBuildTool
 		{
 			return InPath.Replace("\\", "/");
 		}
-		
+
 		/// <summary>
 		/// Appends a build configuration list section for specific target.
 		/// </summary>
@@ -887,7 +888,7 @@ namespace UnrealBuildTool
 						"\t\t\t);" + ProjectFileGenerator.NewLine +
 						"\t\t\trunOnlyForDeploymentPostprocessing = 0;" + ProjectFileGenerator.NewLine +
 						"\t\t\tshellPath = /bin/sh;" + ProjectFileGenerator.NewLine +
-						"\t\t\tshellScript = \"if [ $DEPLOYMENT_LOCATION = \\\"YES\\\" ]\\nthen\\ncp -R " + PayloadDir +"/Binaries/IOS/Payload/" + Target.ProductName + "/ $DSTROOT/$LOCAL_APPS_DIR/" + Target.ProductName + "\\nfi\";" + ProjectFileGenerator.NewLine +
+						"\t\t\tshellScript = \"if [ $DEPLOYMENT_LOCATION = \\\"YES\\\" ]\\nthen\\ncp -R " + PayloadDir + "/Binaries/IOS/Payload/" + Target.ProductName + "/ $DSTROOT/$LOCAL_APPS_DIR/" + Target.ProductName + "\\nfi\";" + ProjectFileGenerator.NewLine +
 						"\t\t};" + ProjectFileGenerator.NewLine;
 				}
 			}
@@ -898,7 +899,7 @@ namespace UnrealBuildTool
 		{
 			var XcodeProject = new XcodeProjectFile(ProjectPath);
 			DirectoryReference ProjectDirectory = ProjectPath.Directory;
-			XcodeProject.AddFilesToProject(SourceFileSearch.FindFiles(ProjectDirectory, SubdirectoryNamesToExclude:new List<string>() { "obj" }, SearchSubdirectories:true), null);
+			XcodeProject.AddFilesToProject(SourceFileSearch.FindFiles(ProjectDirectory, SubdirectoryNamesToExclude: new List<string>() { "obj" }, SearchSubdirectories: true), null);
 			XcodeProject.GenerateSectionsContents(ref PBXBuildFileSection, ref PBXFileReferenceSection, ref PBXSourcesBuildPhaseSection, ref Groups);
 		}
 
@@ -1033,7 +1034,7 @@ namespace UnrealBuildTool
 		/// <param name="TargetDependencies">List of container item proxies to add.</param>
 		private void CreateContainerItemProxySection(ref string PBXContainerItemProxySection, List<XcodeContainerItemProxy> ContainerItemProxies)
 		{
-			foreach(XcodeContainerItemProxy ContainerItemProxy in ContainerItemProxies)
+			foreach (XcodeContainerItemProxy ContainerItemProxy in ContainerItemProxies)
 			{
 				PBXContainerItemProxySection += "\t\t" + ContainerItemProxy.Guid + " /* PBXContainerItemProxy */ = {" + ProjectFileGenerator.NewLine +
 												"\t\t\tisa = PBXContainerItemProxy;" + ProjectFileGenerator.NewLine +
@@ -1052,9 +1053,9 @@ namespace UnrealBuildTool
 		/// <param name="TargetDependencies">List of target dependencies to add.</param>
 		private void CreateTargetDependencySection(ref string PBXTargetDependencySection, List<XcodeTargetDependency> TargetDependencies)
 		{
-			foreach(XcodeTargetDependency TargetDependency in TargetDependencies)
+			foreach (XcodeTargetDependency TargetDependency in TargetDependencies)
 			{
-				PBXTargetDependencySection +=	"\t\t" + TargetDependency.Guid + " /* PBXTargetDependency */ = {" + ProjectFileGenerator.NewLine +
+				PBXTargetDependencySection += "\t\t" + TargetDependency.Guid + " /* PBXTargetDependency */ = {" + ProjectFileGenerator.NewLine +
 												"\t\t\tisa = PBXTargetDependency;" + ProjectFileGenerator.NewLine +
 												"\t\t\ttarget = " + TargetDependency.LegacyTargetGuid + " /* " + TargetDependency.LegacyTargetName + " */;" + ProjectFileGenerator.NewLine +
 												"\t\t\ttargetProxy = " + TargetDependency.ContainerItemProxyGuid + " /* PBXContainerItemProxy */;" + ProjectFileGenerator.NewLine +
@@ -1078,7 +1079,7 @@ namespace UnrealBuildTool
 				FullPath = Utils.MakePathRelativeTo(FullPath, FullProjectPath);
 				FullPath = FullPath.TrimEnd('/');
 			}
-			if ( !IncludeDirectories.Contains(FullPath) )
+			if (!IncludeDirectories.Contains(FullPath))
 			{
 				IncludeDirectories.Add(FullPath);
 			}
@@ -1086,12 +1087,12 @@ namespace UnrealBuildTool
 
 		private void PopulateTargets(List<XcodeProjectTarget> ProjectTargets, List<XcodeContainerItemProxy> ContainerItemProxies, List<XcodeTargetDependency> TargetDependencies, XcodeProjectTarget ProjectTarget, List<XcodeFramework> Frameworks)
 		{
-			List<ProjectTarget> TargetProjects = new List<ProjectTarget> ();
+			List<ProjectTarget> TargetProjects = new List<ProjectTarget>();
 			foreach (var Project in GeneratedProjectFiles)
 			{
 				foreach (var TargetProject in Project.ProjectTargets)
 				{
-					TargetProjects.Add (TargetProject);
+					TargetProjects.Add(TargetProject);
 				}
 			}
 			if (bGeneratingGameProjectFiles && TargetProjects.Count == 0)
@@ -1108,109 +1109,109 @@ namespace UnrealBuildTool
 			}
 			foreach (var TargetProject in TargetProjects)
 			{
-					FileReference TargetPath = TargetProject.TargetFilePath;
-					string TargetName = TargetPath.GetFileNameWithoutAnyExtensions();
-					bool WantProjectFileForTarget = true;
-					bool IsEngineTarget = false;
-					if (bGeneratingGameProjectFiles || bGeneratingRocketProjectFiles)
+				FileReference TargetPath = TargetProject.TargetFilePath;
+				string TargetName = TargetPath.GetFileNameWithoutAnyExtensions();
+				bool WantProjectFileForTarget = true;
+				bool IsEngineTarget = false;
+				if (bGeneratingGameProjectFiles || bGeneratingRocketProjectFiles)
+				{
+					// Check to see if this is an Engine target.  That is, the target is located under the "Engine" folder
+					if (TargetPath.IsUnderDirectory(UnrealBuildTool.EngineDirectory))
 					{
-						// Check to see if this is an Engine target.  That is, the target is located under the "Engine" folder
-						if (TargetPath.IsUnderDirectory(UnrealBuildTool.EngineDirectory))
-						{
-							// This is an engine target
-							IsEngineTarget = true;
-						}
+						// This is an engine target
+						IsEngineTarget = true;
+					}
 
-						if (IsEngineTarget)
+					if (IsEngineTarget)
+					{
+						if (!IncludeEngineSource)
 						{
-							if (!IncludeEngineSource)
-							{
-								// We were asked to exclude engine modules from the generated projects
-								WantProjectFileForTarget = false;
-							}
-							if (bGeneratingGameProjectFiles && this.GameProjectName == TargetName)
-							{
-								WantProjectFileForTarget = true;
-							}
+							// We were asked to exclude engine modules from the generated projects
+							WantProjectFileForTarget = false;
 						}
-
-						if (bGeneratingRocketProjectFiles && TargetName.EndsWith("Server"))
+						if (bGeneratingGameProjectFiles && this.GameProjectName == TargetName)
 						{
-							WantProjectFileForTarget = false;
+							WantProjectFileForTarget = true;
 						}
 					}
 
-					if (WantProjectFileForTarget)
+					if (bGeneratingRocketProjectFiles && TargetName.EndsWith("Server"))
 					{
-						FileReference TargetFilePath = TargetProject.TargetFilePath;
-						var Target = new TargetInfo(UnrealTargetPlatform.Mac, UnrealTargetConfiguration.Development);
-						var TargetRulesObject = TargetProject.TargetRules; // RulesCompiler.CreateTargetRules(TargetName, Target, false, out TargetFilePath);
-						List<UnrealTargetPlatform> SupportedPlatforms = new List<UnrealTargetPlatform>();
-						TargetRulesObject.GetSupportedPlatforms(ref SupportedPlatforms);
-						LinkEnvironmentConfiguration LinkConfiguration = new LinkEnvironmentConfiguration();
-						CPPEnvironmentConfiguration CPPConfiguration = new CPPEnvironmentConfiguration();
-						TargetRulesObject.SetupGlobalEnvironment(Target, ref LinkConfiguration, ref CPPConfiguration);
-
-						if (!LinkConfiguration.bIsBuildingConsoleApplication)
-						{
-							TargetsThatNeedApp.Add(TargetName);
-						}
+						WantProjectFileForTarget = false;
+					}
+				}
+
+				if (WantProjectFileForTarget)
+				{
+					FileReference TargetFilePath = TargetProject.TargetFilePath;
+					var Target = new TargetInfo(UnrealTargetPlatform.Mac, UnrealTargetConfiguration.Development);
+					var TargetRulesObject = TargetProject.TargetRules; // RulesCompiler.CreateTargetRules(TargetName, Target, false, out TargetFilePath);
+					List<UnrealTargetPlatform> SupportedPlatforms = new List<UnrealTargetPlatform>();
+					TargetRulesObject.GetSupportedPlatforms(ref SupportedPlatforms);
+					LinkEnvironmentConfiguration LinkConfiguration = new LinkEnvironmentConfiguration();
+					CPPEnvironmentConfiguration CPPConfiguration = new CPPEnvironmentConfiguration();
+					TargetRulesObject.SetupGlobalEnvironment(Target, ref LinkConfiguration, ref CPPConfiguration);
+
+					if (!LinkConfiguration.bIsBuildingConsoleApplication)
+					{
+						TargetsThatNeedApp.Add(TargetName);
+					}
 
-						// if the project is not an engine project check to make sure we have the correct name
-						string DisplayName = TargetName;
-						if (!IsEngineTarget && TargetRulesObject.Type != TargetRules.TargetType.Program && TargetRulesObject.Type != TargetRules.TargetType.Client)
+					// if the project is not an engine project check to make sure we have the correct name
+					string DisplayName = TargetName;
+					if (!IsEngineTarget && TargetRulesObject.Type != TargetRules.TargetType.Program && TargetRulesObject.Type != TargetRules.TargetType.Client)
+					{
+						List<UProjectInfo> AllGames = UProjectInfo.FilterGameProjects(true, bGeneratingGameProjectFiles ? GameProjectName : null);
+						UProjectInfo ProjectInfo = FindGameContainingFile(AllGames, TargetFilePath);
+						if (ProjectInfo != null)
 						{
-							List<UProjectInfo> AllGames = UProjectInfo.FilterGameProjects(true, bGeneratingGameProjectFiles ? GameProjectName : null);
-							UProjectInfo ProjectInfo = FindGameContainingFile(AllGames, TargetFilePath);
-							if (ProjectInfo != null)
+							DisplayName = ProjectInfo.GameName;
+							if (TargetName.Contains("Editor"))
+							{
+								DisplayName += "Editor";
+							}
+							else if (TargetName.Contains("Server"))
 							{
-								DisplayName = ProjectInfo.GameName;
-								if (TargetName.Contains("Editor"))
-								{
-									DisplayName += "Editor";
-								}
-								else if (TargetName.Contains("Server"))
-								{
-									DisplayName += "Server";
-								}
+								DisplayName += "Server";
 							}
 						}
+					}
 
-						// @todo: Remove target platform param and merge Mac and iOS targets. For now BuildTarget knows how to build iOS, but cannot run iOS apps, so we need separate DeployTarget.
-						bool bIsMacOnly = !SupportedPlatforms.Contains(UnrealTargetPlatform.IOS);
+					// @todo: Remove target platform param and merge Mac and iOS targets. For now BuildTarget knows how to build iOS, but cannot run iOS apps, so we need separate DeployTarget.
+					bool bIsMacOnly = !SupportedPlatforms.Contains(UnrealTargetPlatform.IOS);
 
-						XcodeProjectTarget BuildTarget = new XcodeProjectTarget(DisplayName + " - Mac", TargetName, XcodeTargetType.Legacy, TargetFilePath, "", UnrealTargetPlatform.Mac, bIsMacOnly);
-						if (!bGeneratingRunIOSProject)
-						{
-							ProjectTargets.Add(BuildTarget);
-						}
+					XcodeProjectTarget BuildTarget = new XcodeProjectTarget(DisplayName + " - Mac", TargetName, XcodeTargetType.Legacy, TargetFilePath, "", UnrealTargetPlatform.Mac, bIsMacOnly);
+					if (!bGeneratingRunIOSProject)
+					{
+						ProjectTargets.Add(BuildTarget);
+					}
 
-						if (ProjectFilePlatform.HasFlag(XcodeProjectFilePlatform.iOS) && SupportedPlatforms.Contains(UnrealTargetPlatform.IOS))
+					if (ProjectFilePlatform.HasFlag(XcodeProjectFilePlatform.iOS) && SupportedPlatforms.Contains(UnrealTargetPlatform.IOS))
+					{
+						if ((bGeneratingRocketProjectFiles && TargetName == "UE4Game") || bGeneratingRunIOSProject)
 						{
-							if ((bGeneratingRocketProjectFiles && TargetName == "UE4Game") || bGeneratingRunIOSProject)
-							{
-								// Generate Framework references.
-								List<XcodeFrameworkRef> FrameworkRefs = new List<XcodeFrameworkRef>();
-								foreach (XcodeFramework Framework in Frameworks)
-								{
-									FrameworkRefs.Add(new XcodeFrameworkRef(Framework));
-								}
-
-								XcodeProjectTarget IOSDeployTarget = new XcodeProjectTarget(DisplayName + " - iOS", TargetName, XcodeTargetType.Native, TargetFilePath, TargetName + ".app", UnrealTargetPlatform.IOS, false, null, true, FrameworkRefs);
-								ProjectTargets.Add(IOSDeployTarget);
-							}
-							else
+							// Generate Framework references.
+							List<XcodeFrameworkRef> FrameworkRefs = new List<XcodeFrameworkRef>();
+							foreach (XcodeFramework Framework in Frameworks)
 							{
-								XcodeContainerItemProxy ContainerProxy = new XcodeContainerItemProxy(ProjectTarget.Guid, BuildTarget.Guid, BuildTarget.DisplayName);
-								XcodeTargetDependency TargetDependency = new XcodeTargetDependency(BuildTarget.DisplayName, BuildTarget.Guid, ContainerProxy.Guid);
-								XcodeProjectTarget IOSDeployTarget = new XcodeProjectTarget(DisplayName + " - iOS", TargetName, XcodeTargetType.Native, TargetFilePath, TargetName + ".app", UnrealTargetPlatform.IOS, false, new List<XcodeTargetDependency>() { TargetDependency }, true);
-								ProjectTargets.Add(IOSDeployTarget);
-								ContainerItemProxies.Add(ContainerProxy);
-								TargetDependencies.Add(TargetDependency);
+								FrameworkRefs.Add(new XcodeFrameworkRef(Framework));
 							}
+
+							XcodeProjectTarget IOSDeployTarget = new XcodeProjectTarget(DisplayName + " - iOS", TargetName, XcodeTargetType.Native, TargetFilePath, TargetName + ".app", UnrealTargetPlatform.IOS, false, null, true, FrameworkRefs);
+							ProjectTargets.Add(IOSDeployTarget);
+						}
+						else
+						{
+							XcodeContainerItemProxy ContainerProxy = new XcodeContainerItemProxy(ProjectTarget.Guid, BuildTarget.Guid, BuildTarget.DisplayName);
+							XcodeTargetDependency TargetDependency = new XcodeTargetDependency(BuildTarget.DisplayName, BuildTarget.Guid, ContainerProxy.Guid);
+							XcodeProjectTarget IOSDeployTarget = new XcodeProjectTarget(DisplayName + " - iOS", TargetName, XcodeTargetType.Native, TargetFilePath, TargetName + ".app", UnrealTargetPlatform.IOS, false, new List<XcodeTargetDependency>() { TargetDependency }, true);
+							ProjectTargets.Add(IOSDeployTarget);
+							ContainerItemProxies.Add(ContainerProxy);
+							TargetDependencies.Add(TargetDependency);
 						}
 					}
 				}
+			}
 		}
 
 		/// <summary>
@@ -1237,28 +1238,28 @@ namespace UnrealBuildTool
 			{
 				TargetBinariesFolder = Path.GetDirectoryName(Directory.GetCurrentDirectory());
 
-                if (!ExeBaseName.EndsWith("Editor") && Target.TargetFilePath != null && !Target.TargetFilePath.IsUnderDirectory(new DirectoryReference(TargetBinariesFolder)))
-                {
+				if (!ExeBaseName.EndsWith("Editor") && Target.TargetFilePath != null && !Target.TargetFilePath.IsUnderDirectory(new DirectoryReference(TargetBinariesFolder)))
+				{
 					string DepotRoot = TargetBinariesFolder.TrimEnd(Path.DirectorySeparatorChar);
 					DepotRoot = DepotRoot.Substring(0, DepotRoot.LastIndexOf(Path.DirectorySeparatorChar));
 					DepotRoot = DepotRoot.TrimEnd(Path.DirectorySeparatorChar);
 
-                    if (Target.TargetFilePath.IsUnderDirectory(new DirectoryReference(DepotRoot)))
-                    {
-                        string PreviousPath = Target.TargetFilePath.FullName;
-                        string CurrentPath = Target.TargetFilePath.FullName;
-                        while (!string.IsNullOrWhiteSpace(CurrentPath) && CurrentPath != DepotRoot)
-                        {
-                            PreviousPath = CurrentPath;
-                            CurrentPath = CurrentPath.Substring(0, CurrentPath.LastIndexOf(Path.DirectorySeparatorChar));
-                        }
-
-                        if (CurrentPath == DepotRoot && !string.IsNullOrWhiteSpace(PreviousPath))
-                        {
-                            TargetBinariesFolder = PreviousPath;
-                        }
-                    }
-                }
+					if (Target.TargetFilePath.IsUnderDirectory(new DirectoryReference(DepotRoot)))
+					{
+						string PreviousPath = Target.TargetFilePath.FullName;
+						string CurrentPath = Target.TargetFilePath.FullName;
+						while (!string.IsNullOrWhiteSpace(CurrentPath) && CurrentPath != DepotRoot)
+						{
+							PreviousPath = CurrentPath;
+							CurrentPath = CurrentPath.Substring(0, CurrentPath.LastIndexOf(Path.DirectorySeparatorChar));
+						}
+
+						if (CurrentPath == DepotRoot && !string.IsNullOrWhiteSpace(PreviousPath))
+						{
+							TargetBinariesFolder = PreviousPath;
+						}
+					}
+				}
 
 				if (ExeBaseName == Target.TargetName && (Target.TargetFilePath == null || !Target.TargetFilePath.FullName.StartsWith(TargetBinariesFolder)) && !ExeBaseName.EndsWith("Editor") && ExeExtension == ".app")
 				{
@@ -1545,10 +1546,10 @@ namespace UnrealBuildTool
 				"\tobjects = {" + ProjectFileGenerator.NewLine + ProjectFileGenerator.NewLine);
 
 			// attempt to determine targets for the project
-			List<XcodeProjectTarget> ProjectTargets = new List<XcodeProjectTarget> ();
+			List<XcodeProjectTarget> ProjectTargets = new List<XcodeProjectTarget>();
 			// add mandatory ones
-			XcodeProjectTarget UE4ProjectTarget = new XcodeProjectTarget ("UE4", "UE4", XcodeTargetType.Project, null);
-			XcodeProjectTarget UE4XcodeHelperTarget = new XcodeProjectTarget ("UE4XcodeHelper", "UE4XcodeHelper", XcodeTargetType.XcodeHelper, null, "libUE4XcodeHelper.a");
+			XcodeProjectTarget UE4ProjectTarget = new XcodeProjectTarget("UE4", "UE4", XcodeTargetType.Project, null);
+			XcodeProjectTarget UE4XcodeHelperTarget = new XcodeProjectTarget("UE4XcodeHelper", "UE4XcodeHelper", XcodeTargetType.XcodeHelper, null, "libUE4XcodeHelper.a");
 			ProjectTargets.AddRange(new XcodeProjectTarget[] { UE4ProjectTarget, UE4XcodeHelperTarget });
 
 			if (ProjectFilePlatform.HasFlag(XcodeProjectFilePlatform.iOS))
@@ -1561,19 +1562,19 @@ namespace UnrealBuildTool
 
 			List<XcodeTargetDependency> TargetDependencies = new List<XcodeTargetDependency>();
 			List<XcodeContainerItemProxy> ContainerItemProxies = new List<XcodeContainerItemProxy>();
-            List<XcodeFramework> Frameworks = new List<XcodeFramework>();
-            Frameworks.Add(new XcodeFramework("OpenGLES.framework", "System/Library/Frameworks/OpenGLES.framework", "SDKROOT"));
+			List<XcodeFramework> Frameworks = new List<XcodeFramework>();
+			Frameworks.Add(new XcodeFramework("OpenGLES.framework", "System/Library/Frameworks/OpenGLES.framework", "SDKROOT"));
 			// @todo metal: putting this into the project will make for VERY slow Metal runtime by default...
-//			Frameworks.Add(new XcodeFramework("Metal.framework", "System/Library/Frameworks/Metal.framework", "SDKROOT"));
+			//			Frameworks.Add(new XcodeFramework("Metal.framework", "System/Library/Frameworks/Metal.framework", "SDKROOT"));
 
 			XcodeFramework XCTestFramework = new XcodeFramework("XCTest.framework", "Library/Frameworks/XCTest.framework", "DEVELOPER_DIR");
-            Frameworks.Add(XCTestFramework);
+			Frameworks.Add(XCTestFramework);
 
-//			var AllTargets = DiscoverTargets();
+			//			var AllTargets = DiscoverTargets();
 
 			PopulateTargets(ProjectTargets, ContainerItemProxies, TargetDependencies, UE4ProjectTarget, Frameworks);
 
-			Log.TraceInformation(string.Format ("Found {0} targets!", ProjectTargets.Count));
+			Log.TraceInformation(string.Format("Found {0} targets!", ProjectTargets.Count));
 
 			string PBXBuildFileSection = "/* Begin PBXBuildFile section */" + ProjectFileGenerator.NewLine;
 			string PBXFileReferenceSection = "/* Begin PBXFileReference section */" + ProjectFileGenerator.NewLine;
@@ -1586,17 +1587,17 @@ namespace UnrealBuildTool
 			List<string> SystemIncludeDirectories = new List<string>();
 			List<string> PreprocessorDefinitions = new List<string>();
 
-            foreach (XcodeFramework Framework in Frameworks)
-            {
-                // Add file references.
-                PBXFileReferenceSection += "\t\t" + Framework.Guid + " /* " + Framework.Name + " */ = {"
-                                         + "isa = PBXFileReference; "
-                                         + "lastKnownFileType = wrapper.framework; "
-                                         + "name = " + Framework.Name + "; "
-                                         + "path = " + Framework.Path + "; "
-                                         + "sourceTree = " + Framework.SourceTree + "; "
-                                         + "};" + ProjectFileGenerator.NewLine;
-            }
+			foreach (XcodeFramework Framework in Frameworks)
+			{
+				// Add file references.
+				PBXFileReferenceSection += "\t\t" + Framework.Guid + " /* " + Framework.Name + " */ = {"
+										 + "isa = PBXFileReference; "
+										 + "lastKnownFileType = wrapper.framework; "
+										 + "name = " + Framework.Name + "; "
+										 + "path = " + Framework.Path + "; "
+										 + "sourceTree = " + Framework.SourceTree + "; "
+										 + "};" + ProjectFileGenerator.NewLine;
+			}
 
 			// Set up all the test guids that need to be explicitly referenced later
 			string UE4CmdLineRunMFileGuid = MakeXcodeGuid();
@@ -1658,7 +1659,7 @@ namespace UnrealBuildTool
 			XcodeProjectFileContent.Append(PBXContainerItemProxySection);
 			XcodeProjectFileContent.Append(PBXTargetDependencySection);
 
-			AppendGroups(ref XcodeProjectFileContent, ref Groups, ProjectTargets , Frameworks);
+			AppendGroups(ref XcodeProjectFileContent, ref Groups, ProjectTargets, Frameworks);
 
 			XcodeProjectFileContent.Append(PBXShellScriptBuildPhaseSection);
 			XcodeProjectFileContent.Append(PBXSourcesBuildPhaseSection);
@@ -1669,7 +1670,7 @@ namespace UnrealBuildTool
 			{
 				if (Target.Type == XcodeTargetType.Legacy)
 				{
-					AppendTarget (ref XcodeProjectFileContent, Target);
+					AppendTarget(ref XcodeProjectFileContent, Target);
 				}
 			}
 			XcodeProjectFileContent.Append("/* End PBXLegacyTarget section */" + ProjectFileGenerator.NewLine + ProjectFileGenerator.NewLine);
@@ -1712,7 +1713,7 @@ namespace UnrealBuildTool
 			{
 				if (Target != UE4ProjectTarget)
 				{
-					XcodeProjectFileContent.AppendLine(string.Format ("\t\t\t\t{0} /* {1} */,", Target.Guid, Target.DisplayName));
+					XcodeProjectFileContent.AppendLine(string.Format("\t\t\t\t{0} /* {1} */,", Target.Guid, Target.DisplayName));
 				}
 			}
 
