Convert everything in UBT to use standard C# XML comments, so IntelliSense can display useful information in tooltips.

diff --git a/Engine/Source/Programs/UnrealBuildTool/Android/AndroidProjectGenerator.cs b/Engine/Source/Programs/UnrealBuildTool/Android/AndroidProjectGenerator.cs
--- a/Engine/Source/Programs/UnrealBuildTool/Android/AndroidProjectGenerator.cs
+++ b/Engine/Source/Programs/UnrealBuildTool/Android/AndroidProjectGenerator.cs
@@ -8,18 +8,18 @@ using System.IO;
 
 namespace UnrealBuildTool
 {
-	/**
-	 *	Base class for platform-specific project generators 
-	 */
+	/// <summary>
+	/// Base class for platform-specific project generators
+	/// </summary>
 	class AndroidProjectGenerator : UEPlatformProjectGenerator
 	{
 		static bool CheckedForNsight = false;		// whether we have checked for a recent enough version of Nsight yet
 		static bool NsightInstalled = false;		// true if a recent enough version of Nsight is installed
 		static int NsightVersionCode = 0;			// version code matching detected Nsight
 
-		/**
-		 *	Check to see if a recent enough version of Nsight is installed.
-		 */
+		/// <summary>
+		/// Check to see if a recent enough version of Nsight is installed.
+		/// </summary>
 		bool IsNsightInstalled()
 		{
 			// cache the results since this gets called a number of times
@@ -93,37 +93,33 @@ namespace UnrealBuildTool
 			return NsightInstalled;
 		}
 
-		/**
-		 *	Register the platform with the UEPlatformProjectGenerator class
-		 */
+		/// <summary>
+		/// Register the platform with the UEPlatformProjectGenerator class
+		/// </summary>
 		public override void RegisterPlatformProjectGenerator()
 		{
 			Log.TraceVerbose("        Registering for {0}", UnrealTargetPlatform.Android.ToString());
 			UEPlatformProjectGenerator.RegisterPlatformProjectGenerator(UnrealTargetPlatform.Android, this);
 		}
 
-		/**
-		 *	Whether this build platform has native support for VisualStudio
-		 *	
-		 *	@param	InPlatform			The UnrealTargetPlatform being built
-		 *	@param	InConfiguration		The UnrealTargetConfiguration being built
-		 *	
-		 *	@return	bool				true if native VisualStudio support (or custom VSI) is available
-		 */
+		/// <summary>
+		/// Whether this build platform has native support for VisualStudio
+		/// </summary>
+		/// <param name="InPlatform">  The UnrealTargetPlatform being built</param>
+		/// <param name="InConfiguration"> The UnrealTargetConfiguration being built</param>
+		/// <returns>bool    true if native VisualStudio support (or custom VSI) is available</returns>
 		public override bool HasVisualStudioSupport(UnrealTargetPlatform InPlatform, UnrealTargetConfiguration InConfiguration)
 		{
 			// Debugging, etc. are dependent on the TADP being installed
 			return IsNsightInstalled();
 		}
 
-		/**
-		 *	Return the VisualStudio platform name for this build platform
-		 *	
-		 *	@param	InPlatform			The UnrealTargetPlatform being built
-		 *	@param	InConfiguration		The UnrealTargetConfiguration being built
-		 *	
-		 *	@return	string				The name of the platform that VisualStudio recognizes
-		 */
+		/// <summary>
+		/// Return the VisualStudio platform name for this build platform
+		/// </summary>
+		/// <param name="InPlatform">  The UnrealTargetPlatform being built</param>
+		/// <param name="InConfiguration"> The UnrealTargetConfiguration being built</param>
+		/// <returns>string    The name of the platform that VisualStudio recognizes</returns>
 		public override string GetVisualStudioPlatformName(UnrealTargetPlatform InPlatform, UnrealTargetConfiguration InConfiguration)
 		{
 			if (InPlatform == UnrealTargetPlatform.Android)
@@ -134,13 +130,11 @@ namespace UnrealBuildTool
 			return InPlatform.ToString();
 		}
 
-		/**
-		 * Return any custom property group lines
-		 *	
-		 *	@param	InPlatform			The UnrealTargetPlatform being built
-		 *	
-		 *	@return	string				The custom property import lines for the project file; Empty string if it doesn't require one
-		 */
+		/// <summary>
+		/// Return any custom property group lines
+		/// </summary>
+		/// <param name="InPlatform">  The UnrealTargetPlatform being built</param>
+		/// <returns>string    The custom property import lines for the project file; Empty string if it doesn't require one</returns>
 		public override string GetAdditionalVisualStudioPropertyGroups(UnrealTargetPlatform InPlatform)
 		{
 			if (!IsNsightInstalled())
@@ -153,13 +147,11 @@ namespace UnrealBuildTool
 					"	</PropertyGroup>" + ProjectFileGenerator.NewLine;
 		}
 
-		/**
-		 * Return any custom property group lines
-		 *	
-		 *	@param	InPlatform			The UnrealTargetPlatform being built
-		 *	
-		 *	@return	string				The custom property import lines for the project file; Empty string if it doesn't require one
-		 */
+		/// <summary>
+		/// Return any custom property group lines
+		/// </summary>
+		/// <param name="InPlatform">  The UnrealTargetPlatform being built</param>
+		/// <returns>string    The custom property import lines for the project file; Empty string if it doesn't require one</returns>
 		public override string GetVisualStudioPlatformConfigurationType(UnrealTargetPlatform InPlatform)
 		{
 			if (!IsNsightInstalled())
@@ -170,14 +162,12 @@ namespace UnrealBuildTool
 			return "ExternalBuildSystem";
 		}
 
-		/**
-		 *	Return the platform toolset string to write into the project configuration
-		 *	
-		 *	@param	InPlatform			The UnrealTargetPlatform being built
-		 *	@param	InConfiguration		The UnrealTargetConfiguration being built
-		 *	
-		 *	@return	string				The custom configuration section for the project file; Empty string if it doesn't require one
-		 */
+		/// <summary>
+		/// Return the platform toolset string to write into the project configuration
+		/// </summary>
+		/// <param name="InPlatform">  The UnrealTargetPlatform being built</param>
+		/// <param name="InConfiguration"> The UnrealTargetConfiguration being built</param>
+		/// <returns>string    The custom configuration section for the project file; Empty string if it doesn't require one</returns>
 		public override string GetVisualStudioPlatformToolsetString(UnrealTargetPlatform InPlatform, UnrealTargetConfiguration InConfiguration, VCProjectFile InProjectFile)
 		{
 			if (!IsNsightInstalled())
@@ -189,15 +179,13 @@ namespace UnrealBuildTool
 				+ "\t\t<AndroidNativeAPI>UseTarget</AndroidNativeAPI>" + ProjectFileGenerator.NewLine;
 		}
 
-		/**
-		 * Return any custom paths for VisualStudio this platform requires
-		 * This include ReferencePath, LibraryPath, LibraryWPath, IncludePath and ExecutablePath.
-		 *	
-		 *	@param	InPlatform			The UnrealTargetPlatform being built
-		 *	@param	TargetType			The type of target (game or program)
-		 *	
-		 *	@return	string				The custom path lines for the project file; Empty string if it doesn't require one
-		 */
+		/// <summary>
+		/// Return any custom paths for VisualStudio this platform requires
+		/// This include ReferencePath, LibraryPath, LibraryWPath, IncludePath and ExecutablePath.
+		/// </summary>
+		/// <param name="InPlatform">  The UnrealTargetPlatform being built</param>
+		/// <param name="TargetType">  The type of target (game or program)</param>
+		/// <returns>string    The custom path lines for the project file; Empty string if it doesn't require one</returns>
 		public override string GetVisualStudioPathsEntries(UnrealTargetPlatform InPlatform, UnrealTargetConfiguration InConfiguration, TargetRules.TargetType TargetType, FileReference TargetRulesPath, FileReference ProjectFilePath, FileReference NMakeOutputPath)
 		{
 			if (!IsNsightInstalled())
diff --git a/Engine/Source/Programs/UnrealBuildTool/Android/UEDeployAndroid.cs b/Engine/Source/Programs/UnrealBuildTool/Android/UEDeployAndroid.cs
--- a/Engine/Source/Programs/UnrealBuildTool/Android/UEDeployAndroid.cs
+++ b/Engine/Source/Programs/UnrealBuildTool/Android/UEDeployAndroid.cs
@@ -13,9 +13,9 @@ namespace UnrealBuildTool.Android
 {
 	public class UEDeployAndroid : UEBuildDeploy
 	{
-		/**
-		 *	Register the platform with the UEBuildDeploy class
-		 */
+		/// <summary>
+		/// Register the platform with the UEBuildDeploy class
+		/// </summary>
 		public override void RegisterBuildDeploy()
 		{
 			string NDKPath = Environment.GetEnvironmentVariable("ANDROID_HOME");
@@ -31,10 +31,14 @@ namespace UnrealBuildTool.Android
 			}
 		}
 
-		/** Internal usage for GetApiLevel */
+		/// <summary>
+		/// Internal usage for GetApiLevel
+		/// </summary>
 		private static List<string> PossibleApiLevels = null;
 
-		/** Simple function to pipe output asynchronously */
+		/// <summary>
+		/// Simple function to pipe output asynchronously
+		/// </summary>
 		private static void ParseApiLevel(object Sender, DataReceivedEventArgs Event)
 		{
 			// DataReceivedEventHandler is fired with a null string when the output stream is closed.  We don't want to
diff --git a/Engine/Source/Programs/UnrealBuildTool/Configuration/UEBuildBinary.cs b/Engine/Source/Programs/UnrealBuildTool/Configuration/UEBuildBinary.cs
--- a/Engine/Source/Programs/UnrealBuildTool/Configuration/UEBuildBinary.cs
+++ b/Engine/Source/Programs/UnrealBuildTool/Configuration/UEBuildBinary.cs
@@ -352,10 +352,10 @@ namespace UnrealBuildTool
 			return FileReference.Combine(BinaryPath.Directory, BinaryPath.GetFileNameWithoutExtension() + "-Cmd" + BinaryPath.GetExtension());
 		}
 
-		/**
-		 * Checks whether the binary output paths are appropriate for the distribution
-		 * level of its direct module dependencies
-		 */
+		/// <summary>
+		/// Checks whether the binary output paths are appropriate for the distribution
+		/// level of its direct module dependencies
+		/// </summary>
 		public void CheckOutputDistributionLevelAgainstDependencies(Dictionary<UEBuildModule, UEBuildModuleDistribution> ModuleDistributionCache)
 		{
 			// Find maximum distribution level of its direct dependencies
diff --git a/Engine/Source/Programs/UnrealBuildTool/Configuration/UEBuildConfiguration.cs b/Engine/Source/Programs/UnrealBuildTool/Configuration/UEBuildConfiguration.cs
--- a/Engine/Source/Programs/UnrealBuildTool/Configuration/UEBuildConfiguration.cs
+++ b/Engine/Source/Programs/UnrealBuildTool/Configuration/UEBuildConfiguration.cs
@@ -15,182 +15,272 @@ namespace UnrealBuildTool
 			}
 		}
 
-		/** Whether to include PhysX support */
+		/// <summary>
+		/// Whether to include PhysX support
+		/// </summary>
 		[XmlConfig]
 		public static bool bCompilePhysX;
 
-		/** Whether to include PhysX APEX support */
+		/// <summary>
+		/// Whether to include PhysX APEX support
+		/// </summary>
 		[XmlConfig]
 		public static bool bCompileAPEX;
 
-		/** Whether to allow runtime cooking of physics */
+		/// <summary>
+		/// Whether to allow runtime cooking of physics
+		/// </summary>
 		[XmlConfig]
 		public static bool bRuntimePhysicsCooking;
 
-		/** Whether to include Box2D support */
+		/// <summary>
+		/// Whether to include Box2D support
+		/// </summary>
 		[XmlConfig]
 		public static bool bCompileBox2D;
 
-		/** Whether to include ICU unicode/i18n support in core */
+		/// <summary>
+		/// Whether to include ICU unicode/i18n support in core
+		/// </summary>
 		[XmlConfig]
 		public static bool bCompileICU;
 
-		/** Whether to build a stripped down version of the game specifically for dedicated server. */
+		/// <summary>
+		/// Whether to build a stripped down version of the game specifically for dedicated server.
+		/// </summary>
 		[Obsolete("bBuildDedicatedServer has been deprecated and will be removed in future release. Update your code to use TargetInfo.Type instead or your code will not compile.")]
 		public static bool bBuildDedicatedServer;
 
-		/** Whether to compile the editor or not. Only desktop platforms (Windows or Mac) will use this, other platforms force this to false */
+		/// <summary>
+		/// Whether to compile the editor or not. Only desktop platforms (Windows or Mac) will use this, other platforms force this to false
+		/// </summary>
 		[XmlConfig]
 		public static bool bBuildEditor;
 
-		/** Whether to compile code related to building assets. Consoles generally cannot build assets. Desktop platforms generally can. */
+		/// <summary>
+		/// Whether to compile code related to building assets. Consoles generally cannot build assets. Desktop platforms generally can.
+		/// </summary>
 		[XmlConfig]
 		public static bool bBuildRequiresCookedData;
 
-		/** Whether to compile WITH_EDITORONLY_DATA disabled. Only Windows will use this, other platforms force this to false */
+		/// <summary>
+		/// Whether to compile WITH_EDITORONLY_DATA disabled. Only Windows will use this, other platforms force this to false
+		/// </summary>
 		[XmlConfig]
 		public static bool bBuildWithEditorOnlyData;
 
-		/** Whether to compile the developer tools. */
+		/// <summary>
+		/// Whether to compile the developer tools.
+		/// </summary>
 		[XmlConfig]
 		public static bool bBuildDeveloperTools;
 
-		/** Whether to force compiling the target platform modules, even if they wouldn't normally be built */
+		/// <summary>
+		/// Whether to force compiling the target platform modules, even if they wouldn't normally be built
+		/// </summary>
 		[XmlConfig]
 		public static bool bForceBuildTargetPlatforms;
 
-		/** Whether to force compiling shader format modules, even if they wouldn't normally be built. */
+		/// <summary>
+		/// Whether to force compiling shader format modules, even if they wouldn't normally be built.
+		/// </summary>
 		[XmlConfig]
 		public static bool bForceBuildShaderFormats;
 
-		/** Whether we should compile in support for Simplygon or not. */
+		/// <summary>
+		/// Whether we should compile in support for Simplygon or not.
+		/// </summary>
 		[XmlConfig]
 		public static bool bCompileSimplygon;
 
-		/** Whether we should compile in support for Steam OnlineSubsystem or not. [RCL] FIXME 2014-Apr-17: bCompileSteamOSS means "bHasSteamworksInstalled" for some code, these meanings need to be untangled */
+		/// <summary>
+		/// Whether we should compile in support for Steam OnlineSubsystem or not. [RCL] FIXME 2014-Apr-17: bCompileSteamOSS means "bHasSteamworksInstalled" for some code, these meanings need to be untangled
+		/// </summary>
 		[XmlConfig]
 		public static bool bCompileSteamOSS;
 
-		/** Whether we should compile in support for Mcp OnlineSubsystem or not. */
+		/// <summary>
+		/// Whether we should compile in support for Mcp OnlineSubsystem or not.
+		/// </summary>
 		[XmlConfig]
 		public static bool bCompileMcpOSS;
 
-		/** Whether to compile lean and mean version of UE. */
+		/// <summary>
+		/// Whether to compile lean and mean version of UE.
+		/// </summary>
 		[XmlConfig]
 		public static bool bCompileLeanAndMeanUE;
 
-		/** Whether to generate a list of external files that are required to build a target */
+		/// <summary>
+		/// Whether to generate a list of external files that are required to build a target
+		/// </summary>
 		[XmlConfig]
 		public static bool bGenerateExternalFileList;
 
-		/** Whether to merge to the existing list of external files */
+		/// <summary>
+		/// Whether to merge to the existing list of external files
+		/// </summary>
 		[XmlConfig]
 		public static bool bMergeExternalFileList;
 
-		/** Whether to generate a manifest file that contains the files to add to Perforce */
+		/// <summary>
+		/// Whether to generate a manifest file that contains the files to add to Perforce
+		/// </summary>
 		[XmlConfig]
 		public static bool bGenerateManifest;
 
-		/** Whether to add to the existing manifest (if it exists), or start afresh */
+		/// <summary>
+		/// Whether to add to the existing manifest (if it exists), or start afresh
+		/// </summary>
 		[XmlConfig]
 		public static bool bMergeManifests;
 
-		/** Whether to 'clean' the given project */
+		/// <summary>
+		/// Whether to 'clean' the given project
+		/// </summary>
 		[XmlConfig]
 		public static bool bCleanProject;
 
-		/** Whether we are just running the PrepTargetForDeployment step */
+		/// <summary>
+		/// Whether we are just running the PrepTargetForDeployment step
+		/// </summary>
 		[XmlConfig]
 		public static bool bPrepForDeployment;
 
-		/** Enabled for all builds that include the engine project.  Disabled only when building standalone apps that only link with Core. */
+		/// <summary>
+		/// Enabled for all builds that include the engine project.  Disabled only when building standalone apps that only link with Core.
+		/// </summary>
 		[XmlConfig]
 		public static bool bCompileAgainstEngine;
 
-		/** Enabled for all builds that include the CoreUObject project.  Disabled only when building standalone apps that only link with Core. */
+		/// <summary>
+		/// Enabled for all builds that include the CoreUObject project.  Disabled only when building standalone apps that only link with Core.
+		/// </summary>
 		[XmlConfig]
 		public static bool bCompileAgainstCoreUObject;
 
-		/** If true, include ADO database support in core */
+		/// <summary>
+		/// If true, include ADO database support in core
+		/// </summary>
 		[XmlConfig]
 		public static bool bIncludeADO;
 
-		/** Directory for the third party files/libs */
+		/// <summary>
+		/// Directory for the third party files/libs
+		/// </summary>
 		[Obsolete("Use UEThirdPartySourceDirectory instead of UEThirdPartyDirectory.", true)]
 		[XmlConfig]
 		public static string UEThirdPartyDirectory;
 
-		/** Directory for the third party source */
+		/// <summary>
+		/// Directory for the third party source
+		/// </summary>
 		[XmlConfig]
 		public static string UEThirdPartySourceDirectory;
 
-		/** Directory for the third party binaries */
+		/// <summary>
+		/// Directory for the third party binaries
+		/// </summary>
 		[XmlConfig]
 		public static string UEThirdPartyBinariesDirectory;
 
-		/** If true, force header regeneration. Intended for the build machine */
+		/// <summary>
+		/// If true, force header regeneration. Intended for the build machine
+		/// </summary>
 		[XmlConfig]
 		public static bool bForceHeaderGeneration;
 
-		/** If true, do not build UHT, assume it is already built */
+		/// <summary>
+		/// If true, do not build UHT, assume it is already built
+		/// </summary>
 		[XmlConfig]
 		public static bool bDoNotBuildUHT;
 
-		/** If true, fail if any of the generated header files is out of date. */
+		/// <summary>
+		/// If true, fail if any of the generated header files is out of date.
+		/// </summary>
 		[XmlConfig]
 		public static bool bFailIfGeneratedCodeChanges;
 
-		/** Whether to compile Recast navmesh generation */
+		/// <summary>
+		/// Whether to compile Recast navmesh generation
+		/// </summary>
 		[XmlConfig]
 		public static bool bCompileRecast;
 
-		/** Whether to compile SpeedTree support. */
+		/// <summary>
+		/// Whether to compile SpeedTree support.
+		/// </summary>
 		[XmlConfig]
 		public static bool bCompileSpeedTree;
 
-		/** Enable exceptions for all modules */
+		/// <summary>
+		/// Enable exceptions for all modules
+		/// </summary>
 		[XmlConfig]
 		public static bool bForceEnableExceptions;
 
-		/** Compile server-only code. */
+		/// <summary>
+		/// Compile server-only code.
+		/// </summary>
 		[XmlConfig]
 		public static bool bWithServerCode;
 
-		/** Whether to include stats support even without the engine */
+		/// <summary>
+		/// Whether to include stats support even without the engine
+		/// </summary>
 		[XmlConfig]
 		public static bool bCompileWithStatsWithoutEngine;
 
-		/** Whether to include plugin support */
+		/// <summary>
+		/// Whether to include plugin support
+		/// </summary>
 		[XmlConfig]
 		public static bool bCompileWithPluginSupport;
 
-		/** Whether to turn on logging for test/shipping builds */
+		/// <summary>
+		/// Whether to turn on logging for test/shipping builds
+		/// </summary>
 		[XmlConfig]
 		public static bool bUseLoggingInShipping;
 
-		/** True if we need PhysX vehicle support */
+		/// <summary>
+		/// True if we need PhysX vehicle support
+		/// </summary>
 		[XmlConfig]
 		public static bool bCompilePhysXVehicle;
 
-		/** True if we need FreeType support */
+		/// <summary>
+		/// True if we need FreeType support
+		/// </summary>
 		[XmlConfig]
 		public static bool bCompileFreeType;
 
-		/** True if we want to favor optimizing size over speed */
+		/// <summary>
+		/// True if we want to favor optimizing size over speed
+		/// </summary>
 		[XmlConfig]
 		public static bool bCompileForSize;
 
-		/** True if hot-reload from IDE is allowed */
+		/// <summary>
+		/// True if hot-reload from IDE is allowed
+		/// </summary>
 		[XmlConfig]
 		public static bool bAllowHotReloadFromIDE;
 
-		/** True if performing hot-reload from IDE */
+		/// <summary>
+		/// True if performing hot-reload from IDE
+		/// </summary>
 		public static bool bHotReloadFromIDE;
 
-		/** When true, the targets won't execute their link actions if there was nothing to compile */
+		/// <summary>
+		/// When true, the targets won't execute their link actions if there was nothing to compile
+		/// </summary>
 		public static bool bSkipLinkingWhenNothingToCompile;
 
-		/** Whether to compile CEF3 support. */
+		/// <summary>
+		/// Whether to compile CEF3 support.
+		/// </summary>
 		[XmlConfig]
 		public static bool bCompileCEF3;
 
@@ -257,11 +347,10 @@ namespace UnrealBuildTool
 				&& !(ProjectFileGenerator.bGenerateProjectFiles && ProjectFileGenerator.bGeneratingRocketProjectFiles);
 		}
 
-		/**
-		 * Validates the configuration.
-		 * 
-		 * @warning: the order of validation is important
-		 */
+		/// <summary>
+		/// Validates the configuration.
+		/// Warning: the order of validation is important
+		/// </summary>
 		public static void ValidateConfiguration()
 		{
 			// Lean and mean means no Editor and other frills.
diff --git a/Engine/Source/Programs/UnrealBuildTool/Configuration/UEBuildDeploy.cs b/Engine/Source/Programs/UnrealBuildTool/Configuration/UEBuildDeploy.cs
--- a/Engine/Source/Programs/UnrealBuildTool/Configuration/UEBuildDeploy.cs
+++ b/Engine/Source/Programs/UnrealBuildTool/Configuration/UEBuildDeploy.cs
@@ -10,23 +10,21 @@ namespace UnrealBuildTool
 {
 	public interface IUEBuildDeploy
 	{
-		/**
-		 * Register the platform with the UEBuildDeploy class.
-		 */
+		/// <summary>
+		/// Register the platform with the UEBuildDeploy class.
+		/// </summary>
 		void RegisterBuildDeploy();
 
-		/**
-		 * Prepare the target for deployment.
-		 *	
-		 * @param InTarget The target for deployment.
-		 *	
-		 * @return bool True if successful, false if not.
-		 */
+		/// <summary>
+		/// Prepare the target for deployment.
+		/// </summary>
+		/// <param name="InTarget">The target for deployment.</param>
+		/// <returns>bool True if successful, false if not.</returns>
 		bool PrepTargetForDeployment(UEBuildTarget InTarget);
 
-		/**
-		 * Prepare the target for deployment.
-		 */
+		/// <summary>
+		/// Prepare the target for deployment.
+		/// </summary>
 		bool PrepForUATPackageOrDeploy(string ProjectName, string ProjectDirectory, string ExecutablePath, string EngineDirectory, bool bForDistribution, string CookFlavor, bool bIsDataDeploy);
 	}
 
@@ -37,12 +35,11 @@ namespace UnrealBuildTool
 	{
 		static Dictionary<UnrealTargetPlatform, IUEBuildDeploy> BuildDeployDictionary = new Dictionary<UnrealTargetPlatform, IUEBuildDeploy>();
 
-		/**
-		 *	Register the given platforms UEBuildDeploy instance
-		 *	
-		 *	@param	InPlatform			The UnrealTargetPlatform to register with
-		 *	@param	InBuildDeploy		The UEBuildDeploy instance to use for the InPlatform
-		 */
+		/// <summary>
+		/// Register the given platforms UEBuildDeploy instance
+		/// </summary>
+		/// <param name="InPlatform">  The UnrealTargetPlatform to register with</param>
+		/// <param name="InBuildDeploy"> The UEBuildDeploy instance to use for the InPlatform</param>
 		public static void RegisterBuildDeploy(UnrealTargetPlatform InPlatform, IUEBuildDeploy InBuildDeploy)
 		{
 			if (BuildDeployDictionary.ContainsKey(InPlatform) == true)
@@ -57,13 +54,11 @@ namespace UnrealBuildTool
 			}
 		}
 
-		/**
-		 *	Retrieve the UEBuildDeploy instance for the given TargetPlatform
-		 *	
-		 *	@param	InPlatform			The UnrealTargetPlatform being built
-		 *	
-		 *	@return	UEBuildDeploy		The instance of the build deploy
-		 */
+		/// <summary>
+		/// Retrieve the UEBuildDeploy instance for the given TargetPlatform
+		/// </summary>
+		/// <param name="InPlatform">  The UnrealTargetPlatform being built</param>
+		/// <returns>UEBuildDeploy  The instance of the build deploy</returns>
 		public static IUEBuildDeploy GetBuildDeploy(UnrealTargetPlatform InPlatform)
 		{
 			if (BuildDeployDictionary.ContainsKey(InPlatform) == true)
@@ -74,26 +69,24 @@ namespace UnrealBuildTool
 			return null;
 		}
 
-		/**
-		 *	Register the platform with the UEBuildDeploy class
-		 */
+		/// <summary>
+		/// Register the platform with the UEBuildDeploy class
+		/// </summary>
 		public abstract void RegisterBuildDeploy();
 
-		/**
-		 *	Prepare the target for deployment
-		 *	
-		 *	@param	InTarget		The target for deployment
-		 *	
-		 *	@return	bool			true if successful, false if not
-		 */
+		/// <summary>
+		/// Prepare the target for deployment
+		/// </summary>
+		/// <param name="InTarget"> The target for deployment</param>
+		/// <returns>bool   true if successful, false if not</returns>
 		public virtual bool PrepTargetForDeployment(UEBuildTarget InTarget)
 		{
 			return true;
 		}
 
-		/**
-		 * Prepare the target for deployment
-		 */
+		/// <summary>
+		/// Prepare the target for deployment
+		/// </summary>
 		public virtual bool PrepForUATPackageOrDeploy(string ProjectName, string ProjectDirectory, string ExecutablePath, string EngineDirectory, bool bForDistribution, string CookFlavor, bool bIsDataDeploy)
 		{
 			return true;
diff --git a/Engine/Source/Programs/UnrealBuildTool/Configuration/UEBuildModule.cs b/Engine/Source/Programs/UnrealBuildTool/Configuration/UEBuildModule.cs
--- a/Engine/Source/Programs/UnrealBuildTool/Configuration/UEBuildModule.cs
+++ b/Engine/Source/Programs/UnrealBuildTool/Configuration/UEBuildModule.cs
@@ -12,7 +12,9 @@ using Tools.DotNETCommon.CaselessDictionary;
 
 namespace UnrealBuildTool
 {
-	/** Type of module. Mirrored in UHT as EBuildModuleType */
+	/// <summary>
+	/// Type of module. Mirrored in UHT as EBuildModuleType
+	/// </summary>
 	public enum UEBuildModuleType
 	{
 		Unknown,
@@ -30,10 +32,10 @@ namespace UnrealBuildTool
 			return ModuleType != UEBuildModuleType.Game;
 		}
 	}
-	/**
-	 * Distribution level of module.
-	 * Note: The name of each entry is used to search for/create folders
-	 */
+	/// <summary>
+	/// Distribution level of module.
+	/// Note: The name of each entry is used to search for/create folders
+	/// </summary>
 	public enum UEBuildModuleDistribution
 	{
 		/// Binaries can be distributed to everyone
@@ -50,10 +52,14 @@ namespace UnrealBuildTool
 	}
 
 
-	/** A unit of code compilation and linking. */
+	/// <summary>
+	/// A unit of code compilation and linking.
+	/// </summary>
 	public abstract class UEBuildModule
 	{
-		/** Known module folders */
+		/// <summary>
+		/// Known module folders
+		/// </summary>
 		public static readonly string RuntimeFolder = String.Format("{0}Runtime{0}", Path.DirectorySeparatorChar);
 		public static readonly string DeveloperFolder = String.Format("{0}Developer{0}", Path.DirectorySeparatorChar);
 		public static readonly string EditorFolder = String.Format("{0}Editor{0}", Path.DirectorySeparatorChar);
@@ -100,9 +106,9 @@ namespace UnrealBuildTool
 			return ModuleType;
 		}
 
-		/**
-		 * Checks what distribution level a particular path should have by checking for key folders anywhere in it
-		 */
+		/// <summary>
+		/// Checks what distribution level a particular path should have by checking for key folders anywhere in it
+		/// </summary>
 		public static UEBuildModuleDistribution GetModuleDistributionLevelBasedOnLocation(string FilePath)
 		{
 			// Get full path to ensure all folder separators are the same
@@ -133,9 +139,9 @@ namespace UnrealBuildTool
 			return UEBuildModuleDistribution.Public;
 		}
 
-		/**
-		 * Determines the distribution level of a module based on its directory and includes.
-		 */
+		/// <summary>
+		/// Determines the distribution level of a module based on its directory and includes.
+		/// </summary>
 		public UEBuildModuleDistribution DetermineDistributionLevel()
 		{
 			List<string> PathsToCheck = new List<string>();
@@ -157,38 +163,52 @@ namespace UnrealBuildTool
 			return DistributionLevel;
 		}
 
-		/** Converts an optional string list parameter to a well-defined hash set. */
+		/// <summary>
+		/// Converts an optional string list parameter to a well-defined hash set.
+		/// </summary>
 		protected static HashSet<string> HashSetFromOptionalEnumerableStringParameter(IEnumerable<string> InEnumerableStrings)
 		{
 			return InEnumerableStrings == null ? new HashSet<string>() : new HashSet<string>(InEnumerableStrings);
 		}
 
-		/** The target which owns this module. */
+		/// <summary>
+		/// The target which owns this module.
+		/// </summary>
 		public readonly UEBuildTarget Target;
 
-		/** The name that uniquely identifies the module. */
+		/// <summary>
+		/// The name that uniquely identifies the module.
+		/// </summary>
 		public readonly string Name;
 
-		/** The type of module being built. Used to switch between debug/development and precompiled/source configurations. */
+		/// <summary>
+		/// The type of module being built. Used to switch between debug/development and precompiled/source configurations.
+		/// </summary>
 		public UEBuildModuleType Type;
 
-		/** The rules for this module */
+		/// <summary>
+		/// The rules for this module
+		/// </summary>
 		public ModuleRules Rules;
 
-		/** Path to the module directory */
+		/// <summary>
+		/// Path to the module directory
+		/// </summary>
 		public readonly DirectoryReference ModuleDirectory;
 
-		/** Is this module allowed to be redistributed. */
+		/// <summary>
+		/// Is this module allowed to be redistributed.
+		/// </summary>
 		private readonly bool? IsRedistributableOverride;
 
-		/** The name of the .Build.cs file this module was created from, if any */
+		/// <summary>
+		/// The name of the .Build.cs file this module was created from, if any
+		/// </summary>
 		public FileReference RulesFile;
 
-		/**
-		 * Tells if this module can be redistributed.
-		 * 
-		 * @returns True if this module can be redistributed. False otherwise.
-		 */
+		/// <summary>
+		/// Tells if this module can be redistributed.
+		/// </summary>
 		public bool IsRedistributable()
 		{
 			return IsRedistributableOverride.HasValue
@@ -196,16 +216,24 @@ namespace UnrealBuildTool
 				: (Type != UEBuildModuleType.Developer && Type != UEBuildModuleType.Editor);
 		}
 
-		/** The binary the module will be linked into for the current target.  Only set after UEBuildBinary.BindModules is called. */
+		/// <summary>
+		/// The binary the module will be linked into for the current target.  Only set after UEBuildBinary.BindModules is called.
+		/// </summary>
 		public UEBuildBinary Binary = null;
 
-		/** Whether this module is included in the current target.  Only set after UEBuildBinary.BindModules is called. */
+		/// <summary>
+		/// Whether this module is included in the current target.  Only set after UEBuildBinary.BindModules is called.
+		/// </summary>
 		public bool bIncludedInTarget = false;
 
-		/** Include path for this module's base directory, relative to the Engine/Source directory */
+		/// <summary>
+		/// Include path for this module's base directory, relative to the Engine/Source directory
+		/// </summary>
 		protected string NormalizedModuleIncludePath;
 
-		/** The name of the _API define for this module */
+		/// <summary>
+		/// The name of the _API define for this module
+		/// </summary>
 		protected readonly string ModuleApiDefine;
 
 		protected readonly HashSet<string> PublicDefinitions;
@@ -220,28 +248,44 @@ namespace UnrealBuildTool
 		protected readonly HashSet<string> PublicAdditionalShadowFiles;
 		protected readonly HashSet<UEBuildBundleResource> PublicAdditionalBundleResources;
 
-		/** Names of modules with header files that this module's public interface needs access to. */
+		/// <summary>
+		/// Names of modules with header files that this module's public interface needs access to.
+		/// </summary>
 		protected List<UEBuildModule> PublicIncludePathModules;
 
-		/** Names of modules that this module's public interface depends on. */
+		/// <summary>
+		/// Names of modules that this module's public interface depends on.
+		/// </summary>
 		protected List<UEBuildModule> PublicDependencyModules;
 
-		/** Names of DLLs that this module should delay load */
+		/// <summary>
+		/// Names of DLLs that this module should delay load
+		/// </summary>
 		protected HashSet<string> PublicDelayLoadDLLs;
 
-		/** Names of modules with header files that this module's private implementation needs access to. */
+		/// <summary>
+		/// Names of modules with header files that this module's private implementation needs access to.
+		/// </summary>
 		protected List<UEBuildModule> PrivateIncludePathModules;
 
-		/** Names of modules that this module's private implementation depends on. */
+		/// <summary>
+		/// Names of modules that this module's private implementation depends on.
+		/// </summary>
 		protected List<UEBuildModule> PrivateDependencyModules;
 
-		/** Extra modules this module may require at run time */
+		/// <summary>
+		/// Extra modules this module may require at run time
+		/// </summary>
 		protected List<UEBuildModule> DynamicallyLoadedModules;
 
-		/** Extra modules this module may require at run time, that are on behalf of another platform (i.e. shader formats and the like) */
+		/// <summary>
+		/// Extra modules this module may require at run time, that are on behalf of another platform (i.e. shader formats and the like)
+		/// </summary>
 		protected List<UEBuildModule> PlatformSpecificDynamicallyLoadedModules;
 
-		/** Files which this module depends on at runtime. */
+		/// <summary>
+		/// Files which this module depends on at runtime.
+		/// </summary>
 		public List<RuntimeDependency> RuntimeDependencies;
 
 		public UEBuildModule(
@@ -281,13 +325,17 @@ namespace UnrealBuildTool
 			Target.RegisterModule(this);
 		}
 
-		/** Determines whether this module has a circular dependency on the given module */
+		/// <summary>
+		/// Determines whether this module has a circular dependency on the given module
+		/// </summary>
 		public bool HasCircularDependencyOn(string ModuleName)
 		{
 			return Rules.CircularlyReferencedDependentModules.Contains(ModuleName);
 		}
 
-		/** Sets up the environment for compiling any module that includes the public interface of this module. */
+		/// <summary>
+		/// Sets up the environment for compiling any module that includes the public interface of this module.
+		/// </summary>
 		protected virtual void SetupPublicCompileEnvironment(
 			UEBuildBinary SourceBinary,
 			bool bIncludePathsOnly,
@@ -397,13 +445,17 @@ namespace UnrealBuildTool
 
 		static Regex VCMacroRegex = new Regex(@"\$\([A-Za-z0-9_]+\)");
 
-		/** Checks if path contains a VC macro */
+		/// <summary>
+		/// Checks if path contains a VC macro
+		/// </summary>
 		protected bool DoesPathContainVCMacro(string Path)
 		{
 			return VCMacroRegex.IsMatch(Path);
 		}
 
-		/** Adds PathsToAdd to IncludePaths, performing path normalization and ignoring duplicates. */
+		/// <summary>
+		/// Adds PathsToAdd to IncludePaths, performing path normalization and ignoring duplicates.
+		/// </summary>
 		protected void AddIncludePathsWithChecks(HashSet<string> IncludePaths, HashSet<string> PathsToAdd)
 		{
 			if (ProjectFileGenerator.bGenerateProjectFiles)
@@ -426,7 +478,9 @@ namespace UnrealBuildTool
 			}
 		}
 
-		/** Sets up the environment for compiling this module. */
+		/// <summary>
+		/// Sets up the environment for compiling this module.
+		/// </summary>
 		protected virtual void SetupPrivateCompileEnvironment(
 			HashSet<string> IncludePaths,
 			HashSet<string> SystemIncludePaths,
@@ -462,7 +516,9 @@ namespace UnrealBuildTool
 			}
 		}
 
-		/** Sets up the environment for linking any module that includes the public interface of this module. */
+		/// <summary>
+		/// Sets up the environment for linking any module that includes the public interface of this module.
+		/// </summary>
 		protected virtual void SetupPublicLinkEnvironment(
 			UEBuildBinary SourceBinary,
 			List<string> LibraryPaths,
@@ -533,7 +589,9 @@ namespace UnrealBuildTool
 			}
 		}
 
-		/** Sets up the environment for linking this module. */
+		/// <summary>
+		/// Sets up the environment for linking this module.
+		/// </summary>
 		public virtual void SetupPrivateLinkEnvironment(
 			UEBuildBinary SourceBinary,
 			LinkEnvironment LinkEnvironment,
@@ -557,7 +615,9 @@ namespace UnrealBuildTool
 			}
 		}
 
-		/** Compiles the module, and returns a list of files output by the compiler. */
+		/// <summary>
+		/// Compiles the module, and returns a list of files output by the compiler.
+		/// </summary>
 		public abstract List<FileItem> Compile(IUEToolChain ToolChain, CPPEnvironment GlobalCompileEnvironment, CPPEnvironment CompileEnvironment);
 
 		// Object interface.
@@ -573,37 +633,35 @@ namespace UnrealBuildTool
 			public int Index;
 		}
 
-		/**
-		 * Gets all of the modules referenced by this module
-		 * 
-		 * @param	ReferencedModules	Hash of all referenced modules with their addition index.
-		 * @param	bIncludeDynamicallyLoaded	True if dynamically loaded modules (and all of their dependent modules) should be included.
-		 * @param	bForceCircular	True if circular dependencies should be processed
-		 * @param	bOnlyDirectDependencies	True to return only this module's direct dependencies
-		 */
+		/// <summary>
+		/// Gets all of the modules referenced by this module
+		/// </summary>
+		/// <param name="ReferencedModules">Hash of all referenced modules with their addition index.</param>
+		/// <param name="bIncludeDynamicallyLoaded">True if dynamically loaded modules (and all of their dependent modules) should be included.</param>
+		/// <param name="bForceCircular">True if circular dependencies should be processed</param>
+		/// <param name="bOnlyDirectDependencies">True to return only this module's direct dependencies</param>
 		public virtual void GetAllDependencyModules(CaselessDictionary<ModuleIndexPair> ReferencedModules, bool bIncludeDynamicallyLoaded, bool bForceCircular, bool bOnlyDirectDependencies)
 		{
 		}
 
-		/**
-		 * Gets all of the modules precompiled along with this module
-		 * 
-		 * @param	Modules	Set of all the precompiled modules
-		 */
+		/// <summary>
+		/// Gets all of the modules precompiled along with this module
+		/// </summary>
+		/// <param name="Modules">Set of all the precompiled modules</param>
 		public virtual void RecursivelyAddPrecompiledModules(List<UEBuildModule> Modules)
 		{
 		}
 
-		/**
-		 * Gathers and binds binaries for this module and all of it's dependent modules
-		 */
+		/// <summary>
+		/// Gathers and binds binaries for this module and all of it's dependent modules
+		/// </summary>
 		public virtual void RecursivelyProcessUnboundModules()
 		{
 		}
 
-		/**
-		 * Creates all the modules required for this target
-		 */
+		/// <summary>
+		/// Creates all the modules required for this target
+		/// </summary>
 		public void RecursivelyCreateModules()
 		{
 			// Create all the include path modules. These modules may not be added to the target (and we don't process their dependencies), but they need 
@@ -652,7 +710,9 @@ namespace UnrealBuildTool
 		}
 	};
 
-	/** A module that is never compiled by us, and is only used to group include paths and libraries into a dependency unit. */
+	/// <summary>
+	/// A module that is never compiled by us, and is only used to group include paths and libraries into a dependency unit.
+	/// </summary>
 	class UEBuildExternalModule : UEBuildModule
 	{
 		public UEBuildExternalModule(
@@ -682,14 +742,18 @@ namespace UnrealBuildTool
 		}
 	};
 
-	/** A module that is compiled from C++ code. */
+	/// <summary>
+	/// A module that is compiled from C++ code.
+	/// </summary>
 	public class UEBuildModuleCPP : UEBuildModule
 	{
 		public class AutoGenerateCppInfoClass
 		{
 			public class BuildInfoClass
 			{
-				/** The wildcard of the *.generated.cpp file which was generated for the module */
+				/// <summary>
+				/// The wildcard of the *.generated.cpp file which was generated for the module
+				/// </summary>
 				public readonly string FileWildcard;
 
 				public BuildInfoClass(string InWildcard)
@@ -700,7 +764,9 @@ namespace UnrealBuildTool
 				}
 			}
 
-			/** Information about how to build the .generated.cpp files. If this is null, then we're not building .generated.cpp files for this module. */
+			/// <summary>
+			/// Information about how to build the .generated.cpp files. If this is null, then we're not building .generated.cpp files for this module.
+			/// </summary>
 			public BuildInfoClass BuildInfo;
 
 			public AutoGenerateCppInfoClass(BuildInfoClass InBuildInfo)
@@ -709,7 +775,9 @@ namespace UnrealBuildTool
 			}
 		}
 
-		/** Information about the .generated.cpp file.  If this is null then this module doesn't have any UHT-produced code. */
+		/// <summary>
+		/// Information about the .generated.cpp file.  If this is null then this module doesn't have any UHT-produced code.
+		/// </summary>
 		public AutoGenerateCppInfoClass AutoGenerateCppInfo = null;
 
 		public class SourceFilesClass
@@ -763,26 +831,33 @@ namespace UnrealBuildTool
 			}
 		}
 
-		/**
-		 * Adds additional source cpp files for this module.
-		 *
-		 * @param Files Files to add.
-		 */
+		/// <summary>
+		/// Adds additional source cpp files for this module.
+		/// </summary>
+		/// <param name="Files">Files to add.</param>
 		public void AddAdditionalCPPFiles(IEnumerable<FileItem> Files)
 		{
 			SourceFilesToBuild.CPPFiles.AddRange(Files);
 		}
 
-		/** A list of the absolute paths of source files to be built in this module. */
+		/// <summary>
+		/// A list of the absolute paths of source files to be built in this module.
+		/// </summary>
 		public readonly SourceFilesClass SourceFilesToBuild = new SourceFilesClass();
 
-		/** A list of the source files that were found for the module. */
+		/// <summary>
+		/// A list of the source files that were found for the module.
+		/// </summary>
 		public readonly SourceFilesClass SourceFilesFound = new SourceFilesClass();
 
-		/** The directory for this module's generated code */
+		/// <summary>
+		/// The directory for this module's generated code
+		/// </summary>
 		public readonly DirectoryReference GeneratedCodeDirectory;
 
-		/** The preprocessor definitions used to compile this module's private implementation. */
+		/// <summary>
+		/// The preprocessor definitions used to compile this module's private implementation.
+		/// </summary>
 		HashSet<string> Definitions;
 
 		/// When set, allows this module to report compiler definitions and include paths for Intellisense
@@ -792,17 +867,25 @@ namespace UnrealBuildTool
 
 		public class ProcessedDependenciesClass
 		{
-			/** The file, if any, which is used as the unique PCH for this module */
+			/// <summary>
+			/// The file, if any, which is used as the unique PCH for this module
+			/// </summary>
 			public FileItem UniquePCHHeaderFile = null;
 		}
 
-		/** The processed dependencies for the class */
+		/// <summary>
+		/// The processed dependencies for the class
+		/// </summary>
 		public ProcessedDependenciesClass ProcessedDependencies = null;
 
-		/** @hack to skip adding definitions to compile environment. They will be baked into source code by external code. */
+		/// <summary>
+		/// Hack to skip adding definitions to compile environment. They will be baked into source code by external code.
+		/// </summary>
 		public bool bSkipDefinitionsForCompileEnvironment = false;
 
-		/** Categorizes source files into per-extension buckets */
+		/// <summary>
+		/// Categorizes source files into per-extension buckets
+		/// </summary>
 		private static void CategorizeSourceFiles(IEnumerable<FileItem> InSourceFiles, SourceFilesClass OutSourceFiles)
 		{
 			foreach (var SourceFile in InSourceFiles)
@@ -1755,10 +1838,14 @@ namespace UnrealBuildTool
 		}
 	}
 
-	/** A module that is compiled from C++ CLR code. */
+	/// <summary>
+	/// A module that is compiled from C++ CLR code.
+	/// </summary>
 	class UEBuildModuleCPPCLR : UEBuildModuleCPP
 	{
-		/** The assemblies referenced by the module's private implementation. */
+		/// <summary>
+		/// The assemblies referenced by the module's private implementation.
+		/// </summary>
 		HashSet<string> PrivateAssemblyReferences;
 
 		public UEBuildModuleCPPCLR(
@@ -1854,27 +1941,41 @@ namespace UnrealBuildTool
 
 	public class PrecompileHeaderEnvironment
 	{
-		/** The name of the module this PCH header is a member of */
+		/// <summary>
+		/// The name of the module this PCH header is a member of
+		/// </summary>
 		public readonly string ModuleName;
 
-		/** PCH header file name as it appears in an #include statement in source code (might include partial, or no relative path.)
-			This is needed by some compilers to use PCH features. */
+		/// <summary>
+		/// PCH header file name as it appears in an #include statement in source code (might include partial, or no relative path.)
+		/// This is needed by some compilers to use PCH features.
+		/// </summary>
 		public string PCHHeaderNameInCode;
 
-		/** The source header file that this precompiled header will be generated for */
+		/// <summary>
+		/// The source header file that this precompiled header will be generated for
+		/// </summary>
 		public readonly FileItem PrecompiledHeaderIncludeFilename;
 
-		/** Whether this precompiled header will be built with CLR features enabled.  We won't mix and match CLR PCHs with non-CLR PCHs */
+		/// <summary>
+		/// Whether this precompiled header will be built with CLR features enabled.  We won't mix and match CLR PCHs with non-CLR PCHs
+		/// </summary>
 		public readonly CPPCLRMode CLRMode;
 
-		/** Whether this precompiled header will be built with code optimization enabled. */
+		/// <summary>
+		/// Whether this precompiled header will be built with code optimization enabled.
+		/// </summary>
 		public readonly ModuleRules.CodeOptimization OptimizeCode;
 
-		/** The PCH file we're generating */
+		/// <summary>
+		/// The PCH file we're generating
+		/// </summary>
 		public FileItem PrecompiledHeaderFile = null;
 
-		/** Object files emitted from the compiler when generating this precompiled header.  These will be linked into modules that
-			include this PCH */
+		/// <summary>
+		/// Object files emitted from the compiler when generating this precompiled header.  These will be linked into modules that
+		/// include this PCH
+		/// </summary>
 		public readonly List<FileItem> OutputObjectFiles = new List<FileItem>();
 
 		public PrecompileHeaderEnvironment(string InitModuleName, string InitPCHHeaderNameInCode, FileItem InitPrecompiledHeaderIncludeFilename, CPPCLRMode InitCLRMode, ModuleRules.CodeOptimization InitOptimizeCode)
