Convert everything in UBT to use standard C# XML comments, so IntelliSense can display useful information in tooltips.

diff --git a/Engine/Source/Programs/UnrealBuildTool/Configuration/UEBuildPlatform.cs b/Engine/Source/Programs/UnrealBuildTool/Configuration/UEBuildPlatform.cs
--- a/Engine/Source/Programs/UnrealBuildTool/Configuration/UEBuildPlatform.cs
+++ b/Engine/Source/Programs/UnrealBuildTool/Configuration/UEBuildPlatform.cs
@@ -16,238 +16,211 @@ namespace UnrealBuildTool
 
 	public interface IUEBuildPlatform
 	{
-		/**
-		 * Whether the required external SDKs are installed for this platform
-		 */
+		/// <summary>
+		/// Whether the required external SDKs are installed for this platform
+		/// </summary>
 		SDKStatus HasRequiredSDKsInstalled();
 
-		/**
-		 * If this platform can be compiled with XGE
-		 */
+		/// <summary>
+		/// If this platform can be compiled with XGE
+		/// </summary>
 		bool CanUseXGE();
 
-		/**
-		 * If this platform can be compiled with DMUCS/Distcc
-		 */
+		/// <summary>
+		/// If this platform can be compiled with DMUCS/Distcc
+		/// </summary>
 		bool CanUseDistcc();
 
-		/**
-		 * If this platform can be compiled with SN-DBS
-		 */
+		/// <summary>
+		/// If this platform can be compiled with SN-DBS
+		/// </summary>
 		bool CanUseSNDBS();
 
-		/**
-		 * Register the platform with the UEBuildPlatform class
-		 */
+		/// <summary>
+		/// Register the platform with the UEBuildPlatform class
+		/// </summary>
 		void RegisterBuildPlatform();
 
-		/**
-		 * Attempt to set up AutoSDK for this platform
-		 */
+		/// <summary>
+		/// Attempt to set up AutoSDK for this platform
+		/// </summary>
 		void ManageAndValidateSDK();
 
-		/**
-		 * Retrieve the CPPTargetPlatform for the given UnrealTargetPlatform
-		 *
-		 * @param InUnrealTargetPlatform The UnrealTargetPlatform being build
-		 * 
-		 * @return CPPTargetPlatform The CPPTargetPlatform to compile for
-		 */
+		/// <summary>
+		/// Retrieve the CPPTargetPlatform for the given UnrealTargetPlatform
+		/// </summary>
+		/// <param name="InUnrealTargetPlatform">The UnrealTargetPlatform being build</param>
+		/// <returns>CPPTargetPlatform The CPPTargetPlatform to compile for</returns>
 		CPPTargetPlatform GetCPPTargetPlatform(UnrealTargetPlatform InUnrealTargetPlatform);
 
-		/**
-		 * Get the extension to use for the given binary type
-		 * 
-		 * @param InBinaryType The binary type being built
-		 * 
-		 * @return string The binary extension (i.e. 'exe' or 'dll')
-		 */
+		/// <summary>
+		/// Get the extension to use for the given binary type
+		/// </summary>
+		/// <param name="InBinaryType">The binary type being built</param>
+		/// <returns>string The binary extension (i.e. 'exe' or 'dll')</returns>
 		string GetBinaryExtension(UEBuildBinaryType InBinaryType);
 
-		/**
-		 * Get the extension to use for debug info for the given binary type
-		 * 
-		 * @param InBinaryType The binary type being built
-		 * 
-		 * @return string The debug info extension (i.e. 'pdb')
-		 */
+		/// <summary>
+		/// Get the extension to use for debug info for the given binary type
+		/// </summary>
+		/// <param name="InBinaryType">The binary type being built</param>
+		/// <returns>string The debug info extension (i.e. 'pdb')</returns>
 		string GetDebugInfoExtension(UEBuildBinaryType InBinaryType);
 
-		/**
-		 * Whether incremental linking should be used
-		 * 
-		 * @param InPlatform The CPPTargetPlatform being built
-		 * @param InConfiguration The CPPTargetConfiguration being built
-		 * 
-		 * @return bool true if incremental linking should be used, false if not
-		 */
+		/// <summary>
+		/// Whether incremental linking should be used
+		/// </summary>
+		/// <param name="InPlatform">The CPPTargetPlatform being built</param>
+		/// <param name="InConfiguration">The CPPTargetConfiguration being built</param>
+		/// <returns>bool true if incremental linking should be used, false if not</returns>
 		bool ShouldUseIncrementalLinking(CPPTargetPlatform Platform, CPPTargetConfiguration Configuration);
 
-		/**
-		 * Whether PDB files should be used
-		 * 
-		 * @param InPlatform The CPPTargetPlatform being built
-		 * @param InConfiguration The CPPTargetConfiguration being built
-		 * @param bInCreateDebugInfo true if debug info is getting create, false if not
-		 * 
-		 * @return bool true if PDB files should be used, false if not
-		 */
+		/// <summary>
+		/// Whether PDB files should be used
+		/// </summary>
+		/// <param name="InPlatform">The CPPTargetPlatform being built</param>
+		/// <param name="InConfiguration">The CPPTargetConfiguration being built</param>
+		/// <param name="bInCreateDebugInfo">true if debug info is getting create, false if not</param>
+		/// <returns>bool true if PDB files should be used, false if not</returns>
 		bool ShouldUsePDBFiles(CPPTargetPlatform Platform, CPPTargetConfiguration Configuration, bool bCreateDebugInfo);
 
-		/**
-		 * Whether PCH files should be used
-		 * 
-		 * @param InPlatform The CPPTargetPlatform being built
-		 * @param InConfiguration The CPPTargetConfiguration being built
-		 * 
-		 * @return bool true if PCH files should be used, false if not
-		 */
+		/// <summary>
+		/// Whether PCH files should be used
+		/// </summary>
+		/// <param name="InPlatform">The CPPTargetPlatform being built</param>
+		/// <param name="InConfiguration">The CPPTargetConfiguration being built</param>
+		/// <returns>bool true if PCH files should be used, false if not</returns>
 		bool ShouldUsePCHFiles(CPPTargetPlatform Platform, CPPTargetConfiguration Configuration);
 
-		/**
-		 * Whether the editor should be built for this platform or not
-		 * 
-		 * @param InPlatform The UnrealTargetPlatform being built
-		 * @param InConfiguration The UnrealTargetConfiguration being built
-		 * @return bool true if the editor should be built, false if not
-		 */
+		/// <summary>
+		/// Whether the editor should be built for this platform or not
+		/// </summary>
+		/// <param name="InPlatform">The UnrealTargetPlatform being built</param>
+		/// <param name="InConfiguration">The UnrealTargetConfiguration being built</param>
+		/// <returns>bool true if the editor should be built, false if not</returns>
 		bool ShouldNotBuildEditor(UnrealTargetPlatform InPlatform, UnrealTargetConfiguration InConfiguration);
 
-		/**
-		 * Whether this build should support ONLY cooked data or not
-		 * 
-		 * @param InPlatform The UnrealTargetPlatform being built
-		 * @param InConfiguration The UnrealTargetConfiguration being built
-		 * @return bool true if the editor should be built, false if not
-		 */
+		/// <summary>
+		/// Whether this build should support ONLY cooked data or not
+		/// </summary>
+		/// <param name="InPlatform">The UnrealTargetPlatform being built</param>
+		/// <param name="InConfiguration">The UnrealTargetConfiguration being built</param>
+		/// <returns>bool true if the editor should be built, false if not</returns>
 		bool BuildRequiresCookedData(UnrealTargetPlatform InPlatform, UnrealTargetConfiguration InConfiguration);
 
-		/**
-		 * Whether the platform requires the extra UnityCPPWriter
-		 * This is used to add an extra file for UBT to get the #include dependencies from
-		 * 
-		 * @return bool true if it is required, false if not
-		 */
+		/// <summary>
+		/// Whether the platform requires the extra UnityCPPWriter
+		/// This is used to add an extra file for UBT to get the #include dependencies from
+		/// </summary>
+		/// <returns>bool true if it is required, false if not</returns>
 		bool RequiresExtraUnityCPPWriter();
 
-		/**
-		 * Whether this platform should build a monolithic binary
-		 */
+		/// <summary>
+		/// Whether this platform should build a monolithic binary
+		/// </summary>
 		bool ShouldCompileMonolithicBinary(UnrealTargetPlatform InPlatform);
 
-		/**
-		 * Get a list of extra modules the platform requires.
-		 * This is to allow undisclosed platforms to add modules they need without exposing information about the platform.
-		 * 
-		 * @param Target The target being build
-		 * @param BuildTarget The UEBuildTarget getting build
-		 * @param PlatformExtraModules OUTPUT the list of extra modules the platform needs to add to the target
-		 */
+		/// <summary>
+		/// Get a list of extra modules the platform requires.
+		/// This is to allow undisclosed platforms to add modules they need without exposing information about the platform.
+		/// </summary>
+		/// <param name="Target">The target being build</param>
+		/// <param name="BuildTarget">The UEBuildTarget getting build</param>
+		/// <param name="PlatformExtraModules">OUTPUT the list of extra modules the platform needs to add to the target</param>
 		void GetExtraModules(TargetInfo Target, UEBuildTarget BuildTarget, ref List<string> PlatformExtraModules);
 
-		/**
-		 * Modify the newly created module passed in for this platform.
-		 * This is not required - but allows for hiding details of a
-		 * particular platform.
-		 * 
-		 *  @param  ModuleName		The name of the module
-		 *	@param	Rules			The module rules
-		 *	@param	Target			The target being build
-		 */
+		/// <summary>
+		/// Modify the newly created module passed in for this platform.
+		/// This is not required - but allows for hiding details of a
+		/// particular platform.
+		/// </summary>
+		/// <param name="">ModuleName  The name of the module</param>
+		/// <param name="Rules">  The module rules</param>
+		/// <param name="Target">  The target being build</param>
 		void ModifyModuleRules(string ModuleName, ModuleRules Rules, TargetInfo Target);
 
-		/**
-		 * Setup the target environment for building
-		 * 
-		 * @param InBuildTarget The target being built
-		 */
+		/// <summary>
+		/// Setup the target environment for building
+		/// </summary>
+		/// <param name="InBuildTarget">The target being built</param>
 		void SetUpEnvironment(UEBuildTarget InBuildTarget);
 
-		/**
-		 * Allow the platform to set an optional architecture
-		 */
+		/// <summary>
+		/// Allow the platform to set an optional architecture
+		/// </summary>
 		string GetActiveArchitecture();
 
-		/**
-		 * Allow the platform to apply architecture-specific name according to its rules
-		 * 
-		 * @param BinaryName name of the binary, not specific to any architecture
-		 *
-		 * @return Possibly architecture-specific name
-		 */
+		/// <summary>
+		/// Allow the platform to apply architecture-specific name according to its rules
+		/// </summary>
+		/// <param name="BinaryName">name of the binary, not specific to any architecture</param>
+		/// <returns>Possibly architecture-specific name</returns>
 		string ApplyArchitectureName(string BinaryName);
 
-		/**
-		 * For platforms that need to output multiple files per binary (ie Android "fat" binaries)
-		 * this will emit multiple paths. By default, it simply makes an array from the input
-		 */
+		/// <summary>
+		/// For platforms that need to output multiple files per binary (ie Android "fat" binaries)
+		/// this will emit multiple paths. By default, it simply makes an array from the input
+		/// </summary>
 		List<FileReference> FinalizeBinaryPaths(FileReference BinaryName);
 
-		/**
-		 * Setup the configuration environment for building
-		 * 
-		 * @param InBuildTarget The target being built
-		 */
+		/// <summary>
+		/// Setup the configuration environment for building
+		/// </summary>
+		/// <param name="InBuildTarget">The target being built</param>
 		void SetUpConfigurationEnvironment(UEBuildTarget InBuildTarget);
 
-		/**
-		 * Setup the project environment for building
-		 * 
-		 * @param InBuildTarget The target being built
-		 */
+		/// <summary>
+		/// Setup the project environment for building
+		/// </summary>
+		/// <param name="InBuildTarget">The target being built</param>
 		void SetUpProjectEnvironment(UnrealTargetPlatform InPlatform);
 
-		/**
-		 * Whether this platform should create debug information or not
-		 * 
-		 * @param InPlatform The UnrealTargetPlatform being built
-		 * @param InConfiguration The UnrealTargetConfiguration being built
-		 * 
-		 * @return bool true if debug info should be generated, false if not
-		 */
+		/// <summary>
+		/// Whether this platform should create debug information or not
+		/// </summary>
+		/// <param name="InPlatform">The UnrealTargetPlatform being built</param>
+		/// <param name="InConfiguration">The UnrealTargetConfiguration being built</param>
+		/// <returns>bool true if debug info should be generated, false if not</returns>
 		bool ShouldCreateDebugInfo(UnrealTargetPlatform Platform, UnrealTargetConfiguration Configuration);
 
-		/**
-		 * Gives the platform a chance to 'override' the configuration settings 
-		 * that are overridden on calls to RunUBT.
-		 * 
-		 * @param InPlatform The UnrealTargetPlatform being built
-		 * @param InConfiguration The UnrealTargetConfiguration being built
-		 * 
-		 * @return bool true if debug info should be generated, false if not
-		 */
+		/// <summary>
+		/// Gives the platform a chance to 'override' the configuration settings
+		/// that are overridden on calls to RunUBT.
+		/// </summary>
+		/// <param name="InPlatform">The UnrealTargetPlatform being built</param>
+		/// <param name="InConfiguration">The UnrealTargetConfiguration being built</param>
+		/// <returns>bool true if debug info should be generated, false if not</returns>
 		void ResetBuildConfiguration(UnrealTargetPlatform InPlatform, UnrealTargetConfiguration InConfiguration);
 
-		/**
-		 * Validate the UEBuildConfiguration for this platform
-		 * This is called BEFORE calling UEBuildConfiguration to allow setting 
-		 * various fields used in that function such as CompileLeanAndMean...
-		 */
+		/// <summary>
+		/// Validate the UEBuildConfiguration for this platform
+		/// This is called BEFORE calling UEBuildConfiguration to allow setting
+		/// various fields used in that function such as CompileLeanAndMean...
+		/// </summary>
 		void ValidateUEBuildConfiguration();
 
-		/**
-		 * Validate configuration for this platform
-		 * NOTE: This function can/will modify BuildConfiguration!
-		 * 
-		 * @param InPlatform The CPPTargetPlatform being built
-		 * @param InConfiguration The CPPTargetConfiguration being built
-		 * @param bInCreateDebugInfo true if debug info is getting create, false if not
-		 */
+		/// <summary>
+		/// Validate configuration for this platform
+		/// NOTE: This function can/will modify BuildConfiguration!
+		/// </summary>
+		/// <param name="InPlatform">The CPPTargetPlatform being built</param>
+		/// <param name="InConfiguration">The CPPTargetConfiguration being built</param>
+		/// <param name="bInCreateDebugInfo">true if debug info is getting create, false if not</param>
 		void ValidateBuildConfiguration(CPPTargetConfiguration Configuration, CPPTargetPlatform Platform, bool bCreateDebugInfo);
 
-		/**
-		 * Return whether this platform has uniquely named binaries across multiple games
-		 */
+		/// <summary>
+		/// Return whether this platform has uniquely named binaries across multiple games
+		/// </summary>
 		bool HasUniqueBinaries();
 
-		/**
-		 * Return whether we wish to have this platform's binaries in our builds
-		 */
+		/// <summary>
+		/// Return whether we wish to have this platform's binaries in our builds
+		/// </summary>
 		bool IsBuildRequired();
 
-		/**
-		 * Return whether we wish to have this platform's binaries in our CIS tests
-		 */
+		/// <summary>
+		/// Return whether we wish to have this platform's binaries in our CIS tests
+		/// </summary>
 		bool IsCISRequired();
 
 		/// <summary>
@@ -256,25 +229,22 @@ namespace UnrealBuildTool
 		/// <returns></returns>
 		bool RequiresDeployPrepAfterCompile();
 
-		/**
-		 * Return all valid configurations for this platform
-		 * 
-		 *  Typically, this is always Debug, Development, and Shipping - but Test is a likely future addition for some platforms
-		 */
+		/// <summary>
+		/// Return all valid configurations for this platform
+		/// Typically, this is always Debug, Development, and Shipping - but Test is a likely future addition for some platforms
+		/// </summary>
 		List<UnrealTargetConfiguration> GetConfigurations(UnrealTargetPlatform InUnrealTargetPlatform, bool bIncludeDebug);
 
-		/**
-		 * Setup the binaries for this specific platform.
-		 * 
-		 * @param InBuildTarget The target being built
-		 */
+		/// <summary>
+		/// Setup the binaries for this specific platform.
+		/// </summary>
+		/// <param name="InBuildTarget">The target being built</param>
 		void SetupBinaries(UEBuildTarget InBuildTarget);
 
-		/**
-		 * Check for the default configuration
-		 *
-		 * return true if the project uses the default build config
-		 */
+		/// <summary>
+		/// Check for the default configuration
+		/// return true if the project uses the default build config
+		/// </summary>
 		bool HasDefaultBuildConfig(UnrealTargetPlatform InPlatform, DirectoryReference InProjectPath);
 	}
 
@@ -287,11 +257,10 @@ namespace UnrealBuildTool
 
 		private bool bInitializedProject = false;
 
-		/**
-		 * Attempt to convert a string to an UnrealTargetPlatform enum entry
-		 * 
-		 * @return UnrealTargetPlatform.Unknown on failure (the platform didn't match the enum)
-		 */
+		/// <summary>
+		/// Attempt to convert a string to an UnrealTargetPlatform enum entry
+		/// </summary>
+		/// <returns>UnrealTargetPlatform.Unknown on failure (the platform didn't match the enum)</returns>
 		public static UnrealTargetPlatform ConvertStringToPlatform(string InPlatformName)
 		{
 			// special case x64, to not break anything
@@ -314,12 +283,11 @@ namespace UnrealBuildTool
 			return UnrealTargetPlatform.Unknown;
 		}
 
-		/**
-		 *	Register the given platforms UEBuildPlatform instance
-		 *	
-		 *	@param	InPlatform			The UnrealTargetPlatform to register with
-		 *	@param	InBuildPlatform		The UEBuildPlatform instance to use for the InPlatform
-		 */
+		/// <summary>
+		/// Register the given platforms UEBuildPlatform instance
+		/// </summary>
+		/// <param name="InPlatform">  The UnrealTargetPlatform to register with</param>
+		/// <param name="InBuildPlatform"> The UEBuildPlatform instance to use for the InPlatform</param>
 		public static void RegisterBuildPlatform(UnrealTargetPlatform InPlatform, UEBuildPlatform InBuildPlatform)
 		{
 			if (BuildPlatformDictionary.ContainsKey(InPlatform) == true)
@@ -334,9 +302,9 @@ namespace UnrealBuildTool
 			}
 		}
 
-		/**
-		 *	Unregister the given platform
-		 */
+		/// <summary>
+		/// Unregister the given platform
+		/// </summary>
 		public static void UnregisterBuildPlatform(UnrealTargetPlatform InPlatform)
 		{
 			if (BuildPlatformDictionary.ContainsKey(InPlatform) == true)
@@ -345,18 +313,18 @@ namespace UnrealBuildTool
 			}
 		}
 
-		/**
-		 * Assign a platform as a member of the given group
-		 */
+		/// <summary>
+		/// Assign a platform as a member of the given group
+		/// </summary>
 		public static void RegisterPlatformWithGroup(UnrealTargetPlatform InPlatform, UnrealPlatformGroup InGroup)
 		{
 			// find or add the list of groups for this platform
 			PlatformGroupDictionary.GetOrAddNew(InGroup).Add(InPlatform);
 		}
 
-		/**
-		 * Retrieve the list of platforms in this group (if any)
-		 */
+		/// <summary>
+		/// Retrieve the list of platforms in this group (if any)
+		/// </summary>
 		public static List<UnrealTargetPlatform> GetPlatformsInGroup(UnrealPlatformGroup InGroup)
 		{
 			List<UnrealTargetPlatform> PlatformList;
@@ -364,14 +332,12 @@ namespace UnrealBuildTool
 			return PlatformList;
 		}
 
-		/**
-		 *	Retrieve the IUEBuildPlatform instance for the given TargetPlatform
-		 *	
-		 *	@param	InPlatform			The UnrealTargetPlatform being built
-		 *	@param	bInAllowFailure		If true, do not throw an exception and return null
-		 *	
-		 *	@return	UEBuildPlatform		The instance of the build platform
-		 */
+		/// <summary>
+		/// Retrieve the IUEBuildPlatform instance for the given TargetPlatform
+		/// </summary>
+		/// <param name="InPlatform">  The UnrealTargetPlatform being built</param>
+		/// <param name="bInAllowFailure"> If true, do not throw an exception and return null</param>
+		/// <returns>UEBuildPlatform  The instance of the build platform</returns>
 		public static IUEBuildPlatform GetBuildPlatform(UnrealTargetPlatform InPlatform, bool bInAllowFailure = false)
 		{
 			if (BuildPlatformDictionary.ContainsKey(InPlatform) == true)
@@ -385,14 +351,12 @@ namespace UnrealBuildTool
 			throw new BuildException("GetBuildPlatform: No BuildPlatform found for {0}", InPlatform.ToString());
 		}
 
-		/**
-		 *	Retrieve the IUEBuildPlatform instance for the given CPPTargetPlatform
-		 *	
-		 *	@param	InPlatform			The CPPTargetPlatform being built
-		 *	@param	bInAllowFailure		If true, do not throw an exception and return null
-		 *	
-		 *	@return	UEBuildPlatform		The instance of the build platform
-		 */
+		/// <summary>
+		/// Retrieve the IUEBuildPlatform instance for the given CPPTargetPlatform
+		/// </summary>
+		/// <param name="InPlatform">  The CPPTargetPlatform being built</param>
+		/// <param name="bInAllowFailure"> If true, do not throw an exception and return null</param>
+		/// <returns>UEBuildPlatform  The instance of the build platform</returns>
 		public static IUEBuildPlatform GetBuildPlatformForCPPTargetPlatform(CPPTargetPlatform InPlatform, bool bInAllowFailure = false)
 		{
 			UnrealTargetPlatform UTPlatform = UEBuildTarget.CPPTargetPlatformToUnrealTargetPlatform(InPlatform);
@@ -407,16 +371,15 @@ namespace UnrealBuildTool
 			throw new BuildException("UEBuildPlatform::GetBuildPlatformForCPPTargetPlatform: No BuildPlatform found for {0}", InPlatform.ToString());
 		}
 
-		/**
-		 *	Allow all registered build platforms to modify the newly created module 
-		 *	passed in for the given platform.
-		 *	This is not required - but allows for hiding details of a particular platform.
-		 *	
-		 *  @param  Name			The name of the module
-		 *	@param	Module			The module rules
-		 *	@param	Target			The target being build
-		 *	@param	Only			If this is not unknown, then only run that platform
-		 */
+		/// <summary>
+		/// Allow all registered build platforms to modify the newly created module
+		/// passed in for the given platform.
+		/// This is not required - but allows for hiding details of a particular platform.
+		/// </summary>
+		/// <param name="">Name   The name of the module</param>
+		/// <param name="Module">  The module rules</param>
+		/// <param name="Target">  The target being build</param>
+		/// <param name="Only">  If this is not unknown, then only run that platform</param>
 		public static void PlatformModifyModuleRules(string ModuleName, ModuleRules Rules, TargetInfo Target, UnrealTargetPlatform Only = UnrealTargetPlatform.Unknown)
 		{
 			foreach (var PlatformEntry in BuildPlatformDictionary)
@@ -428,9 +391,9 @@ namespace UnrealBuildTool
 			}
 		}
 
-		/**
-		 * Returns the delimiter used to separate paths in the PATH environment variable for the platform we are executing on.
-		 */
+		/// <summary>
+		/// Returns the delimiter used to separate paths in the PATH environment variable for the platform we are executing on.
+		/// </summary>
 		public String GetPathVarDelimiter()
 		{
 			switch (BuildHostPlatform.Current.Platform)
@@ -450,33 +413,33 @@ namespace UnrealBuildTool
 
 
 
-		/**
-		 *	If this platform can be compiled with XGE
-		 */
+		/// <summary>
+		/// If this platform can be compiled with XGE
+		/// </summary>
 		public virtual bool CanUseXGE()
 		{
 			return true;
 		}
 
-		/**
-		 *	If this platform can be compiled with DMUCS/Distcc
-		 */
+		/// <summary>
+		/// If this platform can be compiled with DMUCS/Distcc
+		/// </summary>
 		public virtual bool CanUseDistcc()
 		{
 			return false;
 		}
 
-		/**
-		 *	If this platform can be compiled with SN-DBS
-		 */
+		/// <summary>
+		/// If this platform can be compiled with SN-DBS
+		/// </summary>
 		public virtual bool CanUseSNDBS()
 		{
 			return false;
 		}
 
-		/**
-		 *	Register the platform with the UEBuildPlatform class
-		 */
+		/// <summary>
+		/// Register the platform with the UEBuildPlatform class
+		/// </summary>
 		public void RegisterBuildPlatform()
 		{
 			ManageAndValidateSDK();
@@ -485,13 +448,11 @@ namespace UnrealBuildTool
 
 		protected abstract void RegisterBuildPlatformInternal();
 
-		/**
-		 *	Retrieve the CPPTargetPlatform for the given UnrealTargetPlatform
-		 *
-		 *	@param	InUnrealTargetPlatform		The UnrealTargetPlatform being build
-		 *	
-		 *	@return	CPPTargetPlatform			The CPPTargetPlatform to compile for
-		 */
+		/// <summary>
+		/// Retrieve the CPPTargetPlatform for the given UnrealTargetPlatform
+		/// </summary>
+		/// <param name="InUnrealTargetPlatform"> The UnrealTargetPlatform being build</param>
+		/// <returns>CPPTargetPlatform   The CPPTargetPlatform to compile for</returns>
 		public abstract CPPTargetPlatform GetCPPTargetPlatform(UnrealTargetPlatform InUnrealTargetPlatform);
 
 		/// <summary>
@@ -513,191 +474,172 @@ namespace UnrealBuildTool
 			return false;
 		}
 
-		/**
-		 *	Get the extension to use for the given binary type
-		 *	
-		 *	@param	InBinaryType		The binary type being built
-		 *	
-		 *	@return	string				The binary extension (i.e. 'exe' or 'dll')
-		 */
+		/// <summary>
+		/// Get the extension to use for the given binary type
+		/// </summary>
+		/// <param name="InBinaryType"> The binary type being built</param>
+		/// <returns>string    The binary extension (i.e. 'exe' or 'dll')</returns>
 		public virtual string GetBinaryExtension(UEBuildBinaryType InBinaryType)
 		{
 			throw new BuildException("GetBinaryExtensiton for {0} not handled in {1}", InBinaryType.ToString(), this.ToString());
 		}
 
-		/**
-		 *	Get the extension to use for debug info for the given binary type
-		 *	
-		 *	@param	InBinaryType		The binary type being built
-		 *	
-		 *	@return	string				The debug info extension (i.e. 'pdb')
-		 */
+		/// <summary>
+		/// Get the extension to use for debug info for the given binary type
+		/// </summary>
+		/// <param name="InBinaryType"> The binary type being built</param>
+		/// <returns>string    The debug info extension (i.e. 'pdb')</returns>
 		public virtual string GetDebugInfoExtension(UEBuildBinaryType InBinaryType)
 		{
 			throw new BuildException("GetDebugInfoExtension for {0} not handled in {1}", InBinaryType.ToString(), this.ToString());
 		}
 
-		/**
-		 *	Whether incremental linking should be used
-		 *	
-		 *	@param	InPlatform			The CPPTargetPlatform being built
-		 *	@param	InConfiguration		The CPPTargetConfiguration being built
-		 *	
-		 *	@return	bool	true if incremental linking should be used, false if not
-		 */
+		/// <summary>
+		/// Whether incremental linking should be used
+		/// </summary>
+		/// <param name="InPlatform">  The CPPTargetPlatform being built</param>
+		/// <param name="InConfiguration"> The CPPTargetConfiguration being built</param>
+		/// <returns>bool true if incremental linking should be used, false if not</returns>
 		public virtual bool ShouldUseIncrementalLinking(CPPTargetPlatform Platform, CPPTargetConfiguration Configuration)
 		{
 			return false;
 		}
 
-		/**
-		 *	Whether PDB files should be used
-		 *	
-		 *	@param	InPlatform			The CPPTargetPlatform being built
-		 *	@param	InConfiguration		The CPPTargetConfiguration being built
-		 *	@param	bInCreateDebugInfo	true if debug info is getting create, false if not
-		 *	
-		 *	@return	bool	true if PDB files should be used, false if not
-		 */
+		/// <summary>
+		/// Whether PDB files should be used
+		/// </summary>
+		/// <param name="InPlatform">  The CPPTargetPlatform being built</param>
+		/// <param name="InConfiguration"> The CPPTargetConfiguration being built</param>
+		/// <param name="bInCreateDebugInfo">true if debug info is getting create, false if not</param>
+		/// <returns>bool true if PDB files should be used, false if not</returns>
 		public virtual bool ShouldUsePDBFiles(CPPTargetPlatform Platform, CPPTargetConfiguration Configuration, bool bCreateDebugInfo)
 		{
 			return false;
 		}
 
-		/**
-		 *	Whether PCH files should be used
-		 *	
-		 *	@param	InPlatform			The CPPTargetPlatform being built
-		 *	@param	InConfiguration		The CPPTargetConfiguration being built
-		 *	
-		 *	@return	bool				true if PCH files should be used, false if not
-		 */
+		/// <summary>
+		/// Whether PCH files should be used
+		/// </summary>
+		/// <param name="InPlatform">  The CPPTargetPlatform being built</param>
+		/// <param name="InConfiguration"> The CPPTargetConfiguration being built</param>
+		/// <returns>bool    true if PCH files should be used, false if not</returns>
 		public virtual bool ShouldUsePCHFiles(CPPTargetPlatform Platform, CPPTargetConfiguration Configuration)
 		{
 			return BuildConfiguration.bUsePCHFiles;
 		}
 
-		/**
-		 *	Whether the editor should be built for this platform or not
-		 *	
-		 *	@param	InPlatform		The UnrealTargetPlatform being built
-		 *	@param	InConfiguration	The UnrealTargetConfiguration being built
-		 *	@return	bool			true if the editor should be built, false if not
-		 */
+		/// <summary>
+		/// Whether the editor should be built for this platform or not
+		/// </summary>
+		/// <param name="InPlatform"> The UnrealTargetPlatform being built</param>
+		/// <param name="InConfiguration">The UnrealTargetConfiguration being built</param>
+		/// <returns>bool   true if the editor should be built, false if not</returns>
 		public virtual bool ShouldNotBuildEditor(UnrealTargetPlatform InPlatform, UnrealTargetConfiguration InConfiguration)
 		{
 			return false;
 		}
 
-		/**
-		 *	Whether this build should support ONLY cooked data or not
-		 *	
-		 *	@param	InPlatform		The UnrealTargetPlatform being built
-		 *	@param	InConfiguration	The UnrealTargetConfiguration being built
-		 *	@return	bool			true if the editor should be built, false if not
-		 */
+		/// <summary>
+		/// Whether this build should support ONLY cooked data or not
+		/// </summary>
+		/// <param name="InPlatform"> The UnrealTargetPlatform being built</param>
+		/// <param name="InConfiguration">The UnrealTargetConfiguration being built</param>
+		/// <returns>bool   true if the editor should be built, false if not</returns>
 		public virtual bool BuildRequiresCookedData(UnrealTargetPlatform InPlatform, UnrealTargetConfiguration InConfiguration)
 		{
 			return false;
 		}
 
-		/**
-		 *	Whether the platform requires the extra UnityCPPWriter
-		 *	This is used to add an extra file for UBT to get the #include dependencies from
-		 *	
-		 *	@return	bool	true if it is required, false if not
-		 */
+		/// <summary>
+		/// Whether the platform requires the extra UnityCPPWriter
+		/// This is used to add an extra file for UBT to get the #include dependencies from
+		/// </summary>
+		/// <returns>bool true if it is required, false if not</returns>
 		public virtual bool RequiresExtraUnityCPPWriter()
 		{
 			return false;
 		}
 
-		/**
-		 * Whether this platform should build a monolithic binary
-		 */
+		/// <summary>
+		/// Whether this platform should build a monolithic binary
+		/// </summary>
 		public virtual bool ShouldCompileMonolithicBinary(UnrealTargetPlatform InPlatform)
 		{
 			return false;
 		}
 
-		/**
-		 *	Get a list of extra modules the platform requires.
-		 *	This is to allow undisclosed platforms to add modules they need without exposing information about the platform.
-		 *	
-		 *	@param	Target						The target being build
-		 *	@param	BuildTarget					The UEBuildTarget getting build
-		 *	@param	PlatformExtraModules		OUTPUT the list of extra modules the platform needs to add to the target
-		 */
+		/// <summary>
+		/// Get a list of extra modules the platform requires.
+		/// This is to allow undisclosed platforms to add modules they need without exposing information about the platform.
+		/// </summary>
+		/// <param name="Target">     The target being build</param>
+		/// <param name="BuildTarget">    The UEBuildTarget getting build</param>
+		/// <param name="PlatformExtraModules"> OUTPUT the list of extra modules the platform needs to add to the target</param>
 		public virtual void GetExtraModules(TargetInfo Target, UEBuildTarget BuildTarget, ref List<string> PlatformExtraModules)
 		{
 		}
 
-		/**
-		 *	Modify the newly created module passed in for this platform.
-		 *	This is not required - but allows for hiding details of a
-		 *	particular platform.
-		 *	
-		 *  @param  ModuleName		The name of the module
-		 *	@param	Rules			The module rules
-		 *	@param	Target			The target being build
-		 */
+		/// <summary>
+		/// Modify the newly created module passed in for this platform.
+		/// This is not required - but allows for hiding details of a
+		/// particular platform.
+		/// </summary>
+		/// <param name="">ModuleName  The name of the module</param>
+		/// <param name="Rules">  The module rules</param>
+		/// <param name="Target">  The target being build</param>
 		public virtual void ModifyModuleRules(string ModuleName, ModuleRules Rules, TargetInfo Target)
 		{
 		}
 
-		/**
-		 *	Setup the target environment for building
-		 *	
-		 *	@param	InBuildTarget		The target being built
-		 */
+		/// <summary>
+		/// Setup the target environment for building
+		/// </summary>
+		/// <param name="InBuildTarget"> The target being built</param>
 		public abstract void SetUpEnvironment(UEBuildTarget InBuildTarget);
 
-		/**
-		 * Allow the platform to set an optional architecture
-		 */
+		/// <summary>
+		/// Allow the platform to set an optional architecture
+		/// </summary>
 		public virtual string GetActiveArchitecture()
 		{
 			// by default, use an empty architecture (which is really just a modifer to the platform for some paths/names)
 			return "";
 		}
 
-		/**
-		 * Allow the platform to override the NMake output name
-		 */
+		/// <summary>
+		/// Allow the platform to override the NMake output name
+		/// </summary>
 		public virtual FileReference ModifyNMakeOutput(FileReference ExeName)
 		{
 			// by default, use original
 			return ExeName;
 		}
 
-		/**
-		 * Allow the platform to apply architecture-specific name according to its rules
-		 * 
-		 * @param BinaryName name of the binary, not specific to any architecture
-		 *
-		 * @return Possibly architecture-specific name
-		 */
+		/// <summary>
+		/// Allow the platform to apply architecture-specific name according to its rules
+		/// </summary>
+		/// <param name="BinaryName">name of the binary, not specific to any architecture</param>
+		/// <returns>Possibly architecture-specific name</returns>
 		public virtual string ApplyArchitectureName(string BinaryName)
 		{
 			// by default, use logic that assumes architectures to start with "-"
 			return BinaryName + GetActiveArchitecture();
 		}
 
-		/**
-		 * For platforms that need to output multiple files per binary (ie Android "fat" binaries)
-		 * this will emit multiple paths. By default, it simply makes an array from the input
-		 */
+		/// <summary>
+		/// For platforms that need to output multiple files per binary (ie Android "fat" binaries)
+		/// this will emit multiple paths. By default, it simply makes an array from the input
+		/// </summary>
 		public virtual List<FileReference> FinalizeBinaryPaths(FileReference BinaryName)
 		{
 			List<FileReference> TempList = new List<FileReference>() { BinaryName };
 			return TempList;
 		}
 
-		/**
-		 *	Setup the configuration environment for building
-		 *	
-		 *	@param	InBuildTarget		The target being built
-		 */
+		/// <summary>
+		/// Setup the configuration environment for building
+		/// </summary>
+		/// <param name="InBuildTarget"> The target being built</param>
 		public virtual void SetUpConfigurationEnvironment(UEBuildTarget InBuildTarget)
 		{
 			// Determine the C++ compile/link configuration based on the Unreal configuration.
@@ -748,11 +690,10 @@ namespace UnrealBuildTool
 			InBuildTarget.GlobalLinkEnvironment.Config.bCreateDebugInfo = InBuildTarget.GlobalCompileEnvironment.Config.bCreateDebugInfo;
 		}
 
-		/**
-		 *	Setup the project environment for building
-		 *	
-		 *	@param	InBuildTarget		The target being built
-		 */
+		/// <summary>
+		/// Setup the project environment for building
+		/// </summary>
+		/// <param name="InBuildTarget"> The target being built</param>
 		public virtual void SetUpProjectEnvironment(UnrealTargetPlatform InPlatform)
 		{
 			if (!bInitializedProject)
@@ -840,69 +781,64 @@ namespace UnrealBuildTool
 			}
 		}
 
-		/**
-		 *	Whether this platform should create debug information or not
-		 *	
-		 *	@param	InPlatform			The UnrealTargetPlatform being built
-		 *	@param	InConfiguration		The UnrealTargetConfiguration being built
-		 *	
-		 *	@return	bool				true if debug info should be generated, false if not
-		 */
+		/// <summary>
+		/// Whether this platform should create debug information or not
+		/// </summary>
+		/// <param name="InPlatform">  The UnrealTargetPlatform being built</param>
+		/// <param name="InConfiguration"> The UnrealTargetConfiguration being built</param>
+		/// <returns>bool    true if debug info should be generated, false if not</returns>
 		public abstract bool ShouldCreateDebugInfo(UnrealTargetPlatform Platform, UnrealTargetConfiguration Configuration);
 
-		/**
-		 *	Gives the platform a chance to 'override' the configuration settings 
-		 *	that are overridden on calls to RunUBT.
-		 *	
-		 *	@param	InPlatform			The UnrealTargetPlatform being built
-		 *	@param	InConfiguration		The UnrealTargetConfiguration being built
-		 *	
-		 *	@return	bool				true if debug info should be generated, false if not
-		 */
+		/// <summary>
+		/// Gives the platform a chance to 'override' the configuration settings
+		/// that are overridden on calls to RunUBT.
+		/// </summary>
+		/// <param name="InPlatform">  The UnrealTargetPlatform being built</param>
+		/// <param name="InConfiguration"> The UnrealTargetConfiguration being built</param>
+		/// <returns>bool    true if debug info should be generated, false if not</returns>
 		public virtual void ResetBuildConfiguration(UnrealTargetPlatform InPlatform, UnrealTargetConfiguration InConfiguration)
 		{
 		}
 
-		/**
-		 *	Validate the UEBuildConfiguration for this platform
-		 *	This is called BEFORE calling UEBuildConfiguration to allow setting 
-		 *	various fields used in that function such as CompileLeanAndMean...
-		 */
+		/// <summary>
+		/// Validate the UEBuildConfiguration for this platform
+		/// This is called BEFORE calling UEBuildConfiguration to allow setting
+		/// various fields used in that function such as CompileLeanAndMean...
+		/// </summary>
 		public virtual void ValidateUEBuildConfiguration()
 		{
 		}
 
-		/**
-		 *	Validate configuration for this platform
-		 *	NOTE: This function can/will modify BuildConfiguration!
-		 *	
-		 *	@param	InPlatform			The CPPTargetPlatform being built
-		 *	@param	InConfiguration		The CPPTargetConfiguration being built
-		 *	@param	bInCreateDebugInfo	true if debug info is getting create, false if not
-		 */
+		/// <summary>
+		/// Validate configuration for this platform
+		/// NOTE: This function can/will modify BuildConfiguration!
+		/// </summary>
+		/// <param name="InPlatform">  The CPPTargetPlatform being built</param>
+		/// <param name="InConfiguration"> The CPPTargetConfiguration being built</param>
+		/// <param name="bInCreateDebugInfo">true if debug info is getting create, false if not</param>
 		public virtual void ValidateBuildConfiguration(CPPTargetConfiguration Configuration, CPPTargetPlatform Platform, bool bCreateDebugInfo)
 		{
 		}
 
-		/**
-		 *	Return whether this platform has uniquely named binaries across multiple games
-		 */
+		/// <summary>
+		/// Return whether this platform has uniquely named binaries across multiple games
+		/// </summary>
 		public virtual bool HasUniqueBinaries()
 		{
 			return true;
 		}
 
-		/**
-		 *	Return whether we wish to have this platform's binaries in our builds
-		 */
+		/// <summary>
+		/// Return whether we wish to have this platform's binaries in our builds
+		/// </summary>
 		public virtual bool IsBuildRequired()
 		{
 			return true;
 		}
 
-		/**
-		 *	Return whether we wish to have this platform's binaries in our CIS tests
-		 */
+		/// <summary>
+		/// Return whether we wish to have this platform's binaries in our CIS tests
+		/// </summary>
 		public virtual bool IsCISRequired()
 		{
 			return true;
@@ -917,11 +853,10 @@ namespace UnrealBuildTool
 			return false;
 		}
 
-		/**
-		 *	Return all valid configurations for this platform
-		 *	
-		 *  Typically, this is always Debug, Development, and Shipping - but Test is a likely future addition for some platforms
-		 */
+		/// <summary>
+		/// Return all valid configurations for this platform
+		/// Typically, this is always Debug, Development, and Shipping - but Test is a likely future addition for some platforms
+		/// </summary>
 		public virtual List<UnrealTargetConfiguration> GetConfigurations(UnrealTargetPlatform InUnrealTargetPlatform, bool bIncludeDebug)
 		{
 			List<UnrealTargetConfiguration> Configurations = new List<UnrealTargetConfiguration>()
@@ -937,11 +872,10 @@ namespace UnrealBuildTool
 			return Configurations;
 		}
 
-		/**
-		 *	Setup the binaries for this specific platform.
-		 *	
-		 *	@param	InBuildTarget		The target being built
-		 */
+		/// <summary>
+		/// Setup the binaries for this specific platform.
+		/// </summary>
+		/// <param name="InBuildTarget"> The target being built</param>
 		public virtual void SetupBinaries(UEBuildTarget InBuildTarget)
 		{
 		}
@@ -994,11 +928,10 @@ namespace UnrealBuildTool
 			return true;
 		}
 
-		/**
-		 * Check for the default configuration
-		 *
-		 * return true if the project uses the default build config
-		 */
+		/// <summary>
+		/// Check for the default configuration
+		/// return true if the project uses the default build config
+		/// </summary>
 		public virtual bool HasDefaultBuildConfig(UnrealTargetPlatform Platform, DirectoryReference ProjectDirectoryName)
 		{
 			string[] BoolKeys = new string[] {
@@ -1019,13 +952,19 @@ namespace UnrealBuildTool
 
 		#region protected AutoSDKs Utility
 
-		/** Name of the file that holds currently install SDK version string */
+		/// <summary>
+		/// Name of the file that holds currently install SDK version string
+		/// </summary>
 		protected static string CurrentlyInstalledSDKStringManifest = "CurrentlyInstalled.txt";
 
-		/** name of the file that holds the last succesfully run SDK setup script version */
+		/// <summary>
+		/// name of the file that holds the last succesfully run SDK setup script version
+		/// </summary>
 		protected static string LastRunScriptVersionManifest = "CurrentlyInstalled.Version.txt";
 
-		/** Name of the file that holds environment variables of current SDK */
+		/// <summary>
+		/// Name of the file that holds environment variables of current SDK
+		/// </summary>
 		protected static string SDKEnvironmentVarsFile = "OutputEnvVars.txt";
 
 		protected static readonly string SDKRootEnvVar = "UE_SDKS_ROOT";
@@ -1034,11 +973,10 @@ namespace UnrealBuildTool
 
 		public static bool bShouldLogInfo = false;
 
-		/** 
-		 * Whether platform supports switching SDKs during runtime
-		 * 
-		 * @return true if supports
-		 */
+		/// <summary>
+		/// Whether platform supports switching SDKs during runtime
+		/// </summary>
+		/// <returns>true if supports</returns>
 		protected virtual bool PlatformSupportsAutoSDKs()
 		{
 			return false;
@@ -1066,31 +1004,29 @@ namespace UnrealBuildTool
 			return !IsAutoSDKDestructive() || !HasAnyManualInstall();
 		}
 
-		/** 
-		 * Returns SDK string as required by the platform 
-		 * 
-		 * @return Valid SDK string
-		 */
+		/// <summary>
+		/// Returns SDK string as required by the platform
+		/// </summary>
+		/// <returns>Valid SDK string</returns>
 		protected virtual string GetRequiredSDKString()
 		{
 			return "";
 		}
 
-		/**
-		* Gets the version number of the SDK setup script itself.  The version in the base should ALWAYS be the master revision from the last refactor.  
-		* If you need to force a rebuild for a given platform, override this for the given platform.
-		* @return Setup script version
-		*/
+		/// <summary>
+		/// Gets the version number of the SDK setup script itself.  The version in the base should ALWAYS be the master revision from the last refactor.
+		/// If you need to force a rebuild for a given platform, override this for the given platform.
+		/// </summary>
+		/// <returns>Setup script version</returns>
 		protected virtual String GetRequiredScriptVersionString()
 		{
 			return "3.0";
 		}
 
-		/** 
-		 * Returns path to platform SDKs
-		 * 
-		 * @return Valid SDK string
-		 */
+		/// <summary>
+		/// Returns path to platform SDKs
+		/// </summary>
+		/// <returns>Valid SDK string</returns>
 		protected string GetPathToPlatformAutoSDKs()
 		{
 			string SDKPath = "";
@@ -1105,25 +1041,23 @@ namespace UnrealBuildTool
 			return SDKPath;
 		}
 
-		/**
-		 * Because most ManualSDK determination depends on reading env vars, if this process is spawned by a process that ALREADY set up
-		 * AutoSDKs then all the SDK env vars will exist, and we will spuriously detect a Manual SDK. (children inherit the environment of the parent process).
-		 * Therefore we write out an env var to set in the command file (OutputEnvVars.txt) such that child processes can determine if their manual SDK detection
-		 * is bogus.  Make it platform specific so that platforms can be in different states.
-		 */
+		/// <summary>
+		/// Because most ManualSDK determination depends on reading env vars, if this process is spawned by a process that ALREADY set up
+		/// AutoSDKs then all the SDK env vars will exist, and we will spuriously detect a Manual SDK. (children inherit the environment of the parent process).
+		/// Therefore we write out an env var to set in the command file (OutputEnvVars.txt) such that child processes can determine if their manual SDK detection
+		/// is bogus.  Make it platform specific so that platforms can be in different states.
+		/// </summary>
 		protected string GetPlatformAutoSDKSetupEnvVar()
 		{
 			return GetSDKTargetPlatformName() + AutoSetupEnvVar;
 		}
 
-		/**
-		 * Gets currently installed version
-		 * 
-		 * @param PlatformSDKRoot absolute path to platform SDK root
-		 * @param OutInstalledSDKVersionString version string as currently installed
-		 * 
-		 * @return true if was able to read it
-		 */
+		/// <summary>
+		/// Gets currently installed version
+		/// </summary>
+		/// <param name="PlatformSDKRoot">absolute path to platform SDK root</param>
+		/// <param name="OutInstalledSDKVersionString">version string as currently installed</param>
+		/// <returns>true if was able to read it</returns>
 		protected bool GetCurrentlyInstalledSDKString(string PlatformSDKRoot, out string OutInstalledSDKVersionString)
 		{
 			if (Directory.Exists(PlatformSDKRoot))
@@ -1153,14 +1087,12 @@ namespace UnrealBuildTool
 			return false;
 		}
 
-		/**
-		 * Gets the version of the last successfully run setup script.
-		 * 
-		 * @param PlatformSDKRoot absolute path to platform SDK root
-		 * @param OutLastRunScriptVersion version string
-		 * 
-		 * @return true if was able to read it
-		 */
+		/// <summary>
+		/// Gets the version of the last successfully run setup script.
+		/// </summary>
+		/// <param name="PlatformSDKRoot">absolute path to platform SDK root</param>
+		/// <param name="OutLastRunScriptVersion">version string</param>
+		/// <returns>true if was able to read it</returns>
 		protected bool GetLastRunScriptVersionString(string PlatformSDKRoot, out string OutLastRunScriptVersion)
 		{
 			if (Directory.Exists(PlatformSDKRoot))
@@ -1184,14 +1116,12 @@ namespace UnrealBuildTool
 			return false;
 		}
 
-		/**
-		 * Sets currently installed version
-		 * 
-		 * @param PlatformSDKRoot absolute path to platform SDK root
-		 * @param InstalledSDKVersionString SDK version string to set
-		 * 
-		 * @return true if was able to set it
-		 */
+		/// <summary>
+		/// Sets currently installed version
+		/// </summary>
+		/// <param name="PlatformSDKRoot">absolute path to platform SDK root</param>
+		/// <param name="InstalledSDKVersionString">SDK version string to set</param>
+		/// <returns>true if was able to set it</returns>
 		protected bool SetCurrentlyInstalledAutoSDKString(String InstalledSDKVersionString)
 		{
 			String PlatformSDKRoot = GetPathToPlatformAutoSDKs();
@@ -1263,12 +1193,11 @@ namespace UnrealBuildTool
 			return false;
 		}
 
-		/**
-		* Returns Hook names as needed by the platform
-		* (e.g. can be overridden with custom executables or scripts)
-		*
-		* @param Hook Hook type
-		*/
+		/// <summary>
+		/// Returns Hook names as needed by the platform
+		/// (e.g. can be overridden with custom executables or scripts)
+		/// </summary>
+		/// <param name="Hook">Hook type</param>
 		protected virtual string GetHookExecutableName(SDKHookType Hook)
 		{
 			if (Hook == SDKHookType.Uninstall)
@@ -1279,16 +1208,14 @@ namespace UnrealBuildTool
 			return "setup.bat";
 		}
 
-		/**
-		 * Runs install/uninstall hooks for SDK
-		 * 
-		 * @param PlatformSDKRoot absolute path to platform SDK root
-		 * @param SDKVersionString version string to run for (can be empty!)
-		 * @param Hook which one of hooks to run
-		 * @param bHookCanBeNonExistent whether a non-existing hook means failure
-		 * 
-		 * @return true if succeeded
-		 */
+		/// <summary>
+		/// Runs install/uninstall hooks for SDK
+		/// </summary>
+		/// <param name="PlatformSDKRoot">absolute path to platform SDK root</param>
+		/// <param name="SDKVersionString">version string to run for (can be empty!)</param>
+		/// <param name="Hook">which one of hooks to run</param>
+		/// <param name="bHookCanBeNonExistent">whether a non-existing hook means failure</param>
+		/// <returns>true if succeeded</returns>
 		protected virtual bool RunAutoSDKHooks(string PlatformSDKRoot, string SDKVersionString, SDKHookType Hook, bool bHookCanBeNonExistent = true)
 		{
 			if (!IsAutoSDKSafe())
@@ -1351,15 +1278,13 @@ namespace UnrealBuildTool
 			return bHookCanBeNonExistent;
 		}
 
-		/**
-		 * Loads environment variables from SDK
-		 * If any commands are added or removed the handling needs to be duplicated in
-		 * TargetPlatformManagerModule.cpp
-		 * 
-		 * @param PlatformSDKRoot absolute path to platform SDK         
-		 * 
-		 * @return true if succeeded
-		 */
+		/// <summary>
+		/// Loads environment variables from SDK
+		/// If any commands are added or removed the handling needs to be duplicated in
+		/// TargetPlatformManagerModule.cpp
+		/// </summary>
+		/// <param name="PlatformSDKRoot">absolute path to platform SDK</param>
+		/// <returns>true if succeeded</returns>
 		protected bool SetupEnvironmentFromAutoSDK(string PlatformSDKRoot)
 		{
 			string EnvVarFile = Path.Combine(PlatformSDKRoot, SDKEnvironmentVarsFile);
@@ -1527,10 +1452,10 @@ namespace UnrealBuildTool
 			}
 		}
 
-		/** 
-		 * Currently installed AutoSDK is written out to a text file in a known location.  
-		 * This function just compares the file's contents with the current requirements.
-		 */
+		/// <summary>
+		/// Currently installed AutoSDK is written out to a text file in a known location.
+		/// This function just compares the file's contents with the current requirements.
+		/// </summary>
 		protected SDKStatus HasRequiredAutoSDKInstalled()
 		{
 			if (PlatformSupportsAutoSDKs() && HasAutoSDKSystemEnabled())
@@ -1624,11 +1549,11 @@ namespace UnrealBuildTool
 			return false;
 		}
 
-		/**
-		* Runs batch files if necessary to set up required AutoSDK.
-		* AutoSDKs are SDKs that have not been setup through a formal installer, but rather come from
-		* a source control directory, or other local copy.
-		*/
+		/// <summary>
+		/// Runs batch files if necessary to set up required AutoSDK.
+		/// AutoSDKs are SDKs that have not been setup through a formal installer, but rather come from
+		/// a source control directory, or other local copy.
+		/// </summary>
 		private void SetupAutoSDK()
 		{
 			if (IsAutoSDKSafe() && PlatformSupportsAutoSDKs() && HasAutoSDKSystemEnabled())
@@ -1681,18 +1606,19 @@ namespace UnrealBuildTool
 
 		#region public AutoSDKs Utility
 
-		/** Enum describing types of hooks a platform SDK can have */
+		/// <summary>
+		/// Enum describing types of hooks a platform SDK can have
+		/// </summary>
 		public enum SDKHookType
 		{
 			Install,
 			Uninstall
 		};
 
-		/** 
-		 * Returns platform-specific name used in SDK repository
-		 * 
-		 * @return path to SDK Repository
-		 */
+		/// <summary>
+		/// Returns platform-specific name used in SDK repository
+		/// </summary>
+		/// <returns>path to SDK Repository</returns>
 		public virtual string GetSDKTargetPlatformName()
 		{
 			return "";
@@ -1715,10 +1641,10 @@ namespace UnrealBuildTool
 			return SDKStatus.Valid;
 		}
 
-		/**
-		 *	Whether the required external SDKs are installed for this platform.  
-		 *	Could be either a manual install or an AutoSDK.
-		 */
+		/// <summary>
+		/// Whether the required external SDKs are installed for this platform.
+		/// Could be either a manual install or an AutoSDK.
+		/// </summary>
 		public SDKStatus HasRequiredSDKsInstalled()
 		{
 			// avoid redundant potentially expensive SDK checks.
diff --git a/Engine/Source/Programs/UnrealBuildTool/Configuration/UEBuildTarget.cs b/Engine/Source/Programs/UnrealBuildTool/Configuration/UEBuildTarget.cs
--- a/Engine/Source/Programs/UnrealBuildTool/Configuration/UEBuildTarget.cs
+++ b/Engine/Source/Programs/UnrealBuildTool/Configuration/UEBuildTarget.cs
@@ -72,17 +72,29 @@ namespace UnrealBuildTool
 		Check,
 	}
 
-	/** Helper class for holding project name, platform and config together. */
+	/// <summary>
+	/// Helper class for holding project name, platform and config together.
+	/// </summary>
 	public class UnrealBuildConfig
 	{
-		/** Project to build. */
+		/// <summary>
+		/// Project to build.
+		/// </summary>
 		public string Name;
-		/** Platform to build the project for. */
+
+		/// <summary>
+		/// Platform to build the project for.
+		/// </summary>
 		public UnrealTargetPlatform Platform;
-		/** Config to build the project with. */
+
+		/// <summary>
+		/// Config to build the project with.
+		/// </summary>
 		public UnrealTargetConfiguration Config;
 
-		/** Constructor */
+		/// <summary>
+		/// Constructor
+		/// </summary>
 		public UnrealBuildConfig(string InName, UnrealTargetPlatform InPlatform, UnrealTargetConfiguration InConfig)
 		{
 			Name = InName;
@@ -90,7 +102,9 @@ namespace UnrealBuildTool
 			Config = InConfig;
 		}
 
-		/** Overridden ToString() to make this class esier to read when debugging. */
+		/// <summary>
+		/// Overridden ToString() to make this class esier to read when debugging.
+		/// </summary>
 		public override string ToString()
 		{
 			return String.Format("{0}, {1}, {2}", Name, Platform, Config);
@@ -162,10 +176,14 @@ namespace UnrealBuildTool
 			OnlyModuleSuffix = InitOnlyModuleSuffix;
 		}
 
-		/** If building only a single module, this is the module name to build */
+		/// <summary>
+		/// If building only a single module, this is the module name to build
+		/// </summary>
 		public readonly string OnlyModuleName;
 
-		/** When building only a single module, the optional suffix for the module file name */
+		/// <summary>
+		/// When building only a single module, the optional suffix for the module file name
+		/// </summary>
 		public readonly string OnlyModuleSuffix;
 	}
 
@@ -190,9 +208,9 @@ namespace UnrealBuildTool
 	}
 
 
-	/**
-	 * A target that can be built
-	 */
+	/// <summary>
+	/// A target that can be built
+	/// </summary>
 	[Serializable]
 	public class UEBuildTarget : ISerializable
 	{
@@ -704,10 +722,10 @@ namespace UnrealBuildTool
 		}
 
 
-		/** 
-		 * Look for all folders with a uproject file, these are valid games
-		 * This is defined as a valid game
-		 */
+		/// <summary>
+		/// Look for all folders with a uproject file, these are valid games
+		/// This is defined as a valid game
+		/// </summary>
 		public static List<DirectoryReference> DiscoverAllGameFolders()
 		{
 			List<DirectoryReference> AllGameFolders = new List<DirectoryReference>();
@@ -722,55 +740,87 @@ namespace UnrealBuildTool
 			return AllGameFolders;
 		}
 
-		/** The target rules */
+		/// <summary>
+		/// The target rules
+		/// </summary>
 		[NonSerialized]
 		public TargetRules Rules = null;
 
-		/** The rules assembly to use when searching for modules */
+		/// <summary>
+		/// The rules assembly to use when searching for modules
+		/// </summary>
 		[NonSerialized]
 		public RulesAssembly RulesAssembly;
 
-		/** The project file for this target */
+		/// <summary>
+		/// The project file for this target
+		/// </summary>
 		public FileReference ProjectFile;
 
-		/** The project descriptor for this target */
+		/// <summary>
+		/// The project descriptor for this target
+		/// </summary>
 		[NonSerialized]
 		public ProjectDescriptor ProjectDescriptor;
 
-		/** Type of target */
+		/// <summary>
+		/// Type of target
+		/// </summary>
 		public TargetRules.TargetType TargetType;
 
-		/** The name of the application the target is part of. For targets with bUseSharedBuildEnvironment = true, this is typically the name of the base application, eg. UE4Editor for any game editor. */
+		/// <summary>
+		/// The name of the application the target is part of. For targets with bUseSharedBuildEnvironment = true, this is typically the name of the base application, eg. UE4Editor for any game editor.
+		/// </summary>
 		public string AppName;
 
-		/** The name of the target */
+		/// <summary>
+		/// The name of the target
+		/// </summary>
 		public string TargetName;
 
-		/** Whether the target uses the shared build environment. If false, AppName==TargetName and all binaries should be written to the project directory. */
+		/// <summary>
+		/// Whether the target uses the shared build environment. If false, AppName==TargetName and all binaries should be written to the project directory.
+		/// </summary>
 		public bool bUseSharedBuildEnvironment;
 
-		/** Platform as defined by the VCProject and passed via the command line. Not the same as internal config names. */
+		/// <summary>
+		/// Platform as defined by the VCProject and passed via the command line. Not the same as internal config names.
+		/// </summary>
 		public UnrealTargetPlatform Platform;
 
-		/** Target as defined by the VCProject and passed via the command line. Not necessarily the same as internal name. */
+		/// <summary>
+		/// Target as defined by the VCProject and passed via the command line. Not necessarily the same as internal name.
+		/// </summary>
 		public UnrealTargetConfiguration Configuration;
 
-		/** TargetInfo object which can be passed to RulesCompiler */
+		/// <summary>
+		/// TargetInfo object which can be passed to RulesCompiler
+		/// </summary>
 		public TargetInfo TargetInfo;
 
-		/** Root directory for the active project. Typically contains the .uproject file, or the engine root. */
+		/// <summary>
+		/// Root directory for the active project. Typically contains the .uproject file, or the engine root.
+		/// </summary>
 		public DirectoryReference ProjectDirectory;
 
-		/** Default directory for intermediate files. Typically underneath ProjectDirectory. */
+		/// <summary>
+		/// Default directory for intermediate files. Typically underneath ProjectDirectory.
+		/// </summary>
 		public DirectoryReference ProjectIntermediateDirectory;
 
-		/** Directory for engine intermediates. For an agnostic editor/game executable, this will be under the engine directory. For monolithic executables this will be the same as the project intermediate directory. */
+		/// <summary>
+		/// Directory for engine intermediates. For an agnostic editor/game executable, this will be under the engine directory. For monolithic executables this will be the same as the project intermediate directory.
+		/// </summary>
 		public DirectoryReference EngineIntermediateDirectory;
 
-		/** Output paths of final executable. */
+		/// <summary>
+		/// Output paths of final executable.
+		/// </summary>
 		public List<FileReference> OutputPaths;
 
-		/** Returns the OutputPath is there is only one entry in OutputPaths */
+		/// <summary>
+		/// Returns the OutputPath is there is only one entry in OutputPaths
+		/// </summary>
 		public FileReference OutputPath
 		{
 			get
@@ -783,65 +833,99 @@ namespace UnrealBuildTool
 			}
 		}
 
-		/** Remote path of the binary if it is to be synced with CookerSync */
+		/// <summary>
+		/// Remote path of the binary if it is to be synced with CookerSync
+		/// </summary>
 		public string RemoteRoot;
 
-		/** Whether to build target modules that can be reused for future builds */
+		/// <summary>
+		/// Whether to build target modules that can be reused for future builds
+		/// </summary>
 		public bool bPrecompile;
 
-		/** Whether to use precompiled target modules */
+		/// <summary>
+		/// Whether to use precompiled target modules
+		/// </summary>
 		public bool bUsePrecompiled;
 
-		/** The C++ environment that all the environments used to compile UE-based modules are derived from. */
+		/// <summary>
+		/// The C++ environment that all the environments used to compile UE-based modules are derived from.
+		/// </summary>
 		[NonSerialized]
 		public CPPEnvironment GlobalCompileEnvironment = new CPPEnvironment();
 
-		/** The link environment all binary link environments are derived from. */
+		/// <summary>
+		/// The link environment all binary link environments are derived from.
+		/// </summary>
 		[NonSerialized]
 		public LinkEnvironment GlobalLinkEnvironment = new LinkEnvironment();
 
-		/** All plugins which are valid for this target */
+		/// <summary>
+		/// All plugins which are valid for this target
+		/// </summary>
 		[NonSerialized]
 		public List<PluginInfo> ValidPlugins;
 
-		/** All plugins which are built for this target */
+		/// <summary>
+		/// All plugins which are built for this target
+		/// </summary>
 		[NonSerialized]
 		public List<PluginInfo> BuildPlugins;
 
-		/** All plugin dependencies for this target. This differs from the list of plugins that is built for Rocket, where we build everything, but link in only the enabled plugins. */
+		/// <summary>
+		/// All plugin dependencies for this target. This differs from the list of plugins that is built for Rocket, where we build everything, but link in only the enabled plugins.
+		/// </summary>
 		[NonSerialized]
 		public List<PluginInfo> EnabledPlugins;
 
-		/** Additional plugin filenames which are foreign to this target */
+		/// <summary>
+		/// Additional plugin filenames which are foreign to this target
+		/// </summary>
 		[NonSerialized]
 		public List<FileReference> ForeignPlugins;
 
-		/** All application binaries; may include binaries not built by this target. */
+		/// <summary>
+		/// All application binaries; may include binaries not built by this target.
+		/// </summary>
 		[NonSerialized]
 		public List<UEBuildBinary> AppBinaries = new List<UEBuildBinary>();
 
-		/** Extra engine module names to either include in the binary (monolithic) or create side-by-side DLLs for (modular) */
+		/// <summary>
+		/// Extra engine module names to either include in the binary (monolithic) or create side-by-side DLLs for (modular)
+		/// </summary>
 		[NonSerialized]
 		public List<string> ExtraModuleNames = new List<string>();
 
-		/** Extra engine module names which are compiled but not linked into any binary in precompiled engine distributions */
+		/// <summary>
+		/// Extra engine module names which are compiled but not linked into any binary in precompiled engine distributions
+		/// </summary>
 		[NonSerialized]
 		public List<UEBuildBinary> PrecompiledBinaries = new List<UEBuildBinary>();
 
-		/** True if re-compiling this target from the editor */
+		/// <summary>
+		/// True if re-compiling this target from the editor
+		/// </summary>
 		public bool bEditorRecompile;
 
-		/** If building only a specific set of modules, these are the modules to build */
+		/// <summary>
+		/// If building only a specific set of modules, these are the modules to build
+		/// </summary>
 		public List<OnlyModule> OnlyModules = new List<OnlyModule>();
 
-		/** Kept to determine the correct module parsing order when filtering modules. */
+		/// <summary>
+		/// Kept to determine the correct module parsing order when filtering modules.
+		/// </summary>
 		[NonSerialized]
 		protected List<UEBuildBinary> NonFilteredModules = new List<UEBuildBinary>();
 
-		/** true if target should be compiled in monolithic mode, false if not */
+		/// <summary>
+		/// true if target should be compiled in monolithic mode, false if not
+		/// </summary>
 		protected bool bCompileMonolithic = false;
 
-		/** Used to keep track of all modules by name. */
+		/// <summary>
+		/// Used to keep track of all modules by name.
+		/// </summary>
 		[NonSerialized]
 		private Dictionary<string, UEBuildModule> Modules = new CaselessDictionary<UEBuildModule>();
 
@@ -872,23 +956,35 @@ namespace UnrealBuildTool
 			public List<string> UHTHeaderNames = new List<string>();
 		}
 
-		/** Used to map names of modules to their .Build.cs filename */
+		/// <summary>
+		/// Used to map names of modules to their .Build.cs filename
+		/// </summary>
 		public CaselessDictionary<FlatModuleCsDataType> FlatModuleCsData = new CaselessDictionary<FlatModuleCsDataType>();
 
-		/** The receipt for this target, which contains a record of this build. */
+		/// <summary>
+		/// The receipt for this target, which contains a record of this build.
+		/// </summary>
 		private TargetReceipt Receipt;
 
-		/** Filename for the receipt for this target. */
+		/// <summary>
+		/// Filename for the receipt for this target.
+		/// </summary>
 		private string ReceiptFileName;
 
-		/** Version manifests to be written to each output folder */
+		/// <summary>
+		/// Version manifests to be written to each output folder
+		/// </summary>
 		private KeyValuePair<FileReference, VersionManifest>[] FileReferenceToVersionManifestPairs;
 
-		/** Force output of the receipt to an additional filename */
+		/// <summary>
+		/// Force output of the receipt to an additional filename
+		/// </summary>
 		[NonSerialized]
 		private string ForceReceiptFileName;
 
-		/** The name of the .Target.cs file, if the target was created with one */
+		/// <summary>
+		/// The name of the .Target.cs file, if the target was created with one
+		/// </summary>
 		private readonly FileReference TargetCsFilenameField;
 		public FileReference TargetCsFilename { get { return TargetCsFilenameField; } }
 
@@ -1537,7 +1633,9 @@ namespace UnrealBuildTool
 			}
 		}
 
-		/** Generates a list of external files which are required to build this target */
+		/// <summary>
+		/// Generates a list of external files which are required to build this target
+		/// </summary>
 		public void GenerateExternalFileList()
 		{
 			string FileListPath = "../Intermediate/Build/ExternalFiles.xml";
@@ -1660,7 +1758,9 @@ namespace UnrealBuildTool
 			Utils.WriteClass<ExternalFileList>(FileList, FileListPath, "");
 		}
 
-		/** Generates a public manifest file for writing out */
+		/// <summary>
+		/// Generates a public manifest file for writing out
+		/// </summary>
 		public void GenerateManifest()
 		{
 			FileReference ManifestPath;
@@ -1718,7 +1818,9 @@ namespace UnrealBuildTool
 			}
 		}
 
-		/** Prepare all the receipts this target (all the .target and .modules files). See the VersionManifest class for an explanation of what these files are. */
+		/// <summary>
+		/// Prepare all the receipts this target (all the .target and .modules files). See the VersionManifest class for an explanation of what these files are.
+		/// </summary>
 		public void PrepareReceipts(IUEToolChain ToolChain)
 		{
 			// Read the version file
@@ -1807,7 +1909,9 @@ namespace UnrealBuildTool
 			FileReferenceToVersionManifestPairs = FileNameToVersionManifest.ToArray();
 		}
 
-		/** Writes out the version manifest */
+		/// <summary>
+		/// Writes out the version manifest
+		/// </summary>
 		public void WriteReceipts()
 		{
 			if (Receipt != null)
@@ -1859,7 +1963,9 @@ namespace UnrealBuildTool
 			return Output;
 		}
 
-		/** Builds the target, appending list of output files and returns building result. */
+		/// <summary>
+		/// Builds the target, appending list of output files and returns building result.
+		/// </summary>
 		public ECompilationResult Build(IUEToolChain TargetToolChain, out List<FileItem> OutputItems, out List<UHTModuleInfo> UObjectModules, out string EULAViolationWarning)
 		{
 			OutputItems = new List<FileItem>();
@@ -2846,25 +2952,19 @@ namespace UnrealBuildTool
 			return Binary;
 		}
 
-		/**
-		 * @return true if debug information is created, false otherwise.
-		 */
+		/// <returns>true if debug information is created, false otherwise.</returns>
 		public bool IsCreatingDebugInfo()
 		{
 			return GlobalCompileEnvironment.Config.bCreateDebugInfo;
 		}
 
-		/**
-		 * @return The overall output directory of actions for this target.
-		 */
+		/// <returns>The overall output directory of actions for this target.</returns>
 		public DirectoryReference GetOutputDirectory()
 		{
 			return GlobalCompileEnvironment.Config.OutputDirectory;
 		}
 
-		/** 
-		 * @return true if we are building for dedicated server, false otherwise.
-		 */
+		/// <returns>true if we are building for dedicated server, false otherwise.</returns>
 		[Obsolete("IsBuildingDedicatedServer() has been deprecated and will be removed in future release. Update your code to use TargetInfo.Type instead or your code will not compile.")]
 		public bool IsBuildingDedicatedServer()
 		{
@@ -2943,7 +3043,9 @@ namespace UnrealBuildTool
 			return UEBuildPlatform.GetBuildPlatform(Platform).FinalizeBinaryPaths(BinaryFile);
 		}
 
-		/** Sets up the modules for the target. */
+		/// <summary>
+		/// Sets up the modules for the target.
+		/// </summary>
 		protected virtual void SetupModules()
 		{
 			var BuildPlatform = UEBuildPlatform.GetBuildPlatform(Platform);
@@ -2952,7 +3054,9 @@ namespace UnrealBuildTool
 			ExtraModuleNames.AddRange(PlatformExtraModules);
 		}
 
-		/** Sets up the plugins for this target */
+		/// <summary>
+		/// Sets up the plugins for this target
+		/// </summary>
 		protected virtual void SetupPlugins()
 		{
 			// Filter the plugins list by the current project
@@ -3024,7 +3128,9 @@ namespace UnrealBuildTool
 			}
 		}
 
-		/** Checks whether a plugin path contains a platform directory fragment */
+		/// <summary>
+		/// Checks whether a plugin path contains a platform directory fragment
+		/// </summary>
 		private bool ShouldExcludePlugin(PluginInfo Plugin, List<string> ExcludeFragments)
 		{
 			string RelativePathFromRoot;
@@ -3039,7 +3145,9 @@ namespace UnrealBuildTool
 			return ExcludeFragments.Any(x => RelativePathFromRoot.Contains(x));
 		}
 
-		/** Sets up the binaries for the target. */
+		/// <summary>
+		/// Sets up the binaries for the target.
+		/// </summary>
 		protected virtual void SetupBinaries()
 		{
 			if (Rules != null)
@@ -3080,7 +3188,9 @@ namespace UnrealBuildTool
 			// Default does nothing
 		}
 
-		/** Sets up the global compile and link environment for the target. */
+		/// <summary>
+		/// Sets up the global compile and link environment for the target.
+		/// </summary>
 		public virtual void SetupGlobalEnvironment(IUEToolChain ToolChain)
 		{
 			var BuildPlatform = UEBuildPlatform.GetBuildPlatform(Platform);
@@ -3338,7 +3448,9 @@ namespace UnrealBuildTool
 			UEBuildPlatform.GetBuildPlatform(Platform).SetUpProjectEnvironment(Platform);
 		}
 
-		/** Registers a module with this target. */
+		/// <summary>
+		/// Registers a module with this target.
+		/// </summary>
 		public void RegisterModule(UEBuildModule Module)
 		{
 			Debug.Assert(Module.Target == this);
@@ -3346,7 +3458,9 @@ namespace UnrealBuildTool
 			FlatModuleCsData.Add(Module.Name, new FlatModuleCsDataType((Module.RulesFile == null) ? null : Module.RulesFile.FullName));
 		}
 
-		/** Create a rules object for the given module, and set any default values for this target */
+		/// <summary>
+		/// Create a rules object for the given module, and set any default values for this target
+		/// </summary>
 		private ModuleRules CreateModuleRulesAndSetDefaults(string ModuleName, out FileReference ModuleFileName)
 		{
 			// Create the rules from the assembly
@@ -3411,7 +3525,9 @@ namespace UnrealBuildTool
 			return RulesObject;
 		}
 
-		/** Finds a module given its name.  Throws an exception if the module couldn't be found. */
+		/// <summary>
+		/// Finds a module given its name.  Throws an exception if the module couldn't be found.
+		/// </summary>
 		public UEBuildModule FindOrCreateModuleByName(string ModuleName)
 		{
 			UEBuildModule Module;
@@ -3684,13 +3800,17 @@ namespace UnrealBuildTool
 			}
 		}
 
-		/** Normalize an include path to be relative to the engine source directory */
+		/// <summary>
+		/// Normalize an include path to be relative to the engine source directory
+		/// </summary>
 		public static string NormalizeIncludePath(DirectoryReference Directory)
 		{
 			return Utils.CleanDirectorySeparators(Directory.MakeRelativeTo(UnrealBuildTool.EngineSourceDirectory), '/');
 		}
 
-		/** Finds a module given its name.  Throws an exception if the module couldn't be found. */
+		/// <summary>
+		/// Finds a module given its name.  Throws an exception if the module couldn't be found.
+		/// </summary>
 		public UEBuildModule GetModuleByName(string Name)
 		{
 			UEBuildModule Result;
diff --git a/Engine/Source/Programs/UnrealBuildTool/HTML5/HTML5ProjectGenerator.cs b/Engine/Source/Programs/UnrealBuildTool/HTML5/HTML5ProjectGenerator.cs
--- a/Engine/Source/Programs/UnrealBuildTool/HTML5/HTML5ProjectGenerator.cs
+++ b/Engine/Source/Programs/UnrealBuildTool/HTML5/HTML5ProjectGenerator.cs
@@ -8,14 +8,14 @@ using System.IO;
 
 namespace UnrealBuildTool
 {
-	/**
-	 *	Base class for platform-specific project generators 
-	 */
+	/// <summary>
+	/// Base class for platform-specific project generators
+	/// </summary>
 	class HTML5ProjectGenerator : UEPlatformProjectGenerator
 	{
-		/**
-		 *	Register the platform with the UEPlatformProjectGenerator class
-		 */
+		/// <summary>
+		/// Register the platform with the UEPlatformProjectGenerator class
+		/// </summary>
 		public override void RegisterPlatformProjectGenerator()
 		{
 			// Register this project generator for HTML5
@@ -26,14 +26,12 @@ namespace UnrealBuildTool
 		///
 		///	VisualStudio project generation functions
 		///	
-		/**
-		 *	Whether this build platform has native support for VisualStudio
-		 *	
-		 *	@param	InPlatform			The UnrealTargetPlatform being built
-		 *	@param	InConfiguration		The UnrealTargetConfiguration being built
-		 *	
-		 *	@return	bool				true if native VisualStudio support (or custom VSI) is available
-		 */
+		/// <summary>
+		/// Whether this build platform has native support for VisualStudio
+		/// </summary>
+		/// <param name="InPlatform">  The UnrealTargetPlatform being built</param>
+		/// <param name="InConfiguration"> The UnrealTargetConfiguration being built</param>
+		/// <returns>bool    true if native VisualStudio support (or custom VSI) is available</returns>
 		public override bool HasVisualStudioSupport(UnrealTargetPlatform InPlatform, UnrealTargetConfiguration InConfiguration)
 		{
 			// HTML5 is not supported in VisualStudio
diff --git a/Engine/Source/Programs/UnrealBuildTool/HTML5/HTML5ToolChain.cs b/Engine/Source/Programs/UnrealBuildTool/HTML5/HTML5ToolChain.cs
--- a/Engine/Source/Programs/UnrealBuildTool/HTML5/HTML5ToolChain.cs
+++ b/Engine/Source/Programs/UnrealBuildTool/HTML5/HTML5ToolChain.cs
@@ -397,12 +397,11 @@ namespace UnrealBuildTool
 			return Result;
 		}
 
-		/**
-		 * Translates clang output warning/error messages into vs-clickable messages
-		 * 
-		 * @param	sender		Sending object
-		 * @param	e			Event arguments (In this case, the line of string output)
-		 */
+		/// <summary>
+		/// Translates clang output warning/error messages into vs-clickable messages
+		/// </summary>
+		/// <param name="sender"> Sending object</param>
+		/// <param name="e">  Event arguments (In this case, the line of string output)</param>
 		protected void RemoteOutputReceivedEventHandler(object sender, DataReceivedEventArgs e)
 		{
 			var Output = e.Data;
diff --git a/Engine/Source/Programs/UnrealBuildTool/HTML5/UEBuildHTML5.cs b/Engine/Source/Programs/UnrealBuildTool/HTML5/UEBuildHTML5.cs
--- a/Engine/Source/Programs/UnrealBuildTool/HTML5/UEBuildHTML5.cs
+++ b/Engine/Source/Programs/UnrealBuildTool/HTML5/UEBuildHTML5.cs
@@ -14,11 +14,10 @@ namespace UnrealBuildTool
 		[XmlConfig]
 		public static string HTML5Architecture = "";
 
-		/** 
-		 * Whether platform supports switching SDKs during runtime
-		 * 
-		 * @return true if supports
-		 */
+		/// <summary>
+		/// Whether platform supports switching SDKs during runtime
+		/// </summary>
+		/// <returns>true if supports</returns>
 		protected override bool PlatformSupportsAutoSDKs()
 		{
 			return false;
@@ -29,11 +28,10 @@ namespace UnrealBuildTool
 			return "HTML5";
 		}
 
-		/** 
-		 * Returns SDK string as required by the platform 
-		 * 
-		 * @return Valid SDK string
-		 */
+		/// <summary>
+		/// Returns SDK string as required by the platform
+		/// </summary>
+		/// <returns>Valid SDK string</returns>
 		protected override string GetRequiredSDKString()
 		{
 			return HTML5SDKInfo.EmscriptenVersion();
@@ -58,9 +56,9 @@ namespace UnrealBuildTool
 			return (ExeName + "-win32").ChangeExtension(".exe");
 		}
 
-		/**
-			*	Whether the required external SDKs are installed for this platform
-			*/
+		/// <summary>
+		/// Whether the required external SDKs are installed for this platform
+		/// </summary>
 		protected override SDKStatus HasRequiredManualSDKInternal()
 		{
 			if (!HTML5SDKInfo.IsSDKInstalled())
@@ -75,9 +73,9 @@ namespace UnrealBuildTool
 			return (GetActiveArchitecture() == "-win32");
 		}
 
-		/**
-		 *	Register the platform with the UEBuildPlatform class
-		 */
+		/// <summary>
+		/// Register the platform with the UEBuildPlatform class
+		/// </summary>
 		protected override void RegisterBuildPlatformInternal()
 		{
 			// Make sure the SDK is installed
@@ -112,13 +110,11 @@ namespace UnrealBuildTool
 
 		}
 
-		/**
-		 *	Retrieve the CPPTargetPlatform for the given UnrealTargetPlatform
-		 *
-		 *	@param	InUnrealTargetPlatform		The UnrealTargetPlatform being build
-		 *	
-		 *	@return	CPPTargetPlatform			The CPPTargetPlatform to compile for
-		 */
+		/// <summary>
+		/// Retrieve the CPPTargetPlatform for the given UnrealTargetPlatform
+		/// </summary>
+		/// <param name="InUnrealTargetPlatform"> The UnrealTargetPlatform being build</param>
+		/// <returns>CPPTargetPlatform   The CPPTargetPlatform to compile for</returns>
 		public override CPPTargetPlatform GetCPPTargetPlatform(UnrealTargetPlatform InUnrealTargetPlatform)
 		{
 			switch (InUnrealTargetPlatform)
@@ -129,13 +125,11 @@ namespace UnrealBuildTool
 			throw new BuildException("HTML5Platform::GetCPPTargetPlatform: Invalid request for {0}", InUnrealTargetPlatform.ToString());
 		}
 
-		/**
-		 *	Get the extension to use for the given binary type
-		 *	
-		 *	@param	InBinaryType		The binary type being built
-		 *	
-		 *	@return	string				The binary extension (ie 'exe' or 'dll')
-		 */
+		/// <summary>
+		/// Get the extension to use for the given binary type
+		/// </summary>
+		/// <param name="InBinaryType"> The binary type being built</param>
+		/// <returns>string    The binary extension (ie 'exe' or 'dll')</returns>
 		public override string GetBinaryExtension(UEBuildBinaryType InBinaryType)
 		{
 			if (GetActiveArchitecture() == "-win32")
@@ -175,13 +169,11 @@ namespace UnrealBuildTool
 			}
 		}
 
-		/**
-		 *	Get the extension to use for debug info for the given binary type
-		 *	
-		 *	@param	InBinaryType		The binary type being built
-		 *	
-		 *	@return	string				The debug info extension (i.e. 'pdb')
-		 */
+		/// <summary>
+		/// Get the extension to use for debug info for the given binary type
+		/// </summary>
+		/// <param name="InBinaryType"> The binary type being built</param>
+		/// <returns>string    The debug info extension (i.e. 'pdb')</returns>
 		public override string GetDebugInfoExtension(UEBuildBinaryType InBinaryType)
 		{
 			if (GetActiveArchitecture() == "-win32")
@@ -201,25 +193,23 @@ namespace UnrealBuildTool
 			}
 		}
 
-		/**
-		 *	Gives the platform a chance to 'override' the configuration settings 
-		 *	that are overridden on calls to RunUBT.
-		 *	
-		 *	@param	InPlatform			The UnrealTargetPlatform being built
-		 *	@param	InConfiguration		The UnrealTargetConfiguration being built
-		 *	
-		 *	@return	bool				true if debug info should be generated, false if not
-		 */
+		/// <summary>
+		/// Gives the platform a chance to 'override' the configuration settings
+		/// that are overridden on calls to RunUBT.
+		/// </summary>
+		/// <param name="InPlatform">  The UnrealTargetPlatform being built</param>
+		/// <param name="InConfiguration"> The UnrealTargetConfiguration being built</param>
+		/// <returns>bool    true if debug info should be generated, false if not</returns>
 		public override void ResetBuildConfiguration(UnrealTargetPlatform InPlatform, UnrealTargetConfiguration InConfiguration)
 		{
 			ValidateUEBuildConfiguration();
 		}
 
-		/**
-		 *	Validate the UEBuildConfiguration for this platform
-		 *	This is called BEFORE calling UEBuildConfiguration to allow setting 
-		 *	various fields used in that function such as CompileLeanAndMean...
-		 */
+		/// <summary>
+		/// Validate the UEBuildConfiguration for this platform
+		/// This is called BEFORE calling UEBuildConfiguration to allow setting
+		/// various fields used in that function such as CompileLeanAndMean...
+		/// </summary>
 		public override void ValidateUEBuildConfiguration()
 		{
 			UEBuildConfiguration.bCompileLeanAndMeanUE = true;
@@ -231,49 +221,44 @@ namespace UnrealBuildTool
 			UEBuildConfiguration.bCompileForSize = true;
 		}
 
-		/**
-		 * Whether this platform should build a monolithic binary
-		 */
+		/// <summary>
+		/// Whether this platform should build a monolithic binary
+		/// </summary>
 		public override bool ShouldCompileMonolithicBinary(UnrealTargetPlatform InPlatform)
 		{
 			// This platform currently always compiles monolithic
 			return true;
 		}
 
-		/**
-		 *	Whether PDB files should be used
-		 *	
-		 *	@param	InPlatform			The CPPTargetPlatform being built
-		 *	@param	InConfiguration		The CPPTargetConfiguration being built
-		 *	@param	bInCreateDebugInfo	true if debug info is getting create, false if not
-		 *	
-		 *	@return	bool	true if PDB files should be used, false if not
-		 */
+		/// <summary>
+		/// Whether PDB files should be used
+		/// </summary>
+		/// <param name="InPlatform">  The CPPTargetPlatform being built</param>
+		/// <param name="InConfiguration"> The CPPTargetConfiguration being built</param>
+		/// <param name="bInCreateDebugInfo">true if debug info is getting create, false if not</param>
+		/// <returns>bool true if PDB files should be used, false if not</returns>
 		public override bool ShouldUsePDBFiles(CPPTargetPlatform Platform, CPPTargetConfiguration Configuration, bool bCreateDebugInfo)
 		{
 			return bCreateDebugInfo;
 		}
 
-		/**
-		 *	Whether PCH files should be used
-		 *	
-		 *	@param	InPlatform			The CPPTargetPlatform being built
-		 *	@param	InConfiguration		The CPPTargetConfiguration being built
-		 *	
-		 *	@return	bool				true if PCH files should be used, false if not
-		 */
+		/// <summary>
+		/// Whether PCH files should be used
+		/// </summary>
+		/// <param name="InPlatform">  The CPPTargetPlatform being built</param>
+		/// <param name="InConfiguration"> The CPPTargetConfiguration being built</param>
+		/// <returns>bool    true if PCH files should be used, false if not</returns>
 		public override bool ShouldUsePCHFiles(CPPTargetPlatform Platform, CPPTargetConfiguration Configuration)
 		{
 			return false;
 		}
 
-		/**
-		 *	Whether the editor should be built for this platform or not
-		 *	
-		 *	@param	InPlatform		The UnrealTargetPlatform being built
-		 *	@param	InConfiguration	The UnrealTargetConfiguration being built
-		 *	@return	bool			true if the editor should be built, false if not
-		 */
+		/// <summary>
+		/// Whether the editor should be built for this platform or not
+		/// </summary>
+		/// <param name="InPlatform"> The UnrealTargetPlatform being built</param>
+		/// <param name="InConfiguration">The UnrealTargetConfiguration being built</param>
+		/// <returns>bool   true if the editor should be built, false if not</returns>
 		public override bool ShouldNotBuildEditor(UnrealTargetPlatform InPlatform, UnrealTargetConfiguration InConfiguration)
 		{
 			return true;
@@ -284,15 +269,14 @@ namespace UnrealBuildTool
 			return true;
 		}
 
-		/**
-		 *	Modify the newly created module passed in for this platform.
-		 *	This is not required - but allows for hiding details of a
-		 *	particular platform.
-		 *	
-		 *  @param  Name			The name of the module
-		 *	@param	Rules			The module rules
-		 *	@param	Target			The target being build
-		 */
+		/// <summary>
+		/// Modify the newly created module passed in for this platform.
+		/// This is not required - but allows for hiding details of a
+		/// particular platform.
+		/// </summary>
+		/// <param name="">Name   The name of the module</param>
+		/// <param name="Rules">  The module rules</param>
+		/// <param name="Target">  The target being build</param>
 		public override void ModifyModuleRules(string ModuleName, ModuleRules Rules, TargetInfo Target)
 		{
 			if (Target.Platform == UnrealTargetPlatform.HTML5)
@@ -324,11 +308,10 @@ namespace UnrealBuildTool
 		}
 
 
-		/**
-		 *	Setup the target environment for building
-		 *	
-		 *	@param	InBuildTarget		The target being built
-		 */
+		/// <summary>
+		/// Setup the target environment for building
+		/// </summary>
+		/// <param name="InBuildTarget"> The target being built</param>
 		public override void SetUpEnvironment(UEBuildTarget InBuildTarget)
 		{
 			if (GetActiveArchitecture() == "-win32")
@@ -356,14 +339,12 @@ namespace UnrealBuildTool
 
 		}
 
-		/**
-		 *	Whether this platform should create debug information or not
-		 *	
-		 *	@param	InPlatform			The UnrealTargetPlatform being built
-		 *	@param	InConfiguration		The UnrealTargetConfiguration being built
-		 *	
-		 *	@return	bool				true if debug info should be generated, false if not
-		 */
+		/// <summary>
+		/// Whether this platform should create debug information or not
+		/// </summary>
+		/// <param name="InPlatform">  The UnrealTargetPlatform being built</param>
+		/// <param name="InConfiguration"> The UnrealTargetConfiguration being built</param>
+		/// <returns>bool    true if debug info should be generated, false if not</returns>
 		public override bool ShouldCreateDebugInfo(UnrealTargetPlatform Platform, UnrealTargetConfiguration Configuration)
 		{
 			switch (Configuration)
@@ -379,11 +360,10 @@ namespace UnrealBuildTool
 			};
 		}
 
-		/**
-		 *	Setup the binaries for this specific platform.
-		 *	
-		 *	@param	InBuildTarget		The target being built
-		 */
+		/// <summary>
+		/// Setup the binaries for this specific platform.
+		/// </summary>
+		/// <param name="InBuildTarget"> The target being built</param>
 		public override void SetupBinaries(UEBuildTarget InBuildTarget)
 		{
 			// 			InBuildTarget.ExtraModuleNames.Add("ES2RHI");
diff --git a/Engine/Source/Programs/UnrealBuildTool/IOS/IOSProjectGenerator.cs b/Engine/Source/Programs/UnrealBuildTool/IOS/IOSProjectGenerator.cs
--- a/Engine/Source/Programs/UnrealBuildTool/IOS/IOSProjectGenerator.cs
+++ b/Engine/Source/Programs/UnrealBuildTool/IOS/IOSProjectGenerator.cs
@@ -8,14 +8,14 @@ using System.IO;
 
 namespace UnrealBuildTool
 {
-	/**
-	 *	Base class for platform-specific project generators 
-	 */
+	/// <summary>
+	/// Base class for platform-specific project generators
+	/// </summary>
 	class IOSProjectGenerator : UEPlatformProjectGenerator
 	{
-		/**
-		 *	Register the platform with the UEPlatformProjectGenerator class
-		 */
+		/// <summary>
+		/// Register the platform with the UEPlatformProjectGenerator class
+		/// </summary>
 		public override void RegisterPlatformProjectGenerator()
 		{
 			// Register this project generator for Mac
@@ -26,14 +26,12 @@ namespace UnrealBuildTool
 		///
 		///	VisualStudio project generation functions
 		///	
-		/**
-		 *	Whether this build platform has native support for VisualStudio
-		 *	
-		 *	@param	InPlatform			The UnrealTargetPlatform being built
-		 *	@param	InConfiguration		The UnrealTargetConfiguration being built
-		 *	
-		 *	@return	bool				true if native VisualStudio support (or custom VSI) is available
-		 */
+		/// <summary>
+		/// Whether this build platform has native support for VisualStudio
+		/// </summary>
+		/// <param name="InPlatform">  The UnrealTargetPlatform being built</param>
+		/// <param name="InConfiguration"> The UnrealTargetConfiguration being built</param>
+		/// <returns>bool    true if native VisualStudio support (or custom VSI) is available</returns>
 		public override bool HasVisualStudioSupport(UnrealTargetPlatform InPlatform, UnrealTargetConfiguration InConfiguration)
 		{
 			// iOS is not supported in VisualStudio
diff --git a/Engine/Source/Programs/UnrealBuildTool/IOS/IOSToolChain.cs b/Engine/Source/Programs/UnrealBuildTool/IOS/IOSToolChain.cs
--- a/Engine/Source/Programs/UnrealBuildTool/IOS/IOSToolChain.cs
+++ b/Engine/Source/Programs/UnrealBuildTool/IOS/IOSToolChain.cs
@@ -32,41 +32,61 @@ namespace UnrealBuildTool
 
 		// If you are looking for where to change the remote compile server name, look in RemoteToolChain.cs
 
-		/** If this is set, then we don't do any post-compile steps except moving the executable into the proper spot on the Mac */
+		/// <summary>
+		/// If this is set, then we don't do any post-compile steps except moving the executable into the proper spot on the Mac
+		/// </summary>
 		[XmlConfig]
 		public static bool bUseDangerouslyFastMode = false;
 
-		/** Which version of the iOS SDK to target at build time */
+		/// <summary>
+		/// Which version of the iOS SDK to target at build time
+		/// </summary>
 		[XmlConfig]
 		public static string IOSSDKVersion = "latest";
 		public static float IOSSDKVersionFloat = 0.0f;
 
-		/** Which version of the iOS to allow at build time */
+		/// <summary>
+		/// Which version of the iOS to allow at build time
+		/// </summary>
 		[XmlConfig]
 		public static string BuildIOSVersion = "7.0";
 
-		/** Which developer directory to root from */
+		/// <summary>
+		/// Which developer directory to root from
+		/// </summary>
 		private static string XcodeDeveloperDir = "xcode-select";
 
-		/** Directory for the developer binaries */
+		/// <summary>
+		/// Directory for the developer binaries
+		/// </summary>
 		private static string ToolchainDir = "";
 
-		/** Location of the SDKs */
+		/// <summary>
+		/// Location of the SDKs
+		/// </summary>
 		private static string BaseSDKDir;
 		private static string BaseSDKDirSim;
 
-		/** Which compiler frontend to use */
+		/// <summary>
+		/// Which compiler frontend to use
+		/// </summary>
 		private static string IOSCompiler = "clang++";
 
-		/** Which linker frontend to use */
+		/// <summary>
+		/// Which linker frontend to use
+		/// </summary>
 		private static string IOSLinker = "clang++";
 
-		/** Which library archiver to use */
+		/// <summary>
+		/// Which library archiver to use
+		/// </summary>
 		private static string IOSArchiver = "libtool";
 
 		public List<FileReference> BuiltBinaries = new List<FileReference>();
 
-		/** Additional frameworks stored locally so we have access without LinkEnvironment */
+		/// <summary>
+		/// Additional frameworks stored locally so we have access without LinkEnvironment
+		/// </summary>
 		public List<UEBuildFramework> RememberedAdditionalFrameworks = new List<UEBuildFramework>();
 
 		private static void SetupXcodePaths(bool bVerbose)
@@ -86,7 +106,9 @@ namespace UnrealBuildTool
 			IOSSDKVersionFloat = float.Parse(IOSSDKVersion, System.Globalization.CultureInfo.InvariantCulture);
 		}
 
-		/** Hunt down the latest IOS sdk if desired */
+		/// <summary>
+		/// Hunt down the latest IOS sdk if desired
+		/// </summary>
 		public override void SetUpGlobalEnvironment()
 		{
 			base.SetUpGlobalEnvironment();
@@ -727,11 +749,10 @@ namespace UnrealBuildTool
 			return RemoteOutputFile;
 		}
 
-		/**
-		 * Generates debug info for a given executable
-		 * 
-		 * @param Executable FileItem describing the executable to generate debug info for
-		 */
+		/// <summary>
+		/// Generates debug info for a given executable
+		/// </summary>
+		/// <param name="Executable">FileItem describing the executable to generate debug info for</param>
 		public FileItem GenerateDebugInfo(FileItem Executable)
 		{
 			// Make a file item for the source and destination files
