Convert everything in UBT to use standard C# XML comments, so IntelliSense can display useful information in tooltips.

diff --git a/Engine/Source/Programs/UnrealBuildTool/IOS/UEBuildIOS.cs b/Engine/Source/Programs/UnrealBuildTool/IOS/UEBuildIOS.cs
--- a/Engine/Source/Programs/UnrealBuildTool/IOS/UEBuildIOS.cs
+++ b/Engine/Source/Programs/UnrealBuildTool/IOS/UEBuildIOS.cs
@@ -15,25 +15,37 @@ namespace UnrealBuildTool
 		[XmlConfig]
 		public static string IOSArchitecture = "";
 
-		/** Which version of the iOS to allow at run time */
+		/// <summary>
+		/// Which version of the iOS to allow at run time
+		/// </summary>
 		public static string RunTimeIOSVersion = "7.0";
 
-		/** which devices the game is allowed to run on */
+		/// <summary>
+		/// which devices the game is allowed to run on
+		/// </summary>
 		public static string RunTimeIOSDevices = "1,2";
 
-		/** The architecture(s) to compile */
+		/// <summary>
+		/// The architecture(s) to compile
+		/// </summary>
 		[XmlConfig]
 		public static string NonShippingArchitectures = "armv7";
 		[XmlConfig]
 		public static string ShippingArchitectures = "armv7,arm64";
 
-		/** additional linker flags for shipping */
+		/// <summary>
+		/// additional linker flags for shipping
+		/// </summary>
 		public static string AdditionalShippingLinkerFlags = "";
 
-		/** additional linker flags for non-shipping */
+		/// <summary>
+		/// additional linker flags for non-shipping
+		/// </summary>
 		public static string AdditionalLinkerFlags = "";
 
-		/** mobile provision to use for code signing */
+		/// <summary>
+		/// mobile provision to use for code signing
+		/// </summary>
 		public static string MobileProvision = "";
 
 		private bool bInitializedProject = false;
@@ -84,9 +96,9 @@ namespace UnrealBuildTool
 			return SDKStatus.Valid;
 		}
 
-		/**
-		 *	Register the platform with the UEBuildPlatform class
-		 */
+		/// <summary>
+		/// Register the platform with the UEBuildPlatform class
+		/// </summary>
 		protected override void RegisterBuildPlatformInternal()
 		{
 			// Register this build platform for IOS
@@ -105,13 +117,11 @@ namespace UnrealBuildTool
 			}
 		}
 
-		/**
-		 *	Retrieve the CPPTargetPlatform for the given UnrealTargetPlatform
-		 *
-		 *	@param	InUnrealTargetPlatform		The UnrealTargetPlatform being build
-		 *	
-		 *	@return	CPPTargetPlatform			The CPPTargetPlatform to compile for
-		 */
+		/// <summary>
+		/// Retrieve the CPPTargetPlatform for the given UnrealTargetPlatform
+		/// </summary>
+		/// <param name="InUnrealTargetPlatform"> The UnrealTargetPlatform being build</param>
+		/// <returns>CPPTargetPlatform   The CPPTargetPlatform to compile for</returns>
 		public override CPPTargetPlatform GetCPPTargetPlatform(UnrealTargetPlatform InUnrealTargetPlatform)
 		{
 			switch (InUnrealTargetPlatform)
@@ -122,13 +132,11 @@ namespace UnrealBuildTool
 			throw new BuildException("IOSPlatform::GetCPPTargetPlatform: Invalid request for {0}", InUnrealTargetPlatform.ToString());
 		}
 
-		/**
-		 *	Get the extension to use for the given binary type
-		 *	
-		 *	@param	InBinaryType		The binary type being built
-		 *	
-		 *	@return	string				The binary extenstion (ie 'exe' or 'dll')
-		 */
+		/// <summary>
+		/// Get the extension to use for the given binary type
+		/// </summary>
+		/// <param name="InBinaryType"> The binary type being built</param>
+		/// <returns>string    The binary extenstion (ie 'exe' or 'dll')</returns>
 		public override string GetBinaryExtension(UEBuildBinaryType InBinaryType)
 		{
 			switch (InBinaryType)
@@ -269,11 +277,10 @@ namespace UnrealBuildTool
 			}
 		}
 
-		/**
-		 * Check for the default configuration
-		 *
-		 * return true if the project uses the default build config
-		 */
+		/// <summary>
+		/// Check for the default configuration
+		/// return true if the project uses the default build config
+		/// </summary>
 		public override bool HasDefaultBuildConfig(UnrealTargetPlatform Platform, DirectoryReference ProjectDirectoryName)
 		{
 			string[] BoolKeys = new string[] {
@@ -297,11 +304,10 @@ namespace UnrealBuildTool
 			return base.HasDefaultBuildConfig(Platform, ProjectDirectoryName);
 		}
 
-		/**
-         *	Setup the target environment for building
-         *	
-         *	@param	InBuildTarget		The target being built
-         */
+		/// <summary>
+		/// Setup the target environment for building
+		/// </summary>
+		/// <param name="InBuildTarget"> The target being built</param>
 		public override void SetUpEnvironment(UEBuildTarget InBuildTarget)
 		{
 			InBuildTarget.GlobalCompileEnvironment.Config.Definitions.Add("PLATFORM_IOS=1");
@@ -334,13 +340,12 @@ namespace UnrealBuildTool
 			InBuildTarget.GlobalLinkEnvironment.Config.AdditionalFrameworks.Add(new UEBuildFramework("StoreKit"));
 		}
 
-		/**
-		*	Whether the editor should be built for this platform or not
-		*	
-		*	@param	InPlatform		The UnrealTargetPlatform being built
-		*	@param	InConfiguration	The UnrealTargetConfiguration being built
-		*	@return	bool			true if the editor should be built, false if not
-		*/
+		/// <summary>
+		/// Whether the editor should be built for this platform or not
+		/// </summary>
+		/// <param name="InPlatform"> The UnrealTargetPlatform being built</param>
+		/// <param name="InConfiguration">The UnrealTargetConfiguration being built</param>
+		/// <returns>bool   true if the editor should be built, false if not</returns>
 		public override bool ShouldNotBuildEditor(UnrealTargetPlatform InPlatform, UnrealTargetConfiguration InConfiguration)
 		{
 			return true;
@@ -351,14 +356,12 @@ namespace UnrealBuildTool
 			return true; // for iOS can only run cooked. this is mostly for testing console code paths.
 		}
 
-		/**
-		 *	Whether this platform should create debug information or not
-		 *	
-		 *	@param	InPlatform			The UnrealTargetPlatform being built
-		 *	@param	InConfiguration		The UnrealTargetConfiguration being built
-		 *	
-		 *	@return	bool				true if debug info should be generated, false if not
-		 */
+		/// <summary>
+		/// Whether this platform should create debug information or not
+		/// </summary>
+		/// <param name="InPlatform">  The UnrealTargetPlatform being built</param>
+		/// <param name="InConfiguration"> The UnrealTargetConfiguration being built</param>
+		/// <returns>bool    true if debug info should be generated, false if not</returns>
 		public override bool ShouldCreateDebugInfo(UnrealTargetPlatform Platform, UnrealTargetConfiguration Configuration)
 		{
 			return true;
@@ -400,26 +403,24 @@ namespace UnrealBuildTool
 			BuildConfiguration.bDeployAfterCompile = true;
 		}
 
-		/**
-		 *	Whether the platform requires the extra UnityCPPWriter
-		 *	This is used to add an extra file for UBT to get the #include dependencies from
-		 *	
-		 *	@return	bool	true if it is required, false if not
-		 */
+		/// <summary>
+		/// Whether the platform requires the extra UnityCPPWriter
+		/// This is used to add an extra file for UBT to get the #include dependencies from
+		/// </summary>
+		/// <returns>bool true if it is required, false if not</returns>
 		public override bool RequiresExtraUnityCPPWriter()
 		{
 			return true;
 		}
 
-		/**
-		 *     Modify the newly created module passed in for this platform.
-		 *     This is not required - but allows for hiding details of a
-		 *     particular platform.
-		 *     
-		 *  @param  Name			The name of the module
-		 *	@param	Rules			The module rules
-		 *	@param	Target			The target being build
-		 */
+		/// <summary>
+		/// Modify the newly created module passed in for this platform.
+		/// This is not required - but allows for hiding details of a
+		/// particular platform.
+		/// </summary>
+		/// <param name="">Name   The name of the module</param>
+		/// <param name="Rules">  The module rules</param>
+		/// <param name="Target">  The target being build</param>
 		public override void ModifyModuleRules(string ModuleName, ModuleRules Rules, TargetInfo Target)
 		{
 			if ((Target.Platform == UnrealTargetPlatform.Win32) || (Target.Platform == UnrealTargetPlatform.Win64) || (Target.Platform == UnrealTargetPlatform.Mac))
diff --git a/Engine/Source/Programs/UnrealBuildTool/IOS/UEDeployIOS.cs b/Engine/Source/Programs/UnrealBuildTool/IOS/UEDeployIOS.cs
--- a/Engine/Source/Programs/UnrealBuildTool/IOS/UEDeployIOS.cs
+++ b/Engine/Source/Programs/UnrealBuildTool/IOS/UEDeployIOS.cs
@@ -12,9 +12,9 @@ namespace UnrealBuildTool.IOS
 {
 	public class UEDeployIOS : UEBuildDeploy
 	{
-		/**
-		 *	Register the platform with the UEBuildDeploy class
-		 */
+		/// <summary>
+		/// Register the platform with the UEBuildDeploy class
+		/// </summary>
 		public override void RegisterBuildDeploy()
 		{
 			// TODO: print debug info and handle any cases that would keep this from registering
diff --git a/Engine/Source/Programs/UnrealBuildTool/Linux/LinuxToolChain.cs b/Engine/Source/Programs/UnrealBuildTool/Linux/LinuxToolChain.cs
--- a/Engine/Source/Programs/UnrealBuildTool/Linux/LinuxToolChain.cs
+++ b/Engine/Source/Programs/UnrealBuildTool/Linux/LinuxToolChain.cs
@@ -44,7 +44,9 @@ namespace UnrealBuildTool
 			return null;
 		}
 
-		/** Splits compiler version string into numerical components, leaving unchanged if not known */
+		/// <summary>
+		/// Splits compiler version string into numerical components, leaving unchanged if not known
+		/// </summary>
 		private void DetermineCompilerMajMinPatchFromVersionString()
 		{
 			string[] Parts = CompilerVersionString.Split('.');
@@ -62,7 +64,9 @@ namespace UnrealBuildTool
 			}
 		}
 
-		/** Queries compiler for the version */
+		/// <summary>
+		/// Queries compiler for the version
+		/// </summary>
 		private bool DetermineCompilerVersion()
 		{
 			CompilerVersionString = null;
@@ -202,7 +206,9 @@ namespace UnrealBuildTool
 			UEToolChain.RegisterPlatformToolChain(CPPTargetPlatform.Linux, this);
 		}
 
-		/** Checks if compiler version matches the requirements */
+		/// <summary>
+		/// Checks if compiler version matches the requirements
+		/// </summary>
 		private static bool CompilerVersionGreaterOrEqual(int Major, int Minor, int Patch)
 		{
 			return CompilerVersionMajor > Major ||
@@ -210,7 +216,9 @@ namespace UnrealBuildTool
 				(CompilerVersionMajor == Major && CompilerVersionMinor == Minor && CompilerVersionPatch >= Patch);
 		}
 
-		/** Architecture-specific compiler switches */
+		/// <summary>
+		/// Architecture-specific compiler switches
+		/// </summary>
 		static string ArchitectureSpecificSwitches(string Architecture)
 		{
 			string Result = "";
@@ -238,7 +246,9 @@ namespace UnrealBuildTool
 			return Result;
 		}
 
-		/** Gets architecture-specific ar paths */
+		/// <summary>
+		/// Gets architecture-specific ar paths
+		/// </summary>
 		private static string GetArPath(string Architecture)
 		{
 			if (CrossCompiling())
@@ -249,7 +259,9 @@ namespace UnrealBuildTool
 			return ArPath;
 		}
 
-		/** Gets architecture-specific ranlib paths */
+		/// <summary>
+		/// Gets architecture-specific ranlib paths
+		/// </summary>
 		private static string GetRanlibPath(string Architecture)
 		{
 			if (CrossCompiling())
@@ -260,7 +272,9 @@ namespace UnrealBuildTool
 			return RanlibPath;
 		}
 
-		/** Gets architecture-specific strip path */
+		/// <summary>
+		/// Gets architecture-specific strip path
+		/// </summary>
 		private static string GetStripPath(string Architecture)
 		{
 			if (CrossCompiling())
@@ -572,16 +586,26 @@ namespace UnrealBuildTool
 		static string RanlibPath;
 		static string StripPath;
 
-		/** Version string of the current compiler, whether clang or gcc or whatever */
+		/// <summary>
+		/// Version string of the current compiler, whether clang or gcc or whatever
+		/// </summary>
 		static string CompilerVersionString;
-		/** Major version of the current compiler, whether clang or gcc or whatever */
+		/// <summary>
+		/// Major version of the current compiler, whether clang or gcc or whatever
+		/// </summary>
 		static int CompilerVersionMajor = -1;
-		/** Minor version of the current compiler, whether clang or gcc or whatever */
+		/// <summary>
+		/// Minor version of the current compiler, whether clang or gcc or whatever
+		/// </summary>
 		static int CompilerVersionMinor = -1;
-		/** Patch version of the current compiler, whether clang or gcc or whatever */
+		/// <summary>
+		/// Patch version of the current compiler, whether clang or gcc or whatever
+		/// </summary>
 		static int CompilerVersionPatch = -1;
 
-		/** Track which scripts need to be deleted before appending to */
+		/// <summary>
+		/// Track which scripts need to be deleted before appending to
+		/// </summary>
 		private bool bHasWipedFixDepsScript = false;
 
 		private static List<FileItem> BundleDependencies = new List<FileItem>();
@@ -728,7 +752,9 @@ namespace UnrealBuildTool
 			return Result;
 		}
 
-		/** Creates an action to archive all the .o files into single .a file */
+		/// <summary>
+		/// Creates an action to archive all the .o files into single .a file
+		/// </summary>
 		public FileItem CreateArchiveAndIndex(LinkEnvironment LinkEnvironment)
 		{
 			// Create an archive action
@@ -1100,7 +1126,9 @@ namespace UnrealBuildTool
 			}
 		}
 
-		/** Converts the passed in path from UBT host to compiler native format. */
+		/// <summary>
+		/// Converts the passed in path from UBT host to compiler native format.
+		/// </summary>
 		public override string ConvertPath(string OriginalPath)
 		{
 			if (BuildHostPlatform.Current.Platform == UnrealTargetPlatform.Linux)
diff --git a/Engine/Source/Programs/UnrealBuildTool/Linux/UEBuildLinux.cs b/Engine/Source/Programs/UnrealBuildTool/Linux/UEBuildLinux.cs
--- a/Engine/Source/Programs/UnrealBuildTool/Linux/UEBuildLinux.cs
+++ b/Engine/Source/Programs/UnrealBuildTool/Linux/UEBuildLinux.cs
@@ -10,61 +10,66 @@ namespace UnrealBuildTool
 {
 	class LinuxPlatform : UEBuildPlatform
 	{
-		/** This is the SDK version we support */
+		/// <summary>
+		/// This is the SDK version we support
+		/// </summary>
 		static private Dictionary<string, string> ExpectedSDKVersions = new Dictionary<string, string>()
 		{
 			{ "x86_64-unknown-linux-gnu", "v6_clang-3.6.0_ld-2.24_glibc-2.12.2" },
 			{ "arm-unknown-linux-gnueabihf", "arm-unknown-linux-gnueabihf_v5_clang-3.5.0-ld-2.23.1-glibc-2.13" },
 		};
 
-		/** Platform name (embeds architecture for now) */
+		/// <summary>
+		/// Platform name (embeds architecture for now)
+		/// </summary>
 		static private string TargetPlatformName = "Linux_x64";
 
-		/** Linux architecture (compiler target triplet) */
+		/// <summary>
+		/// Linux architecture (compiler target triplet)
+		/// </summary>
 		// FIXME: for now switching between architectures is hard-coded
 		static private string DefaultArchitecture = "x86_64-unknown-linux-gnu";
 		//static private string DefaultArchitecture = "arm-unknown-linux-gnueabihf";
 
-		/** The current architecture */
+		/// <summary>
+		/// The current architecture
+		/// </summary>
 		public override string GetActiveArchitecture()
 		{
 			return DefaultArchitecture;
 		}
 
-		/**
-		 * Allow the platform to apply architecture-specific name according to its rules
-		 */
+		/// <summary>
+		/// Allow the platform to apply architecture-specific name according to its rules
+		/// </summary>
 		public override string ApplyArchitectureName(string BinaryName)
 		{
 			// Linux ignores architecture-specific names, although it might be worth it to prepend architecture
 			return BinaryName;
 		}
 
-		/** 
-		 * Whether platform supports switching SDKs during runtime
-		 * 
-		 * @return true if supports
-		 */
+		/// <summary>
+		/// Whether platform supports switching SDKs during runtime
+		/// </summary>
+		/// <returns>true if supports</returns>
 		protected override bool PlatformSupportsAutoSDKs()
 		{
 			return true;
 		}
 
-		/** 
-		 * Returns platform-specific name used in SDK repository
-		 * 
-		 * @return path to SDK Repository
-		 */
+		/// <summary>
+		/// Returns platform-specific name used in SDK repository
+		/// </summary>
+		/// <returns>path to SDK Repository</returns>
 		public override string GetSDKTargetPlatformName()
 		{
 			return TargetPlatformName;
 		}
 
-		/** 
-		 * Returns SDK string as required by the platform 
-		 * 
-		 * @return Valid SDK string
-		 */
+		/// <summary>
+		/// Returns SDK string as required by the platform
+		/// </summary>
+		/// <returns>Valid SDK string</returns>
 		protected override string GetRequiredSDKString()
 		{
 			string SDKString;
@@ -87,9 +92,9 @@ namespace UnrealBuildTool
 			return true;
 		}
 
-		/**
-		 *	Whether the required external SDKs are installed for this platform
-		 */
+		/// <summary>
+		/// Whether the required external SDKs are installed for this platform
+		/// </summary>
 		protected override SDKStatus HasRequiredManualSDKInternal()
 		{
 			if (BuildHostPlatform.Current.Platform == UnrealTargetPlatform.Linux)
@@ -119,9 +124,9 @@ namespace UnrealBuildTool
 			return false;
 		}
 
-		/**
-		 *	Register the platform with the UEBuildPlatform class
-		 */
+		/// <summary>
+		/// Register the platform with the UEBuildPlatform class
+		/// </summary>
 		protected override void RegisterBuildPlatformInternal()
 		{
 			if ((ProjectFileGenerator.bGenerateProjectFiles == true) || (HasRequiredSDKsInstalled() == SDKStatus.Valid))
@@ -149,13 +154,11 @@ namespace UnrealBuildTool
 			}
 		}
 
-		/**
-		 *	Retrieve the CPPTargetPlatform for the given UnrealTargetPlatform
-		 *
-		 *	@param	InUnrealTargetPlatform		The UnrealTargetPlatform being build
-		 *	
-		 *	@return	CPPTargetPlatform			The CPPTargetPlatform to compile for
-		 */
+		/// <summary>
+		/// Retrieve the CPPTargetPlatform for the given UnrealTargetPlatform
+		/// </summary>
+		/// <param name="InUnrealTargetPlatform"> The UnrealTargetPlatform being build</param>
+		/// <returns>CPPTargetPlatform   The CPPTargetPlatform to compile for</returns>
 		public override CPPTargetPlatform GetCPPTargetPlatform(UnrealTargetPlatform InUnrealTargetPlatform)
 		{
 			switch (InUnrealTargetPlatform)
@@ -166,13 +169,11 @@ namespace UnrealBuildTool
 			throw new BuildException("LinuxPlatform::GetCPPTargetPlatform: Invalid request for {0}", InUnrealTargetPlatform.ToString());
 		}
 
-		/**
-		 *	Get the extension to use for the given binary type
-		 *	
-		 *	@param	InBinaryType		The binary type being built
-		 *	
-		 *	@return	string				The binary extension (i.e. 'exe' or 'dll')
-		 */
+		/// <summary>
+		/// Get the extension to use for the given binary type
+		/// </summary>
+		/// <param name="InBinaryType"> The binary type being built</param>
+		/// <returns>string    The binary extension (i.e. 'exe' or 'dll')</returns>
 		public override string GetBinaryExtension(UEBuildBinaryType InBinaryType)
 		{
 			switch (InBinaryType)
@@ -191,48 +192,44 @@ namespace UnrealBuildTool
 			return base.GetBinaryExtension(InBinaryType);
 		}
 
-		/**
-		 *	Get the extension to use for debug info for the given binary type
-		 *	
-		 *	@param	InBinaryType		The binary type being built
-		 *	
-		 *	@return	string				The debug info extension (i.e. 'pdb')
-		 */
+		/// <summary>
+		/// Get the extension to use for debug info for the given binary type
+		/// </summary>
+		/// <param name="InBinaryType"> The binary type being built</param>
+		/// <returns>string    The debug info extension (i.e. 'pdb')</returns>
 		public override string GetDebugInfoExtension(UEBuildBinaryType InBinaryType)
 		{
 			return "";
 		}
 
-		/**
-		 *	Gives the platform a chance to 'override' the configuration settings 
-		 *	that are overridden on calls to RunUBT.
-		 *	
-		 *	@param	InPlatform			The UnrealTargetPlatform being built
-		 *	@param	InConfiguration		The UnrealTargetConfiguration being built
-		 */
+		/// <summary>
+		/// Gives the platform a chance to 'override' the configuration settings
+		/// that are overridden on calls to RunUBT.
+		/// </summary>
+		/// <param name="InPlatform">  The UnrealTargetPlatform being built</param>
+		/// <param name="InConfiguration"> The UnrealTargetConfiguration being built</param>
 		public override void ResetBuildConfiguration(UnrealTargetPlatform InPlatform, UnrealTargetConfiguration InConfiguration)
 		{
 			ValidateUEBuildConfiguration();
 		}
 
-		/**
-		 *	Validate configuration for this platform
-		 *	NOTE: This function can/will modify BuildConfiguration!
-		 *	
-		 *	@param	InPlatform			The CPPTargetPlatform being built
-		 *	@param	InConfiguration		The CPPTargetConfiguration being built
-		 *	@param	bInCreateDebugInfo	true if debug info is getting create, false if not
-		 */
+		/// <summary>
+		/// Validate configuration for this platform
+		/// NOTE: This function can/will modify BuildConfiguration!
+		/// </summary>
+		/// <param name="InPlatform">  The CPPTargetPlatform being built</param>
+		/// <param name="InConfiguration"> The CPPTargetConfiguration being built</param>
+		/// <param name="bInCreateDebugInfo">true if debug info is getting create, false if not</param>
 		public override void ValidateBuildConfiguration(CPPTargetConfiguration Configuration, CPPTargetPlatform Platform, bool bCreateDebugInfo)
 		{
 			UEBuildConfiguration.bCompileSimplygon = false;
 		}
 
-		/**
-		 *	Validate the UEBuildConfiguration for this platform
-		 *	This is called BEFORE calling UEBuildConfiguration to allow setting 
-		 *	various fields used in that function such as CompileLeanAndMean...
-		 */
+		/// <summary>
+		/// Validate the UEBuildConfiguration for this platform
+		/// This is called BEFORE calling UEBuildConfiguration to allow setting
+		/// various fields used in that function such as CompileLeanAndMean...
+		/// </summary>
 		public override void ValidateUEBuildConfiguration()
 		{
 			if (ProjectFileGenerator.bGenerateProjectFiles && !ProjectFileGenerator.bGeneratingRocketProjectFiles)
@@ -250,41 +247,37 @@ namespace UnrealBuildTool
 			BuildConfiguration.bUseSharedPCHs = false;
 		}
 
-		/**
-		 *	Whether PDB files should be used
-		 *	
-		 *	@param	InPlatform			The CPPTargetPlatform being built
-		 *	@param	InConfiguration		The CPPTargetConfiguration being built
-		 *	@param	bInCreateDebugInfo	true if debug info is getting create, false if not
-		 *	
-		 *	@return	bool	true if PDB files should be used, false if not
-		 */
+		/// <summary>
+		/// Whether PDB files should be used
+		/// </summary>
+		/// <param name="InPlatform">  The CPPTargetPlatform being built</param>
+		/// <param name="InConfiguration"> The CPPTargetConfiguration being built</param>
+		/// <param name="bInCreateDebugInfo">true if debug info is getting create, false if not</param>
+		/// <returns>bool true if PDB files should be used, false if not</returns>
 		public override bool ShouldUsePDBFiles(CPPTargetPlatform Platform, CPPTargetConfiguration Configuration, bool bCreateDebugInfo)
 		{
 			return true;
 		}
 
-		/**
-		 *	Get a list of extra modules the platform requires.
-		 *	This is to allow undisclosed platforms to add modules they need without exposing information about the platfomr.
-		 *	
-		 *	@param	Target						The target being build
-		 *	@param	BuildTarget					The UEBuildTarget getting build
-		 *	@param	PlatformExtraModules		OUTPUT the list of extra modules the platform needs to add to the target
-		 */
+		/// <summary>
+		/// Get a list of extra modules the platform requires.
+		/// This is to allow undisclosed platforms to add modules they need without exposing information about the platfomr.
+		/// </summary>
+		/// <param name="Target">     The target being build</param>
+		/// <param name="BuildTarget">    The UEBuildTarget getting build</param>
+		/// <param name="PlatformExtraModules"> OUTPUT the list of extra modules the platform needs to add to the target</param>
 		public override void GetExtraModules(TargetInfo Target, UEBuildTarget BuildTarget, ref List<string> PlatformExtraModules)
 		{
 		}
 
-		/**
-		 *	Modify the newly created module passed in for this platform.
-		 *	This is not required - but allows for hiding details of a
-		 *	particular platform.
-		 *	
-		 *  @param  Name			The name of the module
-		 *	@param	Rules			The module rules
-		 *	@param	Target			The target being build
-		 */
+		/// <summary>
+		/// Modify the newly created module passed in for this platform.
+		/// This is not required - but allows for hiding details of a
+		/// particular platform.
+		/// </summary>
+		/// <param name="">Name   The name of the module</param>
+		/// <param name="Rules">  The module rules</param>
+		/// <param name="Target">  The target being build</param>
 		public override void ModifyModuleRules(string ModuleName, ModuleRules Rules, TargetInfo Target)
 		{
 			if ((Target.Platform == UnrealTargetPlatform.Win32) || (Target.Platform == UnrealTargetPlatform.Win64))
@@ -342,11 +335,10 @@ namespace UnrealBuildTool
 			}
 		}
 
-		/**
-		 *	Setup the target environment for building
-		 *	
-		 *	@param	InBuildTarget		The target being built
-		 */
+		/// <summary>
+		/// Setup the target environment for building
+		/// </summary>
+		/// <param name="InBuildTarget"> The target being built</param>
 		public override void SetUpEnvironment(UEBuildTarget InBuildTarget)
 		{
 			InBuildTarget.GlobalCompileEnvironment.Config.Definitions.Add("PLATFORM_LINUX=1");
@@ -375,14 +367,12 @@ namespace UnrealBuildTool
 			}
 		}
 
-		/**
-		 *	Whether this platform should create debug information or not
-		 *	
-		 *	@param	InPlatform			The UnrealTargetPlatform being built
-		 *	@param	InConfiguration		The UnrealTargetConfiguration being built
-		 *	
-		 *	@return	bool				true if debug info should be generated, false if not
-		 */
+		/// <summary>
+		/// Whether this platform should create debug information or not
+		/// </summary>
+		/// <param name="InPlatform">  The UnrealTargetPlatform being built</param>
+		/// <param name="InConfiguration"> The UnrealTargetConfiguration being built</param>
+		/// <returns>bool    true if debug info should be generated, false if not</returns>
 		public override bool ShouldCreateDebugInfo(UnrealTargetPlatform Platform, UnrealTargetConfiguration Configuration)
 		{
 			switch (Configuration)
@@ -396,11 +386,10 @@ namespace UnrealBuildTool
 			};
 		}
 
-		/**
-		 *	Setup the binaries for this specific platform.
-		 *	
-		 *	@param	InBuildTarget		The target being built
-		 */
+		/// <summary>
+		/// Setup the binaries for this specific platform.
+		/// </summary>
+		/// <param name="InBuildTarget"> The target being built</param>
 		public override void SetupBinaries(UEBuildTarget InBuildTarget)
 		{
 		}
diff --git a/Engine/Source/Programs/UnrealBuildTool/Mac/MacToolChain.cs b/Engine/Source/Programs/UnrealBuildTool/Mac/MacToolChain.cs
--- a/Engine/Source/Programs/UnrealBuildTool/Mac/MacToolChain.cs
+++ b/Engine/Source/Programs/UnrealBuildTool/Mac/MacToolChain.cs
@@ -24,37 +24,57 @@ namespace UnrealBuildTool
 		 *  Do NOT change the defaults to set your values, instead you should set the environment variables
 		 *  properly in your system, as other tools make use of them to work properly!
 		 *  The defaults are there simply for examples so you know what to put in your env vars...
-		 ***********************************************************************
+		 ***********************************************************************/
 
-		/** Which version of the Mac OS SDK to target at build time */
+		/// <summary>
+		/// Which version of the Mac OS SDK to target at build time
+		/// </summary>
 		public static string MacOSSDKVersion = "latest";
 		public static float MacOSSDKVersionFloat = 0.0f;
 
-		/** Which version of the Mac OS X to allow at run time */
+		/// <summary>
+		/// Which version of the Mac OS X to allow at run time
+		/// </summary>
 		public static string MacOSVersion = "10.9";
 
-		/** Minimum version of Mac OS X to actually run on, running on earlier versions will display the system minimum version error dialog & exit. */
+		/// <summary>
+		/// Minimum version of Mac OS X to actually run on, running on earlier versions will display the system minimum version error dialog & exit.
+		/// </summary>
 		public static string MinMacOSVersion = "10.9.2";
 
-		/** Which developer directory to root from? If this is "xcode-select", UBT will query for the currently selected Xcode */
+		/// <summary>
+		/// Which developer directory to root from? If this is "xcode-select", UBT will query for the currently selected Xcode
+		/// </summary>
 		private static string XcodeDeveloperDir = "xcode-select";
 
-		/** Directory for the developer binaries */
+		/// <summary>
+		/// Directory for the developer binaries
+		/// </summary>
 		private static string ToolchainDir = "";
 
-		/** Location of the SDKs */
+		/// <summary>
+		/// Location of the SDKs
+		/// </summary>
 		private static string BaseSDKDir;
 
-		/** Which compiler frontend to use */
+		/// <summary>
+		/// Which compiler frontend to use
+		/// </summary>
 		private static string MacCompiler = "clang++";
 
-		/** Which linker frontend to use */
+		/// <summary>
+		/// Which linker frontend to use
+		/// </summary>
 		private static string MacLinker = "clang++";
 
-		/** Which archiver to use */
+		/// <summary>
+		/// Which archiver to use
+		/// </summary>
 		private static string MacArchiver = "libtool";
 
-		/** Track which scripts need to be deleted before appending to */
+		/// <summary>
+		/// Track which scripts need to be deleted before appending to
+		/// </summary>
 		private bool bHasWipedCopyDylibScript = false;
 		private bool bHasWipedFixDylibScript = false;
 
@@ -1127,11 +1147,10 @@ namespace UnrealBuildTool
 			}
 		}
 
-		/**
-		 * Generates debug info for a given executable
-		 * 
-		 * @param MachOBinary FileItem describing the executable or dylib to generate debug info for
-		 */
+		/// <summary>
+		/// Generates debug info for a given executable
+		/// </summary>
+		/// <param name="MachOBinary">FileItem describing the executable or dylib to generate debug info for</param>
 		public FileItem GenerateDebugInfo(FileItem MachOBinary)
 		{
 			// Make a file item for the source and destination files
@@ -1171,11 +1190,10 @@ namespace UnrealBuildTool
 			return DestFile;
 		}
 
-		/**
-		 * Creates app bundle for a given executable
-		 *
-		 * @param Executable FileItem describing the executable to generate app bundle for
-		 */
+		/// <summary>
+		/// Creates app bundle for a given executable
+		/// </summary>
+		/// <param name="Executable">FileItem describing the executable to generate app bundle for</param>
 		FileItem FinalizeAppBundle(LinkEnvironment LinkEnvironment, FileItem Executable, FileItem FixDylibOutputFile)
 		{
 			// Make a file item for the source and destination files
diff --git a/Engine/Source/Programs/UnrealBuildTool/Mac/UEBuildMac.cs b/Engine/Source/Programs/UnrealBuildTool/Mac/UEBuildMac.cs
--- a/Engine/Source/Programs/UnrealBuildTool/Mac/UEBuildMac.cs
+++ b/Engine/Source/Programs/UnrealBuildTool/Mac/UEBuildMac.cs
@@ -26,9 +26,9 @@ namespace UnrealBuildTool
 			return SDKStatus.Valid;
 		}
 
-		/**
-		 *	Register the platform with the UEBuildPlatform class
-		 */
+		/// <summary>
+		/// Register the platform with the UEBuildPlatform class
+		/// </summary>
 		protected override void RegisterBuildPlatformInternal()
 		{
 			// Register this build platform for Mac
@@ -38,13 +38,11 @@ namespace UnrealBuildTool
 			UEBuildPlatform.RegisterPlatformWithGroup(UnrealTargetPlatform.Mac, UnrealPlatformGroup.Apple);
 		}
 
-		/**
-		 *	Retrieve the CPPTargetPlatform for the given UnrealTargetPlatform
-		 *
-		 *	@param	InUnrealTargetPlatform		The UnrealTargetPlatform being build
-		 *	
-		 *	@return	CPPTargetPlatform			The CPPTargetPlatform to compile for
-		 */
+		/// <summary>
+		/// Retrieve the CPPTargetPlatform for the given UnrealTargetPlatform
+		/// </summary>
+		/// <param name="InUnrealTargetPlatform"> The UnrealTargetPlatform being build</param>
+		/// <returns>CPPTargetPlatform   The CPPTargetPlatform to compile for</returns>
 		public override CPPTargetPlatform GetCPPTargetPlatform(UnrealTargetPlatform InUnrealTargetPlatform)
 		{
 			switch (InUnrealTargetPlatform)
@@ -55,13 +53,11 @@ namespace UnrealBuildTool
 			throw new BuildException("MacPlatform::GetCPPTargetPlatform: Invalid request for {0}", InUnrealTargetPlatform.ToString());
 		}
 
-		/**
-		 *	Get the extension to use for the given binary type
-		 *	
-		 *	@param	InBinaryType		The binrary type being built
-		 *	
-		 *	@return	string				The binary extenstion (ie 'exe' or 'dll')
-		 */
+		/// <summary>
+		/// Get the extension to use for the given binary type
+		/// </summary>
+		/// <param name="InBinaryType"> The binrary type being built</param>
+		/// <returns>string    The binary extenstion (ie 'exe' or 'dll')</returns>
 		public override string GetBinaryExtension(UEBuildBinaryType InBinaryType)
 		{
 			switch (InBinaryType)
@@ -80,13 +76,11 @@ namespace UnrealBuildTool
 			return base.GetBinaryExtension(InBinaryType);
 		}
 
-		/**
-		 *	Get the extension to use for debug info for the given binary type
-		 *	
-		 *	@param	InBinaryType		The binary type being built
-		 *	
-		 *	@return	string				The debug info extension (i.e. 'pdb')
-		 */
+		/// <summary>
+		/// Get the extension to use for debug info for the given binary type
+		/// </summary>
+		/// <param name="InBinaryType"> The binary type being built</param>
+		/// <returns>string    The debug info extension (i.e. 'pdb')</returns>
 		public override string GetDebugInfoExtension(UEBuildBinaryType InBinaryType)
 		{
 			return BuildConfiguration.bGeneratedSYMFile || BuildConfiguration.bUsePDBFiles ? ".dSYM" : "";
@@ -127,11 +121,10 @@ namespace UnrealBuildTool
 			}
 		}
 
-		/**
-		 *	Setup the target environment for building
-		 *	
-		 *	@param	InBuildTarget		The target being built
-		 */
+		/// <summary>
+		/// Setup the target environment for building
+		/// </summary>
+		/// <param name="InBuildTarget"> The target being built</param>
 		public override void SetUpEnvironment(UEBuildTarget InBuildTarget)
 		{
 			InBuildTarget.GlobalCompileEnvironment.Config.Definitions.Add("PLATFORM_MAC=1");
@@ -156,14 +149,12 @@ namespace UnrealBuildTool
 
 		}
 
-		/**
-		 *	Whether this platform should create debug information or not
-		 *
-		 *	@param	InPlatform			The UnrealTargetPlatform being built
-		 *	@param	InConfiguration		The UnrealTargetConfiguration being built
-		 *	
-		 *	@return	bool				true if debug info should be generated, false if not
-		 */
+		/// <summary>
+		/// Whether this platform should create debug information or not
+		/// </summary>
+		/// <param name="InPlatform">  The UnrealTargetPlatform being built</param>
+		/// <param name="InConfiguration"> The UnrealTargetConfiguration being built</param>
+		/// <returns>bool    true if debug info should be generated, false if not</returns>
 		public override bool ShouldCreateDebugInfo(UnrealTargetPlatform Platform, UnrealTargetConfiguration Configuration)
 		{
 			return true;
@@ -202,28 +193,27 @@ namespace UnrealBuildTool
 			}
 		}
 
-		/**
-		 *	Whether the platform requires the extra UnityCPPWriter
-		 *	This is used to add an extra file for UBT to get the #include dependencies from
-		 *	
-		 *	@return	bool	true if it is required, false if not
-		 */
+		/// <summary>
+		/// Whether the platform requires the extra UnityCPPWriter
+		/// This is used to add an extra file for UBT to get the #include dependencies from
+		/// </summary>
+		/// <returns>bool true if it is required, false if not</returns>
 		public override bool RequiresExtraUnityCPPWriter()
 		{
 			return true;
 		}
 
-		/**
-		 *	Return whether we wish to have this platform's binaries in our builds
-		 */
+		/// <summary>
+		/// Return whether we wish to have this platform's binaries in our builds
+		/// </summary>
 		public override bool IsBuildRequired()
 		{
 			return false;
 		}
 
-		/**
-		 *	Return whether we wish to have this platform's binaries in our CIS tests
-		 */
+		/// <summary>
+		/// Return whether we wish to have this platform's binaries in our CIS tests
+		/// </summary>
 		public override bool IsCISRequired()
 		{
 			return false;
diff --git a/Engine/Source/Programs/UnrealBuildTool/Mac/UEDeployMac.cs b/Engine/Source/Programs/UnrealBuildTool/Mac/UEDeployMac.cs
--- a/Engine/Source/Programs/UnrealBuildTool/Mac/UEDeployMac.cs
+++ b/Engine/Source/Programs/UnrealBuildTool/Mac/UEDeployMac.cs
@@ -13,9 +13,9 @@ namespace UnrealBuildTool.IOS
 {
 	class UEDeployMac : UEBuildDeploy
 	{
-		/**
-		 *	Register the platform with the UEBuildDeploy class
-		 */
+		/// <summary>
+		/// Register the platform with the UEBuildDeploy class
+		/// </summary>
 		public override void RegisterBuildDeploy()
 		{
 			// TODO: print debug info and handle any cases that would keep this from registering
diff --git a/Engine/Source/Programs/UnrealBuildTool/System/ActionGraph.cs b/Engine/Source/Programs/UnrealBuildTool/System/ActionGraph.cs
--- a/Engine/Source/Programs/UnrealBuildTool/System/ActionGraph.cs
+++ b/Engine/Source/Programs/UnrealBuildTool/System/ActionGraph.cs
@@ -25,7 +25,9 @@ namespace UnrealBuildTool
 		Link,
 	}
 
-	/** A build action. */
+	/// <summary>
+	/// A build action.
+	/// </summary>
 	[Serializable]
 	public class Action : ISerializable
 	{
@@ -33,58 +35,90 @@ namespace UnrealBuildTool
 		/// Preparation and Assembly (serialized)
 		/// 
 
-		/** The type of this action (for debugging purposes). */
+		/// <summary>
+		/// The type of this action (for debugging purposes).
+		/// </summary>
 		public ActionType ActionType;
 
-		/** Every file this action depends on.  These files need to exist and be up to date in order for this action to even be considered */
+		/// <summary>
+		/// Every file this action depends on.  These files need to exist and be up to date in order for this action to even be considered
+		/// </summary>
 		public List<FileItem> PrerequisiteItems = new List<FileItem>();
 
-		/** The files that this action produces after completing */
+		/// <summary>
+		/// The files that this action produces after completing
+		/// </summary>
 		public List<FileItem> ProducedItems = new List<FileItem>();
 
-		/** Directory from which to execute the program to create produced items */
+		/// <summary>
+		/// Directory from which to execute the program to create produced items
+		/// </summary>
 		public string WorkingDirectory = null;
 
-		/** True if we should log extra information when we run a program to create produced items */
+		/// <summary>
+		/// True if we should log extra information when we run a program to create produced items
+		/// </summary>
 		public bool bPrintDebugInfo = false;
 
-		/** The command to run to create produced items */
+		/// <summary>
+		/// The command to run to create produced items
+		/// </summary>
 		public string CommandPath = null;
 
-		/** Command-line parameters to pass to the program */
+		/// <summary>
+		/// Command-line parameters to pass to the program
+		/// </summary>
 		public string CommandArguments = null;
 
-		/** Optional friendly description of the type of command being performed, for example "Compile" or "Link".  Displayed by some executors. */
+		/// <summary>
+		/// Optional friendly description of the type of command being performed, for example "Compile" or "Link".  Displayed by some executors.
+		/// </summary>
 		public string CommandDescription = null;
 
-		/** Human-readable description of this action that may be displayed as status while invoking the action.  This is often the name of the file being compiled, or an executable file name being linked.  Displayed by some executors. */
+		/// <summary>
+		/// Human-readable description of this action that may be displayed as status while invoking the action.  This is often the name of the file being compiled, or an executable file name being linked.  Displayed by some executors.
+		/// </summary>
 		public string StatusDescription = "...";
 
-		/** True if this action is allowed to be run on a remote machine when a distributed build system is being used, such as XGE */
+		/// <summary>
+		/// True if this action is allowed to be run on a remote machine when a distributed build system is being used, such as XGE
+		/// </summary>
 		public bool bCanExecuteRemotely = false;
 
-		/** True if this action is using the GCC compiler.  Some build systems may be able to optimize for this case. */
+		/// <summary>
+		/// True if this action is using the GCC compiler.  Some build systems may be able to optimize for this case.
+		/// </summary>
 		public bool bIsGCCCompiler = false;
 
-		/** Whether the action is using a pre-compiled header to speed it up. */
+		/// <summary>
+		/// Whether the action is using a pre-compiled header to speed it up.
+		/// </summary>
 		public bool bIsUsingPCH = false;
 
-		/** Whether the files in ProducedItems should be deleted before executing this action, when the action is outdated */
+		/// <summary>
+		/// Whether the files in ProducedItems should be deleted before executing this action, when the action is outdated
+		/// </summary>
 		public bool bShouldDeleteProducedItems = false;
 
-		/**
-		 * Whether we should log this action, whether executed locally or remotely.  This is useful for actions that take time
-		 * but invoke tools without any console output.
-		 */
+		/// <summary>
+		/// Whether we should log this action, whether executed locally or remotely.  This is useful for actions that take time
+		/// but invoke tools without any console output.
+		/// </summary>
 		public bool bShouldOutputStatusDescription = true;
 
-		/** True if any libraries produced by this action should be considered 'import libraries' */
+		/// <summary>
+		/// True if any libraries produced by this action should be considered 'import libraries'
+		/// </summary>
 		public bool bProducesImportLibrary = false;
 
-		/** Optional custom event handler for standard output. */
+		/// <summary>
+		/// Optional custom event handler for standard output.
+		/// </summary>
 		public DataReceivedEventHandler OutputEventHandler = null;	// @todo ubtmake urgent: Delegate variables are not saved, but we are comparing against this in ExecuteActions() for XGE!
 
-		/** Callback used to perform a special action instead of a generic command line */
+		/// <summary>
+		/// Callback used to perform a special action instead of a generic command line
+		/// </summary>
 		public delegate void BlockingActionHandler(Action Action, out int ExitCode, out string Output);
 		public BlockingActionHandler ActionHandler = null;		// @todo ubtmake urgent: Delegate variables are not saved, but we are comparing against this in ExecuteActions() for XGE!
 
@@ -94,16 +128,24 @@ namespace UnrealBuildTool
 		/// Preparation only (not serialized)
 		///
 
-		/** Unique action identifier.  Used for displaying helpful info about detected cycles in the graph. */
+		/// <summary>
+		/// Unique action identifier.  Used for displaying helpful info about detected cycles in the graph.
+		/// </summary>
 		public int UniqueId;
 
-		/** Always-incremented unique id */
+		/// <summary>
+		/// Always-incremented unique id
+		/// </summary>
 		private static int NextUniqueId = 0;
 
-		/** Total number of actions depending on this one. */
+		/// <summary>
+		/// Total number of actions depending on this one.
+		/// </summary>
 		public int NumTotalDependentActions = 0;
 
-		/** Relative cost of producing items for this action. */
+		/// <summary>
+		/// Relative cost of producing items for this action.
+		/// </summary>
 		public long RelativeCost = 0;
 
 
@@ -111,10 +153,14 @@ namespace UnrealBuildTool
 		/// Assembly only (not serialized)
 		///
 
-		/** Start time of action, optionally set by executor. */
+		/// <summary>
+		/// Start time of action, optionally set by executor.
+		/// </summary>
 		public DateTimeOffset StartTime = DateTimeOffset.MinValue;
 
-		/** End time of action, optionally set by executor. */
+		/// <summary>
+		/// End time of action, optionally set by executor.
+		/// </summary>
 		public DateTimeOffset EndTime = DateTimeOffset.MinValue;
 
 
@@ -146,7 +192,9 @@ namespace UnrealBuildTool
 			ProducedItems = (List<FileItem>)SerializationInfo.GetValue("pd", typeof(List<FileItem>));
 		}
 
-		/** ISerializable: Called when serialized to report additional properties that should be saved */
+		/// <summary>
+		/// ISerializable: Called when serialized to report additional properties that should be saved
+		/// </summary>
 		public void GetObjectData(SerializationInfo SerializationInfo, StreamingContext StreamingContext)
 		{
 			SerializationInfo.AddValue("at", (byte)ActionType);
@@ -166,12 +214,11 @@ namespace UnrealBuildTool
 			SerializationInfo.AddValue("pd", ProducedItems);
 		}
 
-		/**
-		 * Compares two actions based on total number of dependent items, descending.
-		 * 
-		 * @param	A	Action to compare
-		 * @param	B	Action to compare
-		 */
+		/// <summary>
+		/// Compares two actions based on total number of dependent items, descending.
+		/// </summary>
+		/// <param name="A">Action to compare</param>
+		/// <param name="B">Action to compare</param>
 		public static int Compare(Action A, Action B)
 		{
 			// Primary sort criteria is total number of dependent files, up to max depth.
@@ -243,7 +290,9 @@ namespace UnrealBuildTool
 			SortActionList();
 		}
 
-		/** Builds a list of actions that need to be executed to produce the specified output items. */
+		/// <summary>
+		/// Builds a list of actions that need to be executed to produce the specified output items.
+		/// </summary>
 		public static List<Action> GetActionsToExecute(Action[] PrerequisiteActions, List<UEBuildTarget> Targets, out Dictionary<UEBuildTarget, List<FileItem>> TargetToOutdatedPrerequisitesMap)
 		{
 			var CheckOutdatednessStartTime = DateTime.UtcNow;
@@ -316,7 +365,9 @@ namespace UnrealBuildTool
 			return ActionsToExecute;
 		}
 
-		/** Executes a list of actions. */
+		/// <summary>
+		/// Executes a list of actions.
+		/// </summary>
 		public static bool ExecuteActions(List<Action> ActionsToExecute, out string ExecutorName)
 		{
 			bool Result = true;
@@ -436,7 +487,9 @@ namespace UnrealBuildTool
 			return Result;
 		}
 
-		/** Links actions with their prerequisite and produced items into an action graph. */
+		/// <summary>
+		/// Links actions with their prerequisite and produced items into an action graph.
+		/// </summary>
 		static void LinkActionsAndItems()
 		{
 			foreach (Action Action in AllActions)
@@ -487,7 +540,9 @@ namespace UnrealBuildTool
 		}
 
 
-		/** Finds and deletes stale hot reload DLLs. */
+		/// <summary>
+		/// Finds and deletes stale hot reload DLLs.
+		/// </summary>
 		public static void DeleteStaleHotReloadDLLs()
 		{
 			var DeleteStartTime = DateTime.UtcNow;
@@ -560,9 +615,9 @@ namespace UnrealBuildTool
 			}
 		}
 
-		/**
-		 * Sorts the action list for improved parallelism with local execution.
-		 */
+		/// <summary>
+		/// Sorts the action list for improved parallelism with local execution.
+		/// </summary>
 		public static void SortActionList()
 		{
 			// Mapping from action to a list of actions that directly or indirectly depend on it (up to a certain depth).
@@ -609,7 +664,9 @@ namespace UnrealBuildTool
 			AllActions.Sort(Action.Compare);
 		}
 
-		/** Checks for cycles in the action graph. */
+		/// <summary>
+		/// Checks for cycles in the action graph.
+		/// </summary>
 		static void DetectActionGraphCycles()
 		{
 			// Starting with actions that only depend on non-produced items, iteratively expand a set of actions that are only dependent on
@@ -715,11 +772,11 @@ namespace UnrealBuildTool
 			}
 		}
 
-		/**
-		 * Determines the full set of actions that must be built to produce an item.
-		 * @param OutputItem - The item to be built.
-		 * @param PrerequisiteActions - The actions that must be built and the root action are 
-		 */
+		/// <summary>
+		/// Determines the full set of actions that must be built to produce an item.
+		/// </summary>
+		/// <param name="OutputItem">- The item to be built.</param>
+		/// <param name="PrerequisiteActions">- The actions that must be built and the root action are</param>
 		public static void GatherPrerequisiteActions(
 			FileItem OutputItem,
 			ref HashSet<Action> PrerequisiteActions
@@ -738,13 +795,13 @@ namespace UnrealBuildTool
 			}
 		}
 
-		/**
-		 * Determines whether an action is outdated based on the modification times for its prerequisite
-		 * and produced items.
-		 * @param RootAction - The action being considered.
-		 * @param OutdatedActionDictionary - 
-		 * @return true if outdated
-		 */
+		/// <summary>
+		/// Determines whether an action is outdated based on the modification times for its prerequisite
+		/// and produced items.
+		/// </summary>
+		/// <param name="RootAction">- The action being considered.</param>
+		/// <param name="OutdatedActionDictionary">-</param>
+		/// <returns>true if outdated</returns>
 		static public bool IsActionOutdated(UEBuildTarget Target, Action RootAction, ref Dictionary<Action, bool> OutdatedActionDictionary, ActionHistory ActionHistory, Dictionary<UEBuildTarget, List<FileItem>> TargetToOutdatedPrerequisitesMap)
 		{
 			// Only compute the outdated-ness for actions that don't aren't cached in the outdated action dictionary.
@@ -1006,10 +1063,10 @@ namespace UnrealBuildTool
 		}
 
 
-		/**
-		 * Builds a dictionary containing the actions from AllActions that are outdated by calling
-		 * IsActionOutdated.
-		 */
+		/// <summary>
+		/// Builds a dictionary containing the actions from AllActions that are outdated by calling
+		/// IsActionOutdated.
+		/// </summary>
 		static void GatherAllOutdatedActions(UEBuildTarget Target, ActionHistory ActionHistory, ref Dictionary<Action, bool> OutdatedActions, Dictionary<UEBuildTarget, List<FileItem>> TargetToOutdatedPrerequisitesMap)
 		{
 			var CheckOutdatednessStartTime = DateTime.UtcNow;
@@ -1026,12 +1083,11 @@ namespace UnrealBuildTool
 			}
 		}
 
-		/**
-		 * Deletes all the items produced by actions in the provided outdated action dictionary. 
-		 * 
-		 * @param	OutdatedActionDictionary	Dictionary of outdated actions
-		 * @param	bShouldDeleteAllFiles		Whether to delete all files associated with outdated items or just ones required
-		 */
+		/// <summary>
+		/// Deletes all the items produced by actions in the provided outdated action dictionary.
+		/// </summary>
+		/// <param name="OutdatedActionDictionary">Dictionary of outdated actions</param>
+		/// <param name="bShouldDeleteAllFiles"> Whether to delete all files associated with outdated items or just ones required</param>
 		static void DeleteOutdatedProducedItems(Dictionary<Action, bool> OutdatedActionDictionary, bool bShouldDeleteAllFiles)
 		{
 			foreach (KeyValuePair<Action, bool> OutdatedActionInfo in OutdatedActionDictionary)
@@ -1055,10 +1111,10 @@ namespace UnrealBuildTool
 			}
 		}
 
-		/**
-		 * Creates directories for all the items produced by actions in the provided outdated action
-		 * dictionary.
-		 */
+		/// <summary>
+		/// Creates directories for all the items produced by actions in the provided outdated action
+		/// dictionary.
+		/// </summary>
 		static void CreateDirectoriesForProducedItems(Dictionary<Action, bool> OutdatedActionDictionary)
 		{
 			foreach (KeyValuePair<Action, bool> OutdatedActionInfo in OutdatedActionDictionary)
diff --git a/Engine/Source/Programs/UnrealBuildTool/System/ActionHistory.cs b/Engine/Source/Programs/UnrealBuildTool/System/ActionHistory.cs
--- a/Engine/Source/Programs/UnrealBuildTool/System/ActionHistory.cs
+++ b/Engine/Source/Programs/UnrealBuildTool/System/ActionHistory.cs
@@ -8,18 +8,24 @@ using System.Runtime.Serialization.Formatters.Binary;
 
 namespace UnrealBuildTool
 {
-	/**
-	 * Caches include dependency information to speed up preprocessing on subsequent runs.
-	 */
+	/// <summary>
+	/// Caches include dependency information to speed up preprocessing on subsequent runs.
+	/// </summary>
 	public class ActionHistory
 	{
-		/** Path to store the cache data to. */
+		/// <summary>
+		/// Path to store the cache data to.
+		/// </summary>
 		private string FilePath;
 
-		/** The command lines used to produce files, keyed by the absolute file paths. */
+		/// <summary>
+		/// The command lines used to produce files, keyed by the absolute file paths.
+		/// </summary>
 		private Dictionary<string, string> ProducedItemToPreviousActionCommandLine;
 
-		/** Whether the dependency cache is dirty and needs to be saved. */
+		/// <summary>
+		/// Whether the dependency cache is dirty and needs to be saved.
+		/// </summary>
 		private bool bIsDirty;
 
 		public ActionHistory(string InFilePath)
diff --git a/Engine/Source/Programs/UnrealBuildTool/System/CPPEnvironment.cs b/Engine/Source/Programs/UnrealBuildTool/System/CPPEnvironment.cs
--- a/Engine/Source/Programs/UnrealBuildTool/System/CPPEnvironment.cs
+++ b/Engine/Source/Programs/UnrealBuildTool/System/CPPEnvironment.cs
@@ -9,7 +9,9 @@ using System.IO;
 
 namespace UnrealBuildTool
 {
-	/** The platforms that may be compilation targets for C++ files. */
+	/// <summary>
+	/// The platforms that may be compilation targets for C++ files.
+	/// </summary>
 	public enum CPPTargetPlatform
 	{
 		Win32,
@@ -26,7 +28,9 @@ namespace UnrealBuildTool
 		Linux,
 	}
 
-	/** The optimization level that may be compilation targets for C++ files. */
+	/// <summary>
+	/// The optimization level that may be compilation targets for C++ files.
+	/// </summary>
 	public enum CPPTargetConfiguration
 	{
 		Debug,
@@ -34,14 +38,18 @@ namespace UnrealBuildTool
 		Shipping
 	}
 
-	/** The optimization level that may be compilation targets for C# files. */
+	/// <summary>
+	/// The optimization level that may be compilation targets for C# files.
+	/// </summary>
 	public enum CSharpTargetConfiguration
 	{
 		Debug,
 		Development,
 	}
 
-	/** The possible interactions between a precompiled header and a C++ file being compiled. */
+	/// <summary>
+	/// The possible interactions between a precompiled header and a C++ file being compiled.
+	/// </summary>
 	public enum PrecompiledHeaderAction
 	{
 		None,
@@ -49,14 +57,18 @@ namespace UnrealBuildTool
 		Create
 	}
 
-	/** Whether the Common Language Runtime is enabled when compiling C++ targets (enables C++/CLI) */
+	/// <summary>
+	/// Whether the Common Language Runtime is enabled when compiling C++ targets (enables C++/CLI)
+	/// </summary>
 	public enum CPPCLRMode
 	{
 		CLRDisabled,
 		CLREnabled,
 	}
 
-	/** Encapsulates the compilation output of compiling a set of C++ files. */
+	/// <summary>
+	/// Encapsulates the compilation output of compiling a set of C++ files.
+	/// </summary>
 	public class CPPOutput
 	{
 		public List<FileItem> ObjectFiles = new List<FileItem>();
@@ -66,31 +78,48 @@ namespace UnrealBuildTool
 
 	public class PrivateAssemblyInfoConfiguration
 	{
-		/** True if the assembly should be copied to an output folder.  Often, referenced assemblies must
-		    be loaded from the directory that the main executable resides in (or a sub-folder).  Note that
-		    PDB files will be copied as well. */
+		/// <summary>
+		/// True if the assembly should be copied to an output folder.  Often, referenced assemblies must
+		/// be loaded from the directory that the main executable resides in (or a sub-folder).  Note that
+		/// PDB files will be copied as well.
+		/// </summary>
 		public bool bShouldCopyToOutputDirectory = false;
 
-		/** Optional sub-directory to copy the assembly to */
+		/// <summary>
+		/// Optional sub-directory to copy the assembly to
+		/// </summary>
 		public string DestinationSubDirectory = "";
 	}
 
-	/** Describes a private assembly dependency */
+	/// <summary>
+	/// Describes a private assembly dependency
+	/// </summary>
 	public class PrivateAssemblyInfo
 	{
-		/** The file item for the reference assembly on disk */
+		/// <summary>
+		/// The file item for the reference assembly on disk
+		/// </summary>
 		public FileItem FileItem = null;
 
-		/** The configuration of the private assembly info */
+		/// <summary>
+		/// The configuration of the private assembly info
+		/// </summary>
 		public PrivateAssemblyInfoConfiguration Config = new PrivateAssemblyInfoConfiguration();
 	}
 
-	/** Encapsulates the environment that a C# file is compiled in. */
+	/// <summary>
+	/// Encapsulates the environment that a C# file is compiled in.
+	/// </summary>
 	public class CSharpEnvironment
 	{
-		/** The configuration to be compiled for. */
+		/// <summary>
+		/// The configuration to be compiled for.
+		/// </summary>
 		public CSharpTargetConfiguration TargetConfiguration;
-		/** The target platform used to set the environment. Doesn't affect output. */
+
+		/// <summary>
+		/// The target platform used to set the environment. Doesn't affect output.
+		/// </summary>
 		public CPPTargetPlatform EnvironmentTargetPlatform;
 	}
 
@@ -98,13 +127,19 @@ namespace UnrealBuildTool
 	{
 		public class TargetInfo
 		{
-			/** The platform to be compiled/linked for. */
+			/// <summary>
+			/// The platform to be compiled/linked for.
+			/// </summary>
 			public CPPTargetPlatform Platform;
 
-			/** The architecture that is being compiled/linked (empty string by default) */
+			/// <summary>
+			/// The architecture that is being compiled/linked (empty string by default)
+			/// </summary>
 			public string Architecture;
 
-			/** The configuration to be compiled/linked for. */
+			/// <summary>
+			/// The configuration to be compiled/linked for.
+			/// </summary>
 			public CPPTargetConfiguration Configuration;
 
 			public TargetInfo()
@@ -173,103 +208,165 @@ namespace UnrealBuildTool
 
 	}
 
-	/** Encapsulates the configuration of an environment that a C++ file is compiled in. */
+	/// <summary>
+	/// Encapsulates the configuration of an environment that a C++ file is compiled in.
+	/// </summary>
 	public class CPPEnvironmentConfiguration : NativeBuildEnvironmentConfiguration
 	{
-		/** The directory to put the output object/debug files in. */
+		/// <summary>
+		/// The directory to put the output object/debug files in.
+		/// </summary>
 		public DirectoryReference OutputDirectory = null;
 
-		/** The directory to shadow source files in for syncing to remote compile servers */
+		/// <summary>
+		/// The directory to shadow source files in for syncing to remote compile servers
+		/// </summary>
 		public DirectoryReference LocalShadowDirectory = null;
 
-		/** PCH header file name as it appears in an #include statement in source code (might include partial, or no relative path.)
-		    This is needed by some compilers to use PCH features. */
+		/// <summary>
+		/// PCH header file name as it appears in an #include statement in source code (might include partial, or no relative path.)
+		/// This is needed by some compilers to use PCH features.
+		/// </summary>
 		public string PCHHeaderNameInCode;
 
-		/** The name of the header file which is precompiled. */
+		/// <summary>
+		/// The name of the header file which is precompiled.
+		/// </summary>
 		public FileReference PrecompiledHeaderIncludeFilename = null;
 
-		/** Whether the compilation should create, use, or do nothing with the precompiled header. */
+		/// <summary>
+		/// Whether the compilation should create, use, or do nothing with the precompiled header.
+		/// </summary>
 		public PrecompiledHeaderAction PrecompiledHeaderAction = PrecompiledHeaderAction.None;
 
-		/** True if the precompiled header needs to be "force included" by the compiler.  This is usually used with
-		    the "Shared PCH" feature of UnrealBuildTool.  Note that certain platform toolchains *always* force
-		    include PCH header files first. */
+		/// <summary>
+		/// True if the precompiled header needs to be "force included" by the compiler.  This is usually used with
+		/// the "Shared PCH" feature of UnrealBuildTool.  Note that certain platform toolchains *always* force
+		/// include PCH header files first.
+		/// </summary>
 		public bool bForceIncludePrecompiledHeader = false;
 
-		/** Use run time type information */
+		/// <summary>
+		/// Use run time type information
+		/// </summary>
 		public bool bUseRTTI = false;
 
-		/** Use AVX instructions */
+		/// <summary>
+		/// Use AVX instructions
+		/// </summary>
 		public bool bUseAVX = false;
 
-		/** Enable buffer security checks.   This should usually be enabled as it prevents severe security risks. */
+		/// <summary>
+		/// Enable buffer security checks.   This should usually be enabled as it prevents severe security risks.
+		/// </summary>
 		public bool bEnableBufferSecurityChecks = true;
 
-		/** If true and unity builds are enabled, this module will build without unity. */
+		/// <summary>
+		/// If true and unity builds are enabled, this module will build without unity.
+		/// </summary>
 		public bool bFasterWithoutUnity = false;
 
-		/** The number of source files in this module before unity build will be activated for that module.  If set to
-		    anything besides -1, will override the default setting which is controlled by MinGameModuleSourceFilesForUnityBuild */
+		/// <summary>
+		/// The number of source files in this module before unity build will be activated for that module.  If set to
+		/// anything besides -1, will override the default setting which is controlled by MinGameModuleSourceFilesForUnityBuild
+		/// </summary>
 		public int MinSourceFilesForUnityBuildOverride = 0;
 
-		/** Overrides BuildConfiguration.MinFilesUsingPrecompiledHeader if non-zero. */
+		/// <summary>
+		/// Overrides BuildConfiguration.MinFilesUsingPrecompiledHeader if non-zero.
+		/// </summary>
 		public int MinFilesUsingPrecompiledHeaderOverride = 0;
 
-		/**  Module uses a #import so must be built locally when compiling with SN-DBS */
+		/// <summary>
+		/// Module uses a #import so must be built locally when compiling with SN-DBS
+		/// </summary>
 		public bool bBuildLocallyWithSNDBS = false;
 
-		/** Enable exception handling */
+		/// <summary>
+		/// Enable exception handling
+		/// </summary>
 		public bool bEnableExceptions = false;
 
-		/** Whether to warn about the use of shadow variables */
+		/// <summary>
+		/// Whether to warn about the use of shadow variables
+		/// </summary>
 		public bool bEnableShadowVariableWarning = false;
 
-		/** True if the environment contains performance critical code. */
+		/// <summary>
+		/// True if the environment contains performance critical code.
+		/// </summary>
 		public ModuleRules.CodeOptimization OptimizeCode = ModuleRules.CodeOptimization.Default;
 
-		/** True if debug info should be created. */
+		/// <summary>
+		/// True if debug info should be created.
+		/// </summary>
 		public bool bCreateDebugInfo = true;
 
-		/** True if we're compiling .cpp files that will go into a library (.lib file) */
+		/// <summary>
+		/// True if we're compiling .cpp files that will go into a library (.lib file)
+		/// </summary>
 		public bool bIsBuildingLibrary = false;
 
-		/** True if we're compiling a DLL */
+		/// <summary>
+		/// True if we're compiling a DLL
+		/// </summary>
 		public bool bIsBuildingDLL = false;
 
-		/** Whether we should compile using the statically-linked CRT. This is not widely supported for the whole engine, but is required for programs that need to run without dependencies. */
+		/// <summary>
+		/// Whether we should compile using the statically-linked CRT. This is not widely supported for the whole engine, but is required for programs that need to run without dependencies.
+		/// </summary>
 		public bool bUseStaticCRT = false;
 
-		/** Whether the CLR (Common Language Runtime) support should be enabled for C++ targets (C++/CLI). */
+		/// <summary>
+		/// Whether the CLR (Common Language Runtime) support should be enabled for C++ targets (C++/CLI).
+		/// </summary>
 		public CPPCLRMode CLRMode = CPPCLRMode.CLRDisabled;
 
-		/** The include paths to look for included files in. */
+		/// <summary>
+		/// The include paths to look for included files in.
+		/// </summary>
 		public readonly CPPIncludeInfo CPPIncludeInfo = new CPPIncludeInfo();
 
-		/** Paths where .NET framework assembly references are found, when compiling CLR applications. */
+		/// <summary>
+		/// Paths where .NET framework assembly references are found, when compiling CLR applications.
+		/// </summary>
 		public List<string> SystemDotNetAssemblyPaths = new List<string>();
 
-		/** Full path and file name of .NET framework assemblies we're referencing */
+		/// <summary>
+		/// Full path and file name of .NET framework assemblies we're referencing
+		/// </summary>
 		public List<string> FrameworkAssemblyDependencies = new List<string>();
 
-		/** List of private CLR assemblies that, when modified, will force recompilation of any CLR source files */
+		/// <summary>
+		/// List of private CLR assemblies that, when modified, will force recompilation of any CLR source files
+		/// </summary>
 		public List<PrivateAssemblyInfoConfiguration> PrivateAssemblyDependencies = new List<PrivateAssemblyInfoConfiguration>();
 
-		/** The C++ preprocessor definitions to use. */
+		/// <summary>
+		/// The C++ preprocessor definitions to use.
+		/// </summary>
 		public List<string> Definitions = new List<string>();
 
-		/** Additional arguments to pass to the compiler. */
+		/// <summary>
+		/// Additional arguments to pass to the compiler.
+		/// </summary>
 		public string AdditionalArguments = "";
 
-		/** A list of additional frameworks whose include paths are needed. */
+		/// <summary>
+		/// A list of additional frameworks whose include paths are needed.
+		/// </summary>
 		public List<UEBuildFramework> AdditionalFrameworks = new List<UEBuildFramework>();
 
-		/** Default constructor. */
+		/// <summary>
+		/// Default constructor.
+		/// </summary>
 		public CPPEnvironmentConfiguration()
 		{
 		}
 
-		/** Copy constructor. */
+		/// <summary>
+		/// Copy constructor.
+		/// </summary>
 		public CPPEnvironmentConfiguration(CPPEnvironmentConfiguration InCopyEnvironment) :
 			base(InCopyEnvironment)
 		{
@@ -323,31 +420,45 @@ namespace UnrealBuildTool
 	}
 
 
-	/** Encapsulates the environment that a C++ file is compiled in. */
+	/// <summary>
+	/// Encapsulates the environment that a C++ file is compiled in.
+	/// </summary>
 	public partial class CPPEnvironment
 	{
-		/** The file containing the precompiled header data. */
+		/// <summary>
+		/// The file containing the precompiled header data.
+		/// </summary>
 		public FileItem PrecompiledHeaderFile = null;
 
-		/** List of private CLR assemblies that, when modified, will force recompilation of any CLR source files */
+		/// <summary>
+		/// List of private CLR assemblies that, when modified, will force recompilation of any CLR source files
+		/// </summary>
 		public List<PrivateAssemblyInfo> PrivateAssemblyDependencies = new List<PrivateAssemblyInfo>();
 
 		public CPPEnvironmentConfiguration Config = new CPPEnvironmentConfiguration();
 
-		/** List of available shared global PCH headers */
+		/// <summary>
+		/// List of available shared global PCH headers
+		/// </summary>
 		public List<SharedPCHHeaderInfo> SharedPCHHeaderFiles = new List<SharedPCHHeaderInfo>();
 
-		/** List of "shared" precompiled header environments, potentially shared between modules */
+		/// <summary>
+		/// List of "shared" precompiled header environments, potentially shared between modules
+		/// </summary>
 		public readonly List<PrecompileHeaderEnvironment> SharedPCHEnvironments = new List<PrecompileHeaderEnvironment>();
 
 		// Whether or not UHT is being built
 		public bool bHackHeaderGenerator;
 
-		/** Default constructor. */
+		/// <summary>
+		/// Default constructor.
+		/// </summary>
 		public CPPEnvironment()
 		{ }
 
-		/** Copy constructor. */
+		/// <summary>
+		/// Copy constructor.
+		/// </summary>
 		protected CPPEnvironment(CPPEnvironment InCopyEnvironment)
 		{
 			PrecompiledHeaderFile = InCopyEnvironment.PrecompiledHeaderFile;
@@ -359,11 +470,10 @@ namespace UnrealBuildTool
 			Config = new CPPEnvironmentConfiguration(InCopyEnvironment.Config);
 		}
 
-		/**
-		 * Whether to use PCH files with the current target
-		 * 
-		 * @return	true if PCH files should be used, false otherwise
-		 */
+		/// <summary>
+		/// Whether to use PCH files with the current target
+		/// </summary>
+		/// <returns>true if PCH files should be used, false otherwise</returns>
 		public bool ShouldUsePCHs()
 		{
 			return UEBuildPlatform.GetBuildPlatformForCPPTargetPlatform(Config.Target.Platform).ShouldUsePCHFiles(Config.Target.Platform, Config.Target.Configuration);
diff --git a/Engine/Source/Programs/UnrealBuildTool/System/CPPHeaders.cs b/Engine/Source/Programs/UnrealBuildTool/System/CPPHeaders.cs
--- a/Engine/Source/Programs/UnrealBuildTool/System/CPPHeaders.cs
+++ b/Engine/Source/Programs/UnrealBuildTool/System/CPPHeaders.cs
@@ -12,20 +12,26 @@ using Tools.DotNETCommon.FileContentsCacheType;
 
 namespace UnrealBuildTool
 {
-	/** For C++ source file items, this structure is used to cache data that will be used for include dependency scanning */
+	/// <summary>
+	/// For C++ source file items, this structure is used to cache data that will be used for include dependency scanning
+	/// </summary>
 	[Serializable]
 	public class CPPIncludeInfo : ISerializable
 	{
-		/** Ordered list of include paths for the module  */
+		/// <summary>
+		/// Ordered list of include paths for the module
+		/// </summary>
 		public HashSet<string> IncludePaths = new HashSet<string>();
 
-		/**
-		 * The include paths where changes to contained files won't cause dependent C++ source files to
-		 * be recompiled, unless BuildConfiguration.bCheckSystemHeadersForModification==true.
-		 */
+		/// <summary>
+		/// The include paths where changes to contained files won't cause dependent C++ source files to
+		/// be recompiled, unless BuildConfiguration.bCheckSystemHeadersForModification==true.
+		/// </summary>
 		public HashSet<string> SystemIncludePaths = new HashSet<string>();
 
-		/** Contains a mapping from filename to the full path of the header in this environment.  This is used to optimized include path lookups at runtime for any given single module. */
+		/// <summary>
+		/// Contains a mapping from filename to the full path of the header in this environment.  This is used to optimized include path lookups at runtime for any given single module.
+		/// </summary>
 		public Dictionary<string, FileItem> IncludeFileSearchDictionary = new Dictionary<string, FileItem>();
 
 		public CPPIncludeInfo()
@@ -78,20 +84,28 @@ namespace UnrealBuildTool
 	/// </summary>
 	public class IncludedFilesSet : HashSet<FileItem>
 	{
-		/** Whether this file list has been fully initialized or not. */
+		/// <summary>
+		/// Whether this file list has been fully initialized or not.
+		/// </summary>
 		public bool bIsInitialized;
 
-		/** List of files which include this file in one of its includes. */
+		/// <summary>
+		/// List of files which include this file in one of its includes.
+		/// </summary>
 		public List<FileItem> CircularDependencies = new List<FileItem>();
 	}
 
 
 	public partial class CPPEnvironment
 	{
-		/** Contains a cache of include dependencies (direct and indirect), one for each target we're building. */
+		/// <summary>
+		/// Contains a cache of include dependencies (direct and indirect), one for each target we're building.
+		/// </summary>
 		public static readonly Dictionary<UEBuildTarget, DependencyCache> IncludeDependencyCache = new Dictionary<UEBuildTarget, DependencyCache>();
 
-		/** Contains a cache of include dependencies (direct and indirect), one for each target we're building. */
+		/// <summary>
+		/// Contains a cache of include dependencies (direct and indirect), one for each target we're building.
+		/// </summary>
 		public static readonly Dictionary<UEBuildTarget, FlatCPPIncludeDependencyCache> FlatCPPIncludeDependencyCache = new Dictionary<UEBuildTarget, FlatCPPIncludeDependencyCache>();
 
 		public static int TotalFindIncludedFileCalls = 0;
@@ -99,12 +113,12 @@ namespace UnrealBuildTool
 		public static int IncludePathSearchAttempts = 0;
 
 
-		/** 
-		 * Finds the header file that is referred to by a partial include filename. 
-		 * @param RelativeIncludePath path relative to the project
-		 * @param bSkipExternalHeader true to skip processing of headers in external path
-		 * @param SourceFilesDirectory - The folder containing the source files we're generating a PCH for
-		 */
+		/// <summary>
+		/// Finds the header file that is referred to by a partial include filename.
+		/// </summary>
+		/// <param name="RelativeIncludePath">path relative to the project</param>
+		/// <param name="bSkipExternalHeader">true to skip processing of headers in external path</param>
+		/// <param name="SourceFilesDirectory">- The folder containing the source files we're generating a PCH for</param>
 		public static FileItem FindIncludedFile(string RelativeIncludePath, bool bSkipExternalHeader, List<string> IncludePathsToSearch, Dictionary<string, FileItem> IncludeFileSearchDictionary)
 		{
 			FileItem Result = null;
@@ -192,10 +206,14 @@ namespace UnrealBuildTool
 			return Result;
 		}
 
-		/** A cache of the list of other files that are directly or indirectly included by a C++ file. */
+		/// <summary>
+		/// A cache of the list of other files that are directly or indirectly included by a C++ file.
+		/// </summary>
 		static Dictionary<FileItem, IncludedFilesSet> ExhaustiveIncludedFilesMap = new Dictionary<FileItem, IncludedFilesSet>();
 
-		/** A cache of all files included by a C++ file, but only has files that we knew about from a previous session, loaded from a cache at startup */
+		/// <summary>
+		/// A cache of all files included by a C++ file, but only has files that we knew about from a previous session, loaded from a cache at startup
+		/// </summary>
 		static Dictionary<FileItem, IncludedFilesSet> OnlyCachedIncludedFilesMap = new Dictionary<FileItem, IncludedFilesSet>();
 
 
@@ -381,14 +399,18 @@ namespace UnrealBuildTool
 		public static int TotalDirectIncludeResolves = 0;
 
 
-		/** Regex that matches #include statements. */
+		/// <summary>
+		/// Regex that matches #include statements.
+		/// </summary>
 		static readonly Regex CPPHeaderRegex = new Regex("(([ \t]*#[ \t]*include[ \t]*[<\"](?<HeaderFile>[^\">]*)[\">][^\n]*\n*)|([^\n]*\n*))*",
 													RegexOptions.Compiled | RegexOptions.Singleline | RegexOptions.ExplicitCapture);
 
 		static readonly Regex MMHeaderRegex = new Regex("(([ \t]*#[ \t]*import[ \t]*[<\"](?<HeaderFile>[^\">]*)[\">][^\n]*\n*)|([^\n]*\n*))*",
 													RegexOptions.Compiled | RegexOptions.Singleline | RegexOptions.ExplicitCapture);
 
-		/** Regex that matches C++ code with UObject declarations which we will need to generated code for. */
+		/// <summary>
+		/// Regex that matches C++ code with UObject declarations which we will need to generated code for.
+		/// </summary>
 		static readonly Regex UObjectRegex = new Regex("^\\s*U(CLASS|STRUCT|ENUM|INTERFACE|DELEGATE)\\b", RegexOptions.Compiled | RegexOptions.Multiline);
 
 		// Maintains a cache of file contents
@@ -401,7 +423,9 @@ namespace UnrealBuildTool
 			return UObjectRegex.IsMatch(Contents);
 		}
 
-		/** Finds the names of files directly included by the given C++ file, and also whether the file contains any UObjects */
+		/// <summary>
+		/// Finds the names of files directly included by the given C++ file, and also whether the file contains any UObjects
+		/// </summary>
 		public static List<DependencyInclude> GetDirectIncludeDependencies(UEBuildTarget Target, FileItem CPPFile, IUEBuildPlatform BuildPlatform, bool bOnlyCachedDependencies)
 		{
 			// Try to fulfill request from cache first.
diff --git a/Engine/Source/Programs/UnrealBuildTool/System/DependencyCache.cs b/Engine/Source/Programs/UnrealBuildTool/System/DependencyCache.cs
--- a/Engine/Source/Programs/UnrealBuildTool/System/DependencyCache.cs
+++ b/Engine/Source/Programs/UnrealBuildTool/System/DependencyCache.cs
@@ -37,40 +37,55 @@ namespace UnrealBuildTool
 		}
 	}
 
-	/**
-	 * Caches include dependency information to speed up preprocessing on subsequent runs.
-	 */
+	/// <summary>
+	/// Caches include dependency information to speed up preprocessing on subsequent runs.
+	/// </summary>
 	public class DependencyCache
 	{
-		/** The version number for binary serialization */
+		/// <summary>
+		/// The version number for binary serialization
+		/// </summary>
 		const int FileVersion = 1;
 
-		/** The file signature for binary serialization */
+		/// <summary>
+		/// The file signature for binary serialization
+		/// </summary>
 		const int FileSignature = ('D' << 24) | ('C' << 16) | FileVersion;
 
-		/** Path to store the cache data to. */
+		/// <summary>
+		/// Path to store the cache data to.
+		/// </summary>
 		private FileReference BackingFile;
 
-		/** The time the cache was created. Used to invalidate entries. */
+		/// <summary>
+		/// The time the cache was created. Used to invalidate entries.
+		/// </summary>
 		public DateTime CreationTimeUtc;
 
-		/** The time the cache was last updated. Stored as the creation date when saved. Not serialized. */
+		/// <summary>
+		/// The time the cache was last updated. Stored as the creation date when saved. Not serialized.
+		/// </summary>
 		private DateTime UpdateTimeUtc;
 
-		/** Dependency lists, keyed on file's absolute path. */
+		/// <summary>
+		/// Dependency lists, keyed on file's absolute path.
+		/// </summary>
 		private Dictionary<FileReference, List<DependencyInclude>> DependencyMap;
 
-		/** A mapping of whether various files exist. Not serialized. */
+		/// <summary>
+		/// A mapping of whether various files exist. Not serialized.
+		/// </summary>
 		private Dictionary<FileReference, bool> FileExistsInfo;
 
-		/** Whether the dependency cache is dirty and needs to be saved. Not serialized. */
+		/// <summary>
+		/// Whether the dependency cache is dirty and needs to be saved. Not serialized.
+		/// </summary>
 		private bool bIsDirty;
 
-		/**
-		 * Creates and deserializes the dependency cache at the passed in location
-		 * 
-		 * @param	CachePath	Name of the cache file to deserialize
-		 */
+		/// <summary>
+		/// Creates and deserializes the dependency cache at the passed in location
+		/// </summary>
+		/// <param name="CachePath">Name of the cache file to deserialize</param>
 		public static DependencyCache Create(FileReference CacheFile)
 		{
 			// See whether the cache file exists.
@@ -103,11 +118,10 @@ namespace UnrealBuildTool
 			return new DependencyCache(CacheFile);
 		}
 
-		/**
-		 * Loads the cache from the passed in file.
-		 * 
-		 * @param	Cache	File to deserialize from
-		 */
+		/// <summary>
+		/// Loads the cache from the passed in file.
+		/// </summary>
+		/// <param name="Cache">File to deserialize from</param>
 		public static DependencyCache Load(FileReference CacheFile)
 		{
 			DependencyCache Result = null;
@@ -202,11 +216,10 @@ namespace UnrealBuildTool
 			return Cache;
 		}
 
-		/**
-		 * Constructor
-		 * 
-		 * @param	Cache	File associated with this cache
-		 */
+		/// <summary>
+		/// Constructor
+		/// </summary>
+		/// <param name="Cache">File associated with this cache</param>
 		protected DependencyCache(FileReference InBackingFile)
 		{
 			BackingFile = InBackingFile;
@@ -217,9 +230,9 @@ namespace UnrealBuildTool
 			CreateFileExistsInfo();
 		}
 
-		/**
-		 * Saves the dependency cache to disk using the update time as the creation time.
-		 */
+		/// <summary>
+		/// Saves the dependency cache to disk using the update time as the creation time.
+		/// </summary>
 		public void Save()
 		{
 			// Only save if we've made changes to it since load.
@@ -273,12 +286,10 @@ namespace UnrealBuildTool
 			}
 		}
 
-		/**
-		 * Returns the direct dependencies of the specified FileItem if it exists in the cache and they are not stale.
-		 *
-		 * @param  File  File to try to find dependencies in cache
-		 * @returns  Dependency info for File, or null if no dependencies are cached or if the cache is stale.
-		 */
+		/// <summary>
+		/// Returns the direct dependencies of the specified FileItem if it exists in the cache and they are not stale.
+		/// </summary>
+		/// <param name="">File  File to try to find dependencies in cache</param>
 		public List<DependencyInclude> GetCachedDependencyInfo(FileItem File)
 		{
 			// Check whether File is in cache.
@@ -324,13 +335,12 @@ namespace UnrealBuildTool
 			return Includes;
 		}
 
-		/**
-		 * Update cache with dependencies for the passed in file.
-		 * 
-		 * @param	File			File to update dependencies for
-		 * @param	Dependencies	List of dependencies to cache for passed in file
-		 * @param	HasUObjects		True if this file was found to contain UObject classes or types
-		 */
+		/// <summary>
+		/// Update cache with dependencies for the passed in file.
+		/// </summary>
+		/// <param name="File">  File to update dependencies for</param>
+		/// <param name="Dependencies">List of dependencies to cache for passed in file</param>
+		/// <param name="HasUObjects"> True if this file was found to contain UObject classes or types</param>
 		public void SetDependencyInfo(FileItem File, List<DependencyInclude> Info)
 		{
 			DependencyMap[File.Reference] = Info;
diff --git a/Engine/Source/Programs/UnrealBuildTool/System/ExternalExecution.cs b/Engine/Source/Programs/UnrealBuildTool/System/ExternalExecution.cs
--- a/Engine/Source/Programs/UnrealBuildTool/System/ExternalExecution.cs
+++ b/Engine/Source/Programs/UnrealBuildTool/System/ExternalExecution.cs
@@ -19,21 +19,44 @@ namespace UnrealBuildTool
 	// processes valid.
 	public enum ECompilationResult
 	{
-		/** Compilation succeeded */
+		/// <summary>
+		/// Compilation succeeded
+		/// </summary>
 		Succeeded = 0,
-		/** Build was canceled, this is used on the engine side only */
+
+		/// <summary>
+		/// Build was canceled, this is used on the engine side only
+		/// </summary>
 		Canceled = 1,
-		/** All targets were up to date, used only with -canskiplink */
+
+		/// <summary>
+		/// All targets were up to date, used only with -canskiplink
+		/// </summary>
 		UpToDate = 2,
-		/** The process has most likely crashed. This is what UE returns in case of an assert */
+
+		/// <summary>
+		/// The process has most likely crashed. This is what UE returns in case of an assert
+		/// </summary>
 		CrashOrAssert = 3,
-		/** Compilation failed because generated code changed which was not supported */
+
+		/// <summary>
+		/// Compilation failed because generated code changed which was not supported
+		/// </summary>
 		FailedDueToHeaderChange = 4,
-		/** Compilation failed due to compilation errors */
+
+		/// <summary>
+		/// Compilation failed due to compilation errors
+		/// </summary>
 		OtherCompilationError = 5,
-		/** Compilation is not supported in the current build */
+
+		/// <summary>
+		/// Compilation is not supported in the current build
+		/// </summary>
 		Unsupported,
-		/** Unknown error */
+
+		/// <summary>
+		/// Unknown error
+		/// </summary>
 		Unknown
 	}
 	public static class CompilationResultExtensions
@@ -44,35 +67,55 @@ namespace UnrealBuildTool
 		}
 	}
 
-	/** Information about a module that needs to be passed to UnrealHeaderTool for code generation */
+	/// <summary>
+	/// Information about a module that needs to be passed to UnrealHeaderTool for code generation
+	/// </summary>
 	[Serializable]
 	public class UHTModuleInfo : ISerializable
 	{
-		/** Module name */
+		/// <summary>
+		/// Module name
+		/// </summary>
 		public string ModuleName;
 
-		/** Module base directory */
+		/// <summary>
+		/// Module base directory
+		/// </summary>
 		public string ModuleDirectory;
 
-		/** Module type */
+		/// <summary>
+		/// Module type
+		/// </summary>
 		public string ModuleType;
 
-		/** Public UObject headers found in the Classes directory (legacy) */
+		/// <summary>
+		/// Public UObject headers found in the Classes directory (legacy)
+		/// </summary>
 		public List<FileItem> PublicUObjectClassesHeaders;
 
-		/** Public headers with UObjects */
+		/// <summary>
+		/// Public headers with UObjects
+		/// </summary>
 		public List<FileItem> PublicUObjectHeaders;
 
-		/** Private headers with UObjects */
+		/// <summary>
+		/// Private headers with UObjects
+		/// </summary>
 		public List<FileItem> PrivateUObjectHeaders;
 
-		/** Module PCH absolute path */
+		/// <summary>
+		/// Module PCH absolute path
+		/// </summary>
 		public string PCH;
 
-		/** Base (i.e. extensionless) path+filename of the .generated files */
+		/// <summary>
+		/// Base (i.e. extensionless) path+filename of the .generated files
+		/// </summary>
 		public string GeneratedCPPFilenameBase;
 
-		/** Version of code generated by UHT */
+		/// <summary>
+		/// Version of code generated by UHT
+		/// </summary>
 		public EGeneratedCodeVersion GeneratedCodeVersion;
 
 		public UHTModuleInfo()
@@ -178,9 +221,9 @@ namespace UnrealBuildTool
 	}
 
 
-	/**
-	 * This handles all running of the UnrealHeaderTool
-	 */
+	/// <summary>
+	/// This handles all running of the UnrealHeaderTool
+	/// </summary>
 	public class ExternalExecution
 	{
 		/// <summary>
@@ -396,12 +439,11 @@ namespace UnrealBuildTool
 			return Timestamp;
 		}
 
-		/**
-		 * Checks the class header files and determines if generated UObject code files are out of date in comparison.
-		 * @param UObjectModules	Modules that we generate headers for
-		 * 
-		 * @return					True if the code files are out of date
-		 * */
+		/// <summary>
+		/// Checks the class header files and determines if generated UObject code files are out of date in comparison.
+		/// </summary>
+		/// <param name="UObjectModules">Modules that we generate headers for</param>
+		/// <returns>    True if the code files are out of date</returns>
 		private static bool AreGeneratedCodeFilesOutOfDate(List<UHTModuleInfo> UObjectModules, DateTime HeaderToolTimestamp)
 		{
 			// Get CoreUObject.generated.cpp timestamp.  If the source files are older than the CoreUObject generated code, we'll
@@ -524,7 +566,9 @@ namespace UnrealBuildTool
 			return false;
 		}
 
-		/** Updates the intermediate include directory timestamps of all the passed in UObject modules */
+		/// <summary>
+		/// Updates the intermediate include directory timestamps of all the passed in UObject modules
+		/// </summary>
 		private static void UpdateDirectoryTimestamps(List<UHTModuleInfo> UObjectModules)
 		{
 			foreach (var Module in UObjectModules)
@@ -572,7 +616,9 @@ namespace UnrealBuildTool
 			}
 		}
 
-		/** Run an external exe (and capture the output), given the exe path and the commandline. */
+		/// <summary>
+		/// Run an external exe (and capture the output), given the exe path and the commandline.
+		/// </summary>
 		public static int RunExternalExecutable(string ExePath, string Commandline)
 		{
 			var ExeInfo = new ProcessStartInfo(ExePath, Commandline);
@@ -589,7 +635,9 @@ namespace UnrealBuildTool
 			}
 		}
 
-		/** Simple function to pipe output asynchronously */
+		/// <summary>
+		/// Simple function to pipe output asynchronously
+		/// </summary>
 		private static void PrintProcessOutputAsync(object Sender, DataReceivedEventArgs Event)
 		{
 			// DataReceivedEventHandler is fired with a null string when the output stream is closed.  We don't want to
@@ -602,10 +650,10 @@ namespace UnrealBuildTool
 
 
 
-		/**
-		 * Builds and runs the header tool and touches the header directories.
-		 * Performs any early outs if headers need no changes, given the UObject modules, tool path, game name, and configuration
-		 */
+		/// <summary>
+		/// Builds and runs the header tool and touches the header directories.
+		/// Performs any early outs if headers need no changes, given the UObject modules, tool path, game name, and configuration
+		/// </summary>
 		public static bool ExecuteHeaderToolIfNecessary(IUEToolChain ToolChain, UEBuildTarget Target, CPPEnvironment GlobalCompileEnvironment, List<UHTModuleInfo> UObjectModules, FileReference ModuleInfoFileName, ref ECompilationResult UHTResult)
 		{
 			if (ProgressWriter.bWriteMarkup)
diff --git a/Engine/Source/Programs/UnrealBuildTool/System/FileItem.cs b/Engine/Source/Programs/UnrealBuildTool/System/FileItem.cs
--- a/Engine/Source/Programs/UnrealBuildTool/System/FileItem.cs
+++ b/Engine/Source/Programs/UnrealBuildTool/System/FileItem.cs
@@ -10,10 +10,10 @@ using System.Runtime.Serialization;
 
 namespace UnrealBuildTool
 {
-	/**
-	 * Represents a file on disk that is used as an input or output of a build action.
-	 * FileItems are created by calling FileItem.GetItemByFileReference, which creates a single FileItem for each unique file path.
-	 */
+	/// <summary>
+	/// Represents a file on disk that is used as an input or output of a build action.
+	/// FileItems are created by calling FileItem.GetItemByFileReference, which creates a single FileItem for each unique file path.
+	/// </summary>
 	[Serializable]
 	public class FileItem : ISerializable
 	{
@@ -21,25 +21,37 @@ namespace UnrealBuildTool
 		/// Preparation and Assembly (serialized)
 		/// 
 
-		/** The action that produces the file. */
+		/// <summary>
+		/// The action that produces the file.
+		/// </summary>
 		public Action ProducingAction = null;
 
-		/** The file reference */
+		/// <summary>
+		/// The file reference
+		/// </summary>
 		public FileReference Reference;
 
-		/** True if any DLLs produced by this  */
+		/// <summary>
+		/// True if any DLLs produced by this
+		/// </summary>
 		public bool bNeedsHotReloadNumbersDLLCleanUp = false;
 
-		/** Whether or not this is a remote file, in which case we can't access it directly */
+		/// <summary>
+		/// Whether or not this is a remote file, in which case we can't access it directly
+		/// </summary>
 		public bool bIsRemoteFile = false;
 
-		/** Accessor for the absolute path to the file */
+		/// <summary>
+		/// Accessor for the absolute path to the file
+		/// </summary>
 		public string AbsolutePath
 		{
 			get { return Reference.FullName; }
 		}
 
-		/** For C++ file items, this stores cached information about the include paths needed in order to include header files from these C++ files.  This is part of UBT's dependency caching optimizations. */
+		/// <summary>
+		/// For C++ file items, this stores cached information about the include paths needed in order to include header files from these C++ files.  This is part of UBT's dependency caching optimizations.
+		/// </summary>
 		public CPPIncludeInfo CachedCPPIncludeInfo
 		{
 			get
@@ -98,11 +110,15 @@ namespace UnrealBuildTool
 		/// Preparation only (not serialized)
 		/// 
 
-		/** PCH header file name as it appears in an #include statement in source code (might include partial, or no relative path.)
-		    This is needed by some compilers to use PCH features. */
+		/// <summary>
+		/// PCH header file name as it appears in an #include statement in source code (might include partial, or no relative path.)
+		/// This is needed by some compilers to use PCH features.
+		/// </summary>
 		public string PCHHeaderNameInCode;
 
-		/** The PCH file that this file will use */
+		/// <summary>
+		/// The PCH file that this file will use
+		/// </summary>
 		public FileReference PrecompiledHeaderIncludeFilename;
 
 
@@ -110,16 +126,24 @@ namespace UnrealBuildTool
 		/// Transients (not serialized)
 		///
 
-		/** The information about the file. */
+		/// <summary>
+		/// The information about the file.
+		/// </summary>
 		public FileInfo Info;
 
-		/** This is true if this item is actually a directory. Consideration for Mac application bundles. Note that Info will be null if true! */
+		/// <summary>
+		/// This is true if this item is actually a directory. Consideration for Mac application bundles. Note that Info will be null if true!
+		/// </summary>
 		public bool IsDirectory;
 
-		/** Relative cost of action associated with producing this file. */
+		/// <summary>
+		/// Relative cost of action associated with producing this file.
+		/// </summary>
 		public long RelativeCost = 0;
 
-		/** The last write time of the file. */
+		/// <summary>
+		/// The last write time of the file.
+		/// </summary>
 		public DateTimeOffset _LastWriteTime;
 		public DateTimeOffset LastWriteTime
 		{
@@ -134,7 +158,9 @@ namespace UnrealBuildTool
 			set { _LastWriteTime = value; }
 		}
 
-		/** Whether the file exists. */
+		/// <summary>
+		/// Whether the file exists.
+		/// </summary>
 		public bool _bExists = false;
 		public bool bExists
 		{
@@ -149,7 +175,9 @@ namespace UnrealBuildTool
 			set { _bExists = value; }
 		}
 
-		/** Size of the file if it exists, otherwise -1 */
+		/// <summary>
+		/// Size of the file if it exists, otherwise -1
+		/// </summary>
 		public long _Length = -1;
 		public long Length
 		{
@@ -169,28 +197,34 @@ namespace UnrealBuildTool
 		/// Statics
 		///
 
-		/** Used for performance debugging */
+		/// <summary>
+		/// Used for performance debugging
+		/// </summary>
 		public static long TotalFileItemCount = 0;
 		public static long MissingFileItemCount = 0;
 
-		/** A case-insensitive dictionary that's used to map each unique file name to a single FileItem object. */
+		/// <summary>
+		/// A case-insensitive dictionary that's used to map each unique file name to a single FileItem object.
+		/// </summary>
 		static Dictionary<FileReference, FileItem> UniqueSourceFileMap = new Dictionary<FileReference, FileItem>();
 
-		/** A list of remote file items that have been created but haven't needed the remote info yet, so we can gang up many into one request */
+		/// <summary>
+		/// A list of remote file items that have been created but haven't needed the remote info yet, so we can gang up many into one request
+		/// </summary>
 		static List<FileItem> DelayedRemoteLookupFiles = new List<FileItem>();
 
-		/**
-		 * Clears the FileItem caches.
-		 */
+		/// <summary>
+		/// Clears the FileItem caches.
+		/// </summary>
 		public static void ClearCaches()
 		{
 			UniqueSourceFileMap.Clear();
 			DelayedRemoteLookupFiles.Clear();
 		}
 
-		/**
-		 * Resolve any outstanding remote file info lookups
-		 */
+		/// <summary>
+		/// Resolve any outstanding remote file info lookups
+		/// </summary>
 		private void LookupOutstandingFiles()
 		{
 			// for remote files, look up any outstanding files
@@ -214,13 +248,13 @@ namespace UnrealBuildTool
 			}
 		}
 
-		/** @return The FileItem that represents the given file path. */
+		/// <returns>The FileItem that represents the given file path.</returns>
 		public static FileItem GetItemByPath(string FilePath)
 		{
 			return GetItemByFileReference(new FileReference(FilePath));
 		}
 
-		/** @return The FileItem that represents the given a full file path. */
+		/// <returns>The FileItem that represents the given a full file path.</returns>
 		public static FileItem GetItemByFileReference(FileReference Reference)
 		{
 			FileItem Result = null;
@@ -234,7 +268,7 @@ namespace UnrealBuildTool
 			}
 		}
 
-		/** @return The remote FileItem that represents the given file path. */
+		/// <returns>The remote FileItem that represents the given file path.</returns>
 		public static FileItem GetRemoteItemByPath(string AbsoluteRemotePath, UnrealTargetPlatform Platform)
 		{
 			if (AbsoluteRemotePath.StartsWith("."))
@@ -255,13 +289,17 @@ namespace UnrealBuildTool
 			}
 		}
 
-		/** If the given file path identifies a file that already exists, returns the FileItem that represents it. */
+		/// <summary>
+		/// If the given file path identifies a file that already exists, returns the FileItem that represents it.
+		/// </summary>
 		public static FileItem GetExistingItemByPath(string FileName)
 		{
 			return GetExistingItemByFileReference(new FileReference(FileName));
 		}
 
-		/** If the given file path identifies a file that already exists, returns the FileItem that represents it. */
+		/// <summary>
+		/// If the given file path identifies a file that already exists, returns the FileItem that represents it.
+		/// </summary>
 		public static FileItem GetExistingItemByFileReference(FileReference FileRef)
 		{
 			FileItem Result = GetItemByFileReference(FileRef);
@@ -287,10 +325,10 @@ namespace UnrealBuildTool
 			return (Encoding.UTF8.GetByteCount(Str) == Str.Length) ? Encoding.ASCII : new UTF8Encoding(false);
 		}
 
-		/**
-		 * Creates a text file with the given contents.  If the contents of the text file aren't changed, it won't write the new contents to
-		 * the file to avoid causing an action to be considered outdated.
-		 */
+		/// <summary>
+		/// Creates a text file with the given contents.  If the contents of the text file aren't changed, it won't write the new contents to
+		/// the file to avoid causing an action to be considered outdated.
+		/// </summary>
 		public static FileItem CreateIntermediateTextFile(FileReference AbsolutePath, string Contents)
 		{
 			// Create the directory if it doesn't exist.
@@ -305,7 +343,9 @@ namespace UnrealBuildTool
 			return GetItemByFileReference(AbsolutePath);
 		}
 
-		/** Deletes the file. */
+		/// <summary>
+		/// Deletes the file.
+		/// </summary>
 		public void Delete()
 		{
 			Debug.Assert(_bExists);
@@ -351,7 +391,9 @@ namespace UnrealBuildTool
 			while (!bFileDeletedSuccessfully && (DeleteTryCount < MaxRetryCount));
 		}
 
-		/** Initialization constructor. */
+		/// <summary>
+		/// Initialization constructor.
+		/// </summary>
 		protected FileItem(FileReference InFile)
 		{
 			Reference = InFile;
@@ -369,7 +411,9 @@ namespace UnrealBuildTool
 		}
 
 
-		/** ISerializable: Constructor called when this object is deserialized */
+		/// <summary>
+		/// ISerializable: Constructor called when this object is deserialized
+		/// </summary>
 		protected FileItem(SerializationInfo SerializationInfo, StreamingContext StreamingContext)
 		{
 			ProducingAction = (Action)SerializationInfo.GetValue("pa", typeof(Action));
@@ -404,7 +448,9 @@ namespace UnrealBuildTool
 		}
 
 
-		/** ISerializable: Called when serialized to report additional properties that should be saved */
+		/// <summary>
+		/// ISerializable: Called when serialized to report additional properties that should be saved
+		/// </summary>
 		public void GetObjectData(SerializationInfo SerializationInfo, StreamingContext StreamingContext)
 		{
 			SerializationInfo.AddValue("pa", ProducingAction);
@@ -453,7 +499,9 @@ namespace UnrealBuildTool
 			}
 		}
 
-		/** Initialization constructor for optionally remote files. */
+		/// <summary>
+		/// Initialization constructor for optionally remote files.
+		/// </summary>
 		protected FileItem(FileReference InReference, bool InIsRemoteFile, UnrealTargetPlatform Platform)
 		{
 			bIsRemoteFile = InIsRemoteFile;
diff --git a/Engine/Source/Programs/UnrealBuildTool/System/FlatCPPIncludeDepencencyCache.cs b/Engine/Source/Programs/UnrealBuildTool/System/FlatCPPIncludeDepencencyCache.cs
--- a/Engine/Source/Programs/UnrealBuildTool/System/FlatCPPIncludeDepencencyCache.cs
+++ b/Engine/Source/Programs/UnrealBuildTool/System/FlatCPPIncludeDepencencyCache.cs
@@ -31,10 +31,14 @@ namespace UnrealBuildTool
 	/// </summary>
 	public class FlatCPPIncludeDependencyCache
 	{
-		/** The version number for binary serialization */
+		/// <summary>
+		/// The version number for binary serialization
+		/// </summary>
 		const int FileVersion = 1;
 
-		/** The file signature for binary serialization */
+		/// <summary>
+		/// The file signature for binary serialization
+		/// </summary>
 		const int FileSignature = ('F' << 24) | ('C' << 16) | FileVersion;
 
 		/// File name of this cache, should be unique for every includes context (e.g. target)
