Convert everything in UBT to use standard C# XML comments, so IntelliSense can display useful information in tooltips.

diff --git a/Engine/Source/Programs/UnrealBuildTool/System/LinkEnvironment.cs b/Engine/Source/Programs/UnrealBuildTool/System/LinkEnvironment.cs
--- a/Engine/Source/Programs/UnrealBuildTool/System/LinkEnvironment.cs
+++ b/Engine/Source/Programs/UnrealBuildTool/System/LinkEnvironment.cs
@@ -12,19 +12,29 @@ namespace UnrealBuildTool
 	/// </summary>
 	public class LinkEnvironmentConfiguration : NativeBuildEnvironmentConfiguration
 	{
-		/** The directory to put the non-executable files in (PDBs, import library, etc) */
+		/// <summary>
+		/// The directory to put the non-executable files in (PDBs, import library, etc)
+		/// </summary>
 		public DirectoryReference OutputDirectory;
 
-		/** Intermediate file directory */
+		/// <summary>
+		/// Intermediate file directory
+		/// </summary>
 		public DirectoryReference IntermediateDirectory;
 
-		/** The directory to shadow source files in for syncing to remote compile servers */
+		/// <summary>
+		/// The directory to shadow source files in for syncing to remote compile servers
+		/// </summary>
 		public DirectoryReference LocalShadowDirectory = null;
 
-		/** The file path for the executable file that is output by the linker. */
+		/// <summary>
+		/// The file path for the executable file that is output by the linker.
+		/// </summary>
 		public List<FileReference> OutputFilePaths = new List<FileReference>();
 
-		/** Returns the OutputFilePath is there is only one entry in OutputFilePaths */
+		/// <summary>
+		/// Returns the OutputFilePath is there is only one entry in OutputFilePaths
+		/// </summary>
 		public FileReference OutputFilePath
 		{
 			get
@@ -37,82 +47,124 @@ namespace UnrealBuildTool
 			}
 		}
 
-		/** A list of the paths used to find libraries. */
+		/// <summary>
+		/// A list of the paths used to find libraries.
+		/// </summary>
 		public List<string> LibraryPaths = new List<string>();
 
-		/** A list of libraries to exclude from linking. */
+		/// <summary>
+		/// A list of libraries to exclude from linking.
+		/// </summary>
 		public List<string> ExcludedLibraries = new List<string>();
 
-		/** A list of additional libraries to link in. */
+		/// <summary>
+		/// A list of additional libraries to link in.
+		/// </summary>
 		public List<string> AdditionalLibraries = new List<string>();
 
-		/** A list of additional frameworks to link in. */
+		/// <summary>
+		/// A list of additional frameworks to link in.
+		/// </summary>
 		public List<UEBuildFramework> AdditionalFrameworks = new List<UEBuildFramework>();
 
-		/** For builds that execute on a remote machine (e.g. iPhone), this list contains additional files that
-			need to be copied over in order for the app to link successfully.  Source/header files and PCHs are
-			automatically copied.  Usually this is simply a list of precompiled third party library dependencies. */
+		/// <summary>
+		/// For builds that execute on a remote machine (e.g. iPhone), this list contains additional files that
+		/// need to be copied over in order for the app to link successfully.  Source/header files and PCHs are
+		/// automatically copied.  Usually this is simply a list of precompiled third party library dependencies.
+		/// </summary>
 		public List<string> AdditionalShadowFiles = new List<string>();
 
-		/** The iOS/Mac frameworks to link in */
+		/// <summary>
+		/// The iOS/Mac frameworks to link in
+		/// </summary>
 		public List<string> Frameworks = new List<string>();
 		public List<string> WeakFrameworks = new List<string>();
 
-		/** iOS/Mac resources that should be copied to the app bundle */
+		/// <summary>
+		/// iOS/Mac resources that should be copied to the app bundle
+		/// </summary>
 		public List<UEBuildBundleResource> AdditionalBundleResources = new List<UEBuildBundleResource>();
 
-		/**
-		 * A list of the dynamically linked libraries that shouldn't be loaded until they are first called
-		 * into.
-		 */
+		/// <summary>
+		/// A list of the dynamically linked libraries that shouldn't be loaded until they are first called
+		/// into.
+		/// </summary>
 		public List<string> DelayLoadDLLs = new List<string>();
 
-		/** Additional arguments to pass to the linker. */
+		/// <summary>
+		/// Additional arguments to pass to the linker.
+		/// </summary>
 		public string AdditionalArguments = "";
 
-		/** True if debug info should be created. */
+		/// <summary>
+		/// True if debug info should be created.
+		/// </summary>
 		public bool bCreateDebugInfo = true;
 
-		/** Whether the CLR (Common Language Runtime) support should be enabled for C++ targets (C++/CLI). */
+		/// <summary>
+		/// Whether the CLR (Common Language Runtime) support should be enabled for C++ targets (C++/CLI).
+		/// </summary>
 		public CPPCLRMode CLRMode = CPPCLRMode.CLRDisabled;
 
-		/** True if we're compiling .cpp files that will go into a library (.lib file) */
+		/// <summary>
+		/// True if we're compiling .cpp files that will go into a library (.lib file)
+		/// </summary>
 		public bool bIsBuildingLibrary = false;
 
-		/** True if we're compiling a DLL */
+		/// <summary>
+		/// True if we're compiling a DLL
+		/// </summary>
 		public bool bIsBuildingDLL = false;
 
-		/** True if this is a console application that's being build */
+		/// <summary>
+		/// True if this is a console application that's being build
+		/// </summary>
 		public bool bIsBuildingConsoleApplication = false;
 
-		/** This setting is replaced by UEBuildBinaryConfiguration.bBuildAdditionalConsoleApp. */
+		/// <summary>
+		/// This setting is replaced by UEBuildBinaryConfiguration.bBuildAdditionalConsoleApp.
+		/// </summary>
 		[Obsolete("This setting is replaced by UEBuildBinaryConfiguration.bBuildAdditionalConsoleApp. It is explicitly set to true for editor targets, and defaults to false otherwise.")]
 		public bool bBuildAdditionalConsoleApplication { set { } }
 
-		/** If set, overrides the program entry function on Windows platform.  This is used by the base UE4
-		    program so we can link in either command-line mode or windowed mode without having to recompile the Launch module. */
+		/// <summary>
+		/// If set, overrides the program entry function on Windows platform.  This is used by the base UE4
+		/// program so we can link in either command-line mode or windowed mode without having to recompile the Launch module.
+		/// </summary>
 		public string WindowsEntryPointOverride = String.Empty;
 
-		/** True if we're building a EXE/DLL target with an import library, and that library is needed by a dependency that
-		    we're directly dependent on. */
+		/// <summary>
+		/// True if we're building a EXE/DLL target with an import library, and that library is needed by a dependency that
+		/// we're directly dependent on.
+		/// </summary>
 		public bool bIsCrossReferenced = false;
 
-		/** True if we should include dependent libraries when building a static library */
+		/// <summary>
+		/// True if we should include dependent libraries when building a static library
+		/// </summary>
 		public bool bIncludeDependentLibrariesInLibrary = false;
 
-		/** True if the application we're linking has any exports, and we should be expecting the linker to
-		    generate a .lib and/or .exp file along with the target output file */
+		/// <summary>
+		/// True if the application we're linking has any exports, and we should be expecting the linker to
+		/// generate a .lib and/or .exp file along with the target output file
+		/// </summary>
 		public bool bHasExports = true;
 
-		/** True if we're building a .NET assembly (e.g. C# project) */
+		/// <summary>
+		/// True if we're building a .NET assembly (e.g. C# project)
+		/// </summary>
 		public bool bIsBuildingDotNetAssembly = false;
 
-		/** Default constructor. */
+		/// <summary>
+		/// Default constructor.
+		/// </summary>
 		public LinkEnvironmentConfiguration()
 		{
 		}
 
-		/** Copy constructor. */
+		/// <summary>
+		/// Copy constructor.
+		/// </summary>
 		public LinkEnvironmentConfiguration(LinkEnvironmentConfiguration InCopyEnvironment) :
 			base(InCopyEnvironment)
 		{
@@ -142,27 +194,41 @@ namespace UnrealBuildTool
 		}
 	}
 
-	/** Encapsulates the environment that is used to link object files. */
+	/// <summary>
+	/// Encapsulates the environment that is used to link object files.
+	/// </summary>
 	public class LinkEnvironment
 	{
-		/** Whether we're linking in monolithic mode. Determines if linking should produce import library file. Relevant only for VC++, clang stores imports in shared library. */
+		/// <summary>
+		/// Whether we're linking in monolithic mode. Determines if linking should produce import library file. Relevant only for VC++, clang stores imports in shared library.
+		/// </summary>
 		public bool bShouldCompileMonolithic = false;
 
-		/** A list of the object files to be linked. */
+		/// <summary>
+		/// A list of the object files to be linked.
+		/// </summary>
 		public List<FileItem> InputFiles = new List<FileItem>();
 
-		/** A list of dependent static or import libraries that need to be linked. */
+		/// <summary>
+		/// A list of dependent static or import libraries that need to be linked.
+		/// </summary>
 		public List<FileItem> InputLibraries = new List<FileItem>();
 
-		/** The LinkEnvironmentConfiguration. */
+		/// <summary>
+		/// The LinkEnvironmentConfiguration.
+		/// </summary>
 		public LinkEnvironmentConfiguration Config = new LinkEnvironmentConfiguration();
 
-		/** Default constructor. */
+		/// <summary>
+		/// Default constructor.
+		/// </summary>
 		public LinkEnvironment()
 		{
 		}
 
-		/** Copy constructor. */
+		/// <summary>
+		/// Copy constructor.
+		/// </summary>
 		protected LinkEnvironment(LinkEnvironment InCopyEnvironment)
 		{
 			InputFiles.AddRange(InCopyEnvironment.InputFiles);
diff --git a/Engine/Source/Programs/UnrealBuildTool/System/LocalExecutor.cs b/Engine/Source/Programs/UnrealBuildTool/System/LocalExecutor.cs
--- a/Engine/Source/Programs/UnrealBuildTool/System/LocalExecutor.cs
+++ b/Engine/Source/Programs/UnrealBuildTool/System/LocalExecutor.cs
@@ -14,20 +14,35 @@ namespace UnrealBuildTool
 {
 	class ActionThread
 	{
-		/** Cache the exit code from the command so that the executor can report errors */
+		/// <summary>
+		/// Cache the exit code from the command so that the executor can report errors
+		/// </summary>
 		public int ExitCode = 0;
-		/** Set to true only when the local or RPC action is complete */
+
+		/// <summary>
+		/// Set to true only when the local or RPC action is complete
+		/// </summary>
 		public bool bComplete = false;
-		/** Cache the action that this thread is managing */
+
+		/// <summary>
+		/// Cache the action that this thread is managing
+		/// </summary>
 		Action Action;
-		/** For reporting status to the user */
+
+		/// <summary>
+		/// For reporting status to the user
+		/// </summary>
 		int JobNumber;
 		int TotalJobs;
 
-		/** Regex that matches environment variables in $(Variable) format. */
+		/// <summary>
+		/// Regex that matches environment variables in $(Variable) format.
+		/// </summary>
 		static Regex EnvironmentVariableRegex = new Regex("\\$\\(([\\d\\w]+)\\)");
 
-		/** Replaces the environment variables references in a string with their values. */
+		/// <summary>
+		/// Replaces the environment variables references in a string with their values.
+		/// </summary>
 		public static string ExpandEnvironmentVariables(string Text)
 		{
 			foreach (Match EnvironmentVariableMatch in EnvironmentVariableRegex.Matches(Text))
@@ -38,9 +53,9 @@ namespace UnrealBuildTool
 			return Text;
 		}
 
-		/**
-		 * Constructor, takes the action to process
-		 */
+		/// <summary>
+		/// Constructor, takes the action to process
+		/// </summary>
 		public ActionThread(Action InAction, int InJobNumber, int InTotalJobs)
 		{
 			Action = InAction;
@@ -48,9 +63,9 @@ namespace UnrealBuildTool
 			TotalJobs = InTotalJobs;
 		}
 
-		/**
-		 * Sends a string to an action's OutputEventHandler
-		 */
+		/// <summary>
+		/// Sends a string to an action's OutputEventHandler
+		/// </summary>
 		private static void SendOutputToEventHandler(Action Action, string Output)
 		{
 			// pass the output to any handler requested
@@ -85,12 +100,11 @@ namespace UnrealBuildTool
 		}
 
 
-		/**
-		 * Used when debuging Actions outputs all action return values to debug out
-		 * 
-		 * @param	sender		Sending object
-		 * @param	e			Event arguments (In this case, the line of string output)
-		 */
+		/// <summary>
+		/// Used when debuging Actions outputs all action return values to debug out
+		/// </summary>
+		/// <param name="sender"> Sending object</param>
+		/// <param name="e">  Event arguments (In this case, the line of string output)</param>
 		protected void ActionDebugOutput(object sender, DataReceivedEventArgs e)
 		{
 			var Output = e.Data;
@@ -103,9 +117,9 @@ namespace UnrealBuildTool
 		}
 
 
-		/**
-		 * The actual function to run in a thread. This is potentially long and blocking
-		 */
+		/// <summary>
+		/// The actual function to run in a thread. This is potentially long and blocking
+		/// </summary>
 		private void ThreadFunc()
 		{
 			// thread start time
@@ -312,9 +326,9 @@ namespace UnrealBuildTool
 			bComplete = true;
 		}
 
-		/**
-		 * Starts a thread and runs the action in that thread
-		 */
+		/// <summary>
+		/// Starts a thread and runs the action in that thread
+		/// </summary>
 		public void Run()
 		{
 			Thread T = new Thread(ThreadFunc);
@@ -324,10 +338,10 @@ namespace UnrealBuildTool
 
 	public class LocalExecutor
 	{
-		/**
-		 * Executes the specified actions locally.
-		 * @return True if all the tasks successfully executed, or false if any of them failed.
-		 */
+		/// <summary>
+		/// Executes the specified actions locally.
+		/// </summary>
+		/// <returns>True if all the tasks successfully executed, or false if any of them failed.</returns>
 		public static bool ExecuteActions(List<Action> Actions)
 		{
 			// Time to sleep after each iteration of the loop in order to not busy wait.
diff --git a/Engine/Source/Programs/UnrealBuildTool/System/Project.cs b/Engine/Source/Programs/UnrealBuildTool/System/Project.cs
--- a/Engine/Source/Programs/UnrealBuildTool/System/Project.cs
+++ b/Engine/Source/Programs/UnrealBuildTool/System/Project.cs
@@ -399,7 +399,9 @@ namespace UnrealBuildTool
 			return new SourceFile(InitFilePath, InitProjectSubFolder);
 		}
 
-		/** Takes the given path and tries to rebase it relative to the project or solution directory variables. */
+		/// <summary>
+		/// Takes the given path and tries to rebase it relative to the project or solution directory variables.
+		/// </summary>
 		public string NormalizeProjectPath(string InputPath)
 		{
 			// If the path is rooted in an environment variable, leave it be.
@@ -413,7 +415,9 @@ namespace UnrealBuildTool
 			}
 		}
 
-		/** Takes the given path and tries to rebase it relative to the project. */
+		/// <summary>
+		/// Takes the given path and tries to rebase it relative to the project.
+		/// </summary>
 		public string NormalizeProjectPath(FileReference InputPath)
 		{
 			// Try to make it relative to the solution directory.
@@ -427,7 +431,9 @@ namespace UnrealBuildTool
 			}
 		}
 
-		/** Takes the given path, normalizes it, and quotes it if necessary. */
+		/// <summary>
+		/// Takes the given path, normalizes it, and quotes it if necessary.
+		/// </summary>
 		public string EscapePath(string InputPath)
 		{
 			string Result = InputPath;
@@ -438,7 +444,9 @@ namespace UnrealBuildTool
 			return Result;
 		}
 
-		/** Visualizer for the debugger */
+		/// <summary>
+		/// Visualizer for the debugger
+		/// </summary>
 		public override string ToString()
 		{
 			return ProjectFilePath.ToString();
diff --git a/Engine/Source/Programs/UnrealBuildTool/System/RPCUtilHelper.cs b/Engine/Source/Programs/UnrealBuildTool/System/RPCUtilHelper.cs
--- a/Engine/Source/Programs/UnrealBuildTool/System/RPCUtilHelper.cs
+++ b/Engine/Source/Programs/UnrealBuildTool/System/RPCUtilHelper.cs
@@ -20,16 +20,24 @@ namespace UnrealBuildTool
 {
 	public class RPCUtilHelper
 	{
-		/** The Mac we are compiling on */
+		/// <summary>
+		/// The Mac we are compiling on
+		/// </summary>
 		private static string MacName;
 
-		/** A socket per command thread */
+		/// <summary>
+		/// A socket per command thread
+		/// </summary>
 		private static Hashtable CommandThreadSockets = new Hashtable();
 
-		/** Time difference between remote and local idea's of UTC time */
+		/// <summary>
+		/// Time difference between remote and local idea's of UTC time
+		/// </summary>
 		private static TimeSpan TimeDifferenceFromRemote = new TimeSpan(0);
 
-		/** The number of commands the remote side should be able to run at once */
+		/// <summary>
+		/// The number of commands the remote side should be able to run at once
+		/// </summary>
 		private static int MaxRemoteCommandsAllowed = 0;
 
 		private static DateTime RemoteToLocalTime(string RemoteTime)
@@ -129,9 +137,9 @@ namespace UnrealBuildTool
 			return RemoteToolChain.RemoteToolChainErrorCode.NoError;
 		}
 
-		/**
-		 * Handle a thread ending
-		 */
+		/// <summary>
+		/// Handle a thread ending
+		/// </summary>
 		public static void OnThreadComplete()
 		{
 			lock (CommandThreadSockets)
@@ -171,10 +179,10 @@ namespace UnrealBuildTool
 			return ThreadSocket;
 		}
 
-		/**
-		 * This function should be used as the ActionHandler delegate method for Actions that
-		 * need to run over RPCUtility. It will block until the remote command completes
-		 */
+		/// <summary>
+		/// This function should be used as the ActionHandler delegate method for Actions that
+		/// need to run over RPCUtility. It will block until the remote command completes
+		/// </summary>
 		static public void RPCActionHandler(Action Action, out int ExitCode, out string Output)
 		{
 			Hashtable Results = RPCUtilHelper.Command(Action.WorkingDirectory, Action.CommandPath, Action.CommandArguments,
@@ -202,9 +210,7 @@ namespace UnrealBuildTool
 			}
 		}
 
-		/** 
-		 * @return the modification time on the remote machine, accounting for rough difference in time between the two machines
-		 */
+		/// <returns>the modification time on the remote machine, accounting for rough difference in time between the two machines</returns>
 		public static bool GetRemoteFileInfo(string RemotePath, out DateTime ModificationTime, out long Length)
 		{
 			if (RemoteToolChain.bUseRPCUtil)
diff --git a/Engine/Source/Programs/UnrealBuildTool/System/RulesCompiler.cs b/Engine/Source/Programs/UnrealBuildTool/System/RulesCompiler.cs
--- a/Engine/Source/Programs/UnrealBuildTool/System/RulesCompiler.cs
+++ b/Engine/Source/Programs/UnrealBuildTool/System/RulesCompiler.cs
@@ -194,99 +194,159 @@ namespace UnrealBuildTool
 		/// Precompiled header usage for this module
 		public PCHUsageMode PCHUsage = PCHUsageMode.Default;
 
-		/** Use run time type information */
+		/// <summary>
+		/// Use run time type information
+		/// </summary>
 		public bool bUseRTTI = false;
 
-		/** Use AVX instructions */
+		/// <summary>
+		/// Use AVX instructions
+		/// </summary>
 		public bool bUseAVX = false;
 
-		/** Enable buffer security checks.  This should usually be enabled as it prevents severe security risks. */
+		/// <summary>
+		/// Enable buffer security checks.  This should usually be enabled as it prevents severe security risks.
+		/// </summary>
 		public bool bEnableBufferSecurityChecks = true;
 
-		/** Enable exception handling */
+		/// <summary>
+		/// Enable exception handling
+		/// </summary>
 		public bool bEnableExceptions = false;
 
-		/** Enable warnings for shadowed variables */
+		/// <summary>
+		/// Enable warnings for shadowed variables
+		/// </summary>
 		public bool bEnableShadowVariableWarnings = true;
 
-		/** If true and unity builds are enabled, this module will build without unity. */
+		/// <summary>
+		/// If true and unity builds are enabled, this module will build without unity.
+		/// </summary>
 		public bool bFasterWithoutUnity = false;
 
-		/** The number of source files in this module before unity build will be activated for that module.  If set to
-		    anything besides -1, will override the default setting which is controlled by MinGameModuleSourceFilesForUnityBuild */
+		/// <summary>
+		/// The number of source files in this module before unity build will be activated for that module.  If set to
+		/// anything besides -1, will override the default setting which is controlled by MinGameModuleSourceFilesForUnityBuild
+		/// </summary>
 		public int MinSourceFilesForUnityBuildOverride = 0;
 
-		/** Overrides BuildConfiguration.MinFilesUsingPrecompiledHeader if non-zero. */
+		/// <summary>
+		/// Overrides BuildConfiguration.MinFilesUsingPrecompiledHeader if non-zero.
+		/// </summary>
 		public int MinFilesUsingPrecompiledHeaderOverride = 0;
 
-		/**  Module uses a #import so must be built locally when compiling with SN-DBS */
+		/// <summary>
+		/// Module uses a #import so must be built locally when compiling with SN-DBS
+		/// </summary>
 		public bool bBuildLocallyWithSNDBS = false;
 
+		/// <summary>
 		/// List of modules with header files that our module's public headers needs access to, but we don't need to "import" or link against.
+		/// </summary>
 		public List<string> PublicIncludePathModuleNames = new List<string>();
 
+		/// <summary>
 		/// List of public dependency module names.  These are modules that are required by our public source files.
+		/// </summary>
 		public List<string> PublicDependencyModuleNames = new List<string>();
 
+		/// <summary>
 		/// List of modules with header files that our module's private code files needs access to, but we don't need to "import" or link against.
+		/// </summary>
 		public List<string> PrivateIncludePathModuleNames = new List<string>();
 
+		/// <summary>
 		/// List of private dependency module names.  These are modules that our private code depends on but nothing in our public
 		/// include files depend on.
+		/// </summary>
 		public List<string> PrivateDependencyModuleNames = new List<string>();
 
+		/// <summary>
 		/// List of module dependencies that should be treated as circular references.  This modules must have already been added to
 		/// either the public or private dependent module list.
+		/// </summary>
 		public List<string> CircularlyReferencedDependentModules = new List<string>();
 
+		/// <summary>
 		/// System include paths.  These are public stable header file directories that are not checked when resolving header dependencies.
+		/// </summary>
 		public List<string> PublicSystemIncludePaths = new List<string>();
 
+		/// <summary>
 		/// List of all paths to include files that are exposed to other modules
+		/// </summary>
 		public List<string> PublicIncludePaths = new List<string>();
 
+		/// <summary>
 		/// List of all paths to this module's internal include files, not exposed to other modules
+		/// </summary>
 		public List<string> PrivateIncludePaths = new List<string>();
 
+		/// <summary>
 		/// List of library paths - typically used for External (third party) modules
+		/// </summary>
 		public List<string> PublicLibraryPaths = new List<string>();
 
+		/// <summary>
 		/// List of addition libraries - typically used for External (third party) modules
+		/// </summary>
 		public List<string> PublicAdditionalLibraries = new List<string>();
 
+		/// <summary>
 		// List of frameworks
+		/// </summary>
 		public List<string> PublicFrameworks = new List<string>();
 
+		/// <summary>
 		// List of weak frameworks (for OS version transitions)
+		/// </summary>
 		public List<string> PublicWeakFrameworks = new List<string>();
 
+		/// <summary>
 		/// List of addition frameworks - typically used for External (third party) modules on Mac and iOS
+		/// </summary>
 		public List<UEBuildFramework> PublicAdditionalFrameworks = new List<UEBuildFramework>();
 
+		/// <summary>
 		/// List of addition resources that should be copied to the app bundle for Mac or iOS
+		/// </summary>
 		public List<UEBuildBundleResource> AdditionalBundleResources = new List<UEBuildBundleResource>();
 
+		/// <summary>
 		/// For builds that execute on a remote machine (e.g. iOS), this list contains additional files that
 		/// need to be copied over in order for the app to link successfully.  Source/header files and PCHs are
 		/// automatically copied.  Usually this is simply a list of precompiled third party library dependencies.
+		/// </summary>
 		public List<string> PublicAdditionalShadowFiles = new List<string>();
 
+		/// <summary>
 		/// List of delay load DLLs - typically used for External (third party) modules
+		/// </summary>
 		public List<string> PublicDelayLoadDLLs = new List<string>();
 
+		/// <summary>
 		/// Additional compiler definitions for this module
+		/// </summary>
 		public List<string> Definitions = new List<string>();
 
-		/** CLR modules only: The assemblies referenced by the module's private implementation. */
+		/// <summary>
+		/// CLR modules only: The assemblies referenced by the module's private implementation.
+		/// </summary>
 		public List<string> PrivateAssemblyReferences = new List<string>();
 
+		/// <summary>
 		/// Addition modules this module may require at run-time 
+		/// </summary>
 		public List<string> DynamicallyLoadedModuleNames = new List<string>();
 
+		/// <summary>
 		/// Extra modules this module may require at run time, that are on behalf of another platform (i.e. shader formats and the like)
+		/// </summary>
 		public List<string> PlatformSpecificDynamicallyLoadedModuleNames = new List<string>();
 
+		/// <summary>
 		/// List of files which this module depends on at runtime. These files will be staged along with the target.
+		/// </summary>
 		public List<RuntimeDependency> RuntimeDependencies = new List<RuntimeDependency>();
 
 		/// <summary>
@@ -385,7 +445,9 @@ namespace UnrealBuildTool
 			Definitions.Add(string.Format("WITH_BOX2D={0}", bSupported ? 1 : 0));
 		}
 
-		/** Redistribution override flag for this module. */
+		/// <summary>
+		/// Redistribution override flag for this module.
+		/// </summary>
 		public bool? IsRedistributableOverride { get; set; }
 	}
 
diff --git a/Engine/Source/Programs/UnrealBuildTool/System/SNDBS.cs b/Engine/Source/Programs/UnrealBuildTool/System/SNDBS.cs
--- a/Engine/Source/Programs/UnrealBuildTool/System/SNDBS.cs
+++ b/Engine/Source/Programs/UnrealBuildTool/System/SNDBS.cs
@@ -18,7 +18,9 @@ namespace UnrealBuildTool
 		static private int MaxActionsToExecuteInParallel;
 		static private int JobNumber;
 
-		/** The possible result of executing tasks with SN-DBS. */
+		/// <summary>
+		/// The possible result of executing tasks with SN-DBS.
+		/// </summary>
 		public enum ExecutionResult
 		{
 			Unavailable,
@@ -26,12 +28,11 @@ namespace UnrealBuildTool
 			TasksSucceeded,
 		}
 
-		/**
-		 * Used when debugging Actions outputs all action return values to debug out
-		 * 
-		 * @param	sender		Sending object
-		 * @param	e			Event arguments (In this case, the line of string output)
-		 */
+		/// <summary>
+		/// Used when debugging Actions outputs all action return values to debug out
+		/// </summary>
+		/// <param name="sender"> Sending object</param>
+		/// <param name="e">  Event arguments (In this case, the line of string output)</param>
 		static protected void ActionDebugOutput(object sender, DataReceivedEventArgs e)
 		{
 			var Output = e.Data;
diff --git a/Engine/Source/Programs/UnrealBuildTool/System/ThirdPartyHeaderFinder.cs b/Engine/Source/Programs/UnrealBuildTool/System/ThirdPartyHeaderFinder.cs
--- a/Engine/Source/Programs/UnrealBuildTool/System/ThirdPartyHeaderFinder.cs
+++ b/Engine/Source/Programs/UnrealBuildTool/System/ThirdPartyHeaderFinder.cs
@@ -8,24 +8,34 @@ using System.Diagnostics;
 
 namespace UnrealBuildTool
 {
-	/**
-	 * Helper class for finding 3rd party headers included in public engine headers.
-	 **/
+	/// <summary>
+	/// Helper class for finding 3rd party headers included in public engine headers.
+	/// </summary>
 	public class ThirdPartyHeaderFinder
 	{
-		/** List of all third party headers from the current game's modules */
+		/// <summary>
+		/// List of all third party headers from the current game's modules
+		/// </summary>
 		static List<Header> ThirdPartyHeaders;
-		/** List of all engine headers referenced by the current game. */
+
+		/// <summary>
+		/// List of all engine headers referenced by the current game.
+		/// </summary>
 		static List<Header> EngineHeaders;
 
-		/**
-		 * Helper class for storing headers included from public engine headers
-		 **/
+		/// <summary>
+		/// Helper class for storing headers included from public engine headers
+		/// </summary>
 		class IncludePath
 		{
-			/** Path of public engine header includes from which the other header is included. */
+			/// <summary>
+			/// Path of public engine header includes from which the other header is included.
+			/// </summary>
 			public string PublicHeaderPath;
-			/** The other header included from the public engine header (may be third party or private) */
+
+			/// <summary>
+			/// The other header included from the public engine header (may be third party or private)
+			/// </summary>
 			public Header OtherHeader;
 
 			public IncludePath(string InPublicPath, Header Other)
@@ -35,20 +45,34 @@ namespace UnrealBuildTool
 			}
 		}
 
-		/**
-		 * Helper class for storing info on a header file.
-		 **/
+		/// <summary>
+		/// Helper class for storing info on a header file.
+		/// </summary>
 		class Header
 		{
-			/** Relative path to the header file. */
+			/// <summary>
+			/// Relative path to the header file.
+			/// </summary>
 			public string Path;
-			/** Name (with extension) of the header file. */
+
+			/// <summary>
+			/// Name (with extension) of the header file.
+			/// </summary>
 			public string Name;
-			/** If true this header is a public engine header file. */
+
+			/// <summary>
+			/// If true this header is a public engine header file.
+			/// </summary>
 			public bool IsPublic;
-			/** Contents of the header file. */
+
+			/// <summary>
+			/// Contents of the header file.
+			/// </summary>
 			public string Contents;
-			/** Pregenerated variations of includes. */
+
+			/// <summary>
+			/// Pregenerated variations of includes.
+			/// </summary>
 			public string[] TestIncludes;
 
 			public Header(string InPath, bool InPublic)
@@ -78,9 +102,9 @@ namespace UnrealBuildTool
 			}
 		}
 
-		/**
-		 * Find all third party and private header includes in public engine headers.
-		 **/
+		/// <summary>
+		/// Find all third party and private header includes in public engine headers.
+		/// </summary>
 		public static void FindThirdPartyIncludes(UnrealTargetPlatform Platform, UnrealTargetConfiguration Configuration)
 		{
 			Log.TraceInformation("Looking for third party header includes in public engine header files (this may take a few minutes)...");
diff --git a/Engine/Source/Programs/UnrealBuildTool/System/UEPlatformProjectGenerator.cs b/Engine/Source/Programs/UnrealBuildTool/System/UEPlatformProjectGenerator.cs
--- a/Engine/Source/Programs/UnrealBuildTool/System/UEPlatformProjectGenerator.cs
+++ b/Engine/Source/Programs/UnrealBuildTool/System/UEPlatformProjectGenerator.cs
@@ -8,19 +8,18 @@ using System.IO;
 
 namespace UnrealBuildTool
 {
-	/**
-	 *	Base class for platform-specific project generators 
-	 */
+	/// <summary>
+	/// Base class for platform-specific project generators
+	/// </summary>
 	public abstract class UEPlatformProjectGenerator
 	{
 		static Dictionary<UnrealTargetPlatform, UEPlatformProjectGenerator> ProjectGeneratorDictionary = new Dictionary<UnrealTargetPlatform, UEPlatformProjectGenerator>();
 
-		/**
-		 *	Register the given platforms UEPlatformProjectGenerator instance
-		 *	
-		 *	@param	InPlatform			The UnrealTargetPlatform to register with
-		 *	@param	InProjectGenerator	The UEPlatformProjectGenerator instance to use for the InPlatform
-		 */
+		/// <summary>
+		/// Register the given platforms UEPlatformProjectGenerator instance
+		/// </summary>
+		/// <param name="InPlatform">  The UnrealTargetPlatform to register with</param>
+		/// <param name="InProjectGenerator">The UEPlatformProjectGenerator instance to use for the InPlatform</param>
 		public static void RegisterPlatformProjectGenerator(UnrealTargetPlatform InPlatform, UEPlatformProjectGenerator InProjectGenerator)
 		{
 			// Make sure the build platform is legal
@@ -44,14 +43,12 @@ namespace UnrealBuildTool
 			}
 		}
 
-		/**
-		 *	Retrieve the UEPlatformProjectGenerator instance for the given TargetPlatform
-		 *	
-		 *	@param	InPlatform					The UnrealTargetPlatform being built
-		 *	@param	bInAllowFailure				If true, do not throw an exception and return null
-		 *	
-		 *	@return	UEPlatformProjectGenerator	The instance of the project generator
-		 */
+		/// <summary>
+		/// Retrieve the UEPlatformProjectGenerator instance for the given TargetPlatform
+		/// </summary>
+		/// <param name="InPlatform">    The UnrealTargetPlatform being built</param>
+		/// <param name="bInAllowFailure">   If true, do not throw an exception and return null</param>
+		/// <returns>UEPlatformProjectGenerator The instance of the project generator</returns>
 		public static UEPlatformProjectGenerator GetPlatformProjectGenerator(UnrealTargetPlatform InPlatform, bool bInAllowFailure = false)
 		{
 			if (ProjectGeneratorDictionary.ContainsKey(InPlatform) == true)
@@ -110,9 +107,9 @@ namespace UnrealBuildTool
 			return bRequiresVSUserFileGeneration;
 		}
 
-		/**
-		 *	Register the platform with the UEPlatformProjectGenerator class
-		 */
+		/// <summary>
+		/// Register the platform with the UEPlatformProjectGenerator class
+		/// </summary>
 		public abstract void RegisterPlatformProjectGenerator();
 
 		public virtual void GenerateGameProjectStub(ProjectFileGenerator InGenerator, string InTargetName, string InTargetFilepath, TargetRules InTargetRules,
@@ -135,80 +132,68 @@ namespace UnrealBuildTool
 		///
 		///	VisualStudio project generation functions
 		///	
-		/**
-		 *	Whether this build platform has native support for VisualStudio
-		 *	
-		 *	@param	InPlatform			The UnrealTargetPlatform being built
-		 *	@param	InConfiguration		The UnrealTargetConfiguration being built
-		 *	
-		 *	@return	bool				true if native VisualStudio support (or custom VSI) is available
-		 */
+		/// <summary>
+		/// Whether this build platform has native support for VisualStudio
+		/// </summary>
+		/// <param name="InPlatform">  The UnrealTargetPlatform being built</param>
+		/// <param name="InConfiguration"> The UnrealTargetConfiguration being built</param>
+		/// <returns>bool    true if native VisualStudio support (or custom VSI) is available</returns>
 		public virtual bool HasVisualStudioSupport(UnrealTargetPlatform InPlatform, UnrealTargetConfiguration InConfiguration)
 		{
 			// By default, we assume this is true
 			return true;
 		}
 
-		/**
-		 *	Return the VisualStudio platform name for this build platform
-		 *	
-		 *	@param	InPlatform			The UnrealTargetPlatform being built
-		 *	@param	InConfiguration		The UnrealTargetConfiguration being built
-		 *	
-		 *	@return	string				The name of the platform that VisualStudio recognizes
-		 */
+		/// <summary>
+		/// Return the VisualStudio platform name for this build platform
+		/// </summary>
+		/// <param name="InPlatform">  The UnrealTargetPlatform being built</param>
+		/// <param name="InConfiguration"> The UnrealTargetConfiguration being built</param>
+		/// <returns>string    The name of the platform that VisualStudio recognizes</returns>
 		public virtual string GetVisualStudioPlatformName(UnrealTargetPlatform InPlatform, UnrealTargetConfiguration InConfiguration)
 		{
 			// By default, return the platform string
 			return InPlatform.ToString();
 		}
 
-		/**
-		 *	Return the platform toolset string to write into the project configuration
-		 *	
-		 *	@param	InPlatform			The UnrealTargetPlatform being built
-		 *	@param	InConfiguration		The UnrealTargetConfiguration being built
-		 *	
-		 *	@return	string				The custom configuration section for the project file; Empty string if it doesn't require one
-		 */
+		/// <summary>
+		/// Return the platform toolset string to write into the project configuration
+		/// </summary>
+		/// <param name="InPlatform">  The UnrealTargetPlatform being built</param>
+		/// <param name="InConfiguration"> The UnrealTargetConfiguration being built</param>
+		/// <returns>string    The custom configuration section for the project file; Empty string if it doesn't require one</returns>
 		public virtual string GetVisualStudioPlatformToolsetString(UnrealTargetPlatform InPlatform, UnrealTargetConfiguration InConfiguration, VCProjectFile InProjectFile)
 		{
 			return "";
 		}
 
-		/**
-		 * Return any custom property group lines
-		 *	
-		 *	@param	InPlatform			The UnrealTargetPlatform being built
-		 *	
-		 *	@return	string				The custom property import lines for the project file; Empty string if it doesn't require one
-		 */
+		/// <summary>
+		/// Return any custom property group lines
+		/// </summary>
+		/// <param name="InPlatform">  The UnrealTargetPlatform being built</param>
+		/// <returns>string    The custom property import lines for the project file; Empty string if it doesn't require one</returns>
 		public virtual string GetAdditionalVisualStudioPropertyGroups(UnrealTargetPlatform InPlatform)
 		{
 			return "";
 		}
 
-		/**
-		 * Return any custom property group lines
-		 *	
-		 *	@param	InPlatform			The UnrealTargetPlatform being built
-		 *	
-		 *	@return	string				The platform configuration type.  Defaults to "Makefile" unless overridden
-		 */
+		/// <summary>
+		/// Return any custom property group lines
+		/// </summary>
+		/// <param name="InPlatform">  The UnrealTargetPlatform being built</param>
+		/// <returns>string    The platform configuration type.  Defaults to "Makefile" unless overridden</returns>
 		public virtual string GetVisualStudioPlatformConfigurationType(UnrealTargetPlatform InPlatform)
 		{
 			return "Makefile";
 		}
 
-		/**
-		 * Return any custom paths for VisualStudio this platform requires
-		 * This include ReferencePath, LibraryPath, LibraryWPath, IncludePath and ExecutablePath.
-		 *	
-		 *	@param	InPlatform			The UnrealTargetPlatform being built
-		 *	@param	TargetType			The type of target (game or program)
-		 *	
-		 *	@return	string				The custom path lines for the project file; Empty string if it doesn't require one
-		 */
+		/// <summary>
+		/// Return any custom paths for VisualStudio this platform requires
+		/// This include ReferencePath, LibraryPath, LibraryWPath, IncludePath and ExecutablePath.
+		/// </summary>
+		/// <param name="InPlatform">  The UnrealTargetPlatform being built</param>
+		/// <param name="TargetType">  The type of target (game or program)</param>
+		/// <returns>string    The custom path lines for the project file; Empty string if it doesn't require one</returns>
 		public virtual string GetVisualStudioPathsEntries(UnrealTargetPlatform InPlatform, UnrealTargetConfiguration InConfiguration, TargetRules.TargetType TargetType, FileReference TargetRulesPath, FileReference ProjectFilePath, FileReference NMakeOutputPath)
 		{
 			// NOTE: We are intentionally overriding defaults for these paths with empty strings.  We never want Visual Studio's
@@ -227,48 +212,41 @@ namespace UnrealBuildTool
 			return PathsLines;
 		}
 
-		/**
-		 * Return any custom property import lines
-		 *	
-		 *	@param	InPlatform			The UnrealTargetPlatform being built
-		 *	
-		 *	@return	string				The custom property import lines for the project file; Empty string if it doesn't require one
-		 */
+		/// <summary>
+		/// Return any custom property import lines
+		/// </summary>
+		/// <param name="InPlatform">  The UnrealTargetPlatform being built</param>
+		/// <returns>string    The custom property import lines for the project file; Empty string if it doesn't require one</returns>
 		public virtual string GetAdditionalVisualStudioImportSettings(UnrealTargetPlatform InPlatform)
 		{
 			return "";
 		}
 
-		/**
-		 * Return any custom layout directory sections
-		 * 
-		 *	@param	InPlatform			The UnrealTargetPlatform being built
-		 *	@param	TargetType			The type of target (game or program)
-		 *	
-		 *	@return	string				The custom property import lines for the project file; Empty string if it doesn't require one
-		 */
+		/// <summary>
+		/// Return any custom layout directory sections
+		/// </summary>
+		/// <param name="InPlatform">  The UnrealTargetPlatform being built</param>
+		/// <param name="TargetType">  The type of target (game or program)</param>
+		/// <returns>string    The custom property import lines for the project file; Empty string if it doesn't require one</returns>
 		public virtual string GetVisualStudioLayoutDirSection(UnrealTargetPlatform InPlatform, UnrealTargetConfiguration InConfiguration, string InConditionString, TargetRules.TargetType TargetType, FileReference TargetRulesPath, FileReference ProjectFilePath, FileReference NMakeOutputPath)
 		{
 			return "";
 		}
 
-		/**
-		 *	Get the output manifest section, if required
-		 *	
-		 *	@param	InPlatform			The UnrealTargetPlatform being built
-		 *	
-		 *	@return	string				The output manifest section for the project file; Empty string if it doesn't require one
-		 */
+		/// <summary>
+		/// Get the output manifest section, if required
+		/// </summary>
+		/// <param name="InPlatform">  The UnrealTargetPlatform being built</param>
+		/// <returns>string    The output manifest section for the project file; Empty string if it doesn't require one</returns>
 		public virtual string GetVisualStudioOutputManifestSection(UnrealTargetPlatform InPlatform, TargetRules.TargetType TargetType, FileReference TargetRulesPath, FileReference ProjectFilePath)
 		{
 			return "";
 		}
 
-		/**
-		 * Get whether this platform deploys 
-		 * 
-		 * @return	bool		true if the 'Deploy' option should be enabled
-		 */
+		/// <summary>
+		/// Get whether this platform deploys
+		/// </summary>
+		/// <returns>bool  true if the 'Deploy' option should be enabled</returns>
 		public virtual bool GetVisualStudioDeploymentEnabled(UnrealTargetPlatform InPlatform, UnrealTargetConfiguration InConfiguration)
 		{
 			return false;
diff --git a/Engine/Source/Programs/UnrealBuildTool/System/UProjectInfo.cs b/Engine/Source/Programs/UnrealBuildTool/System/UProjectInfo.cs
--- a/Engine/Source/Programs/UnrealBuildTool/System/UProjectInfo.cs
+++ b/Engine/Source/Programs/UnrealBuildTool/System/UProjectInfo.cs
@@ -29,11 +29,19 @@ namespace UnrealBuildTool
 			bIsCodeProject = bInIsCodeProject;
 		}
 
-		/** Map of relative or complete project file names to the project info */
+		/// <summary>
+		/// Map of relative or complete project file names to the project info
+		/// </summary>
 		static Dictionary<FileReference, UProjectInfo> ProjectInfoDictionary = new Dictionary<FileReference, UProjectInfo>();
-		/** Map of short project file names to the relative or complete project file name */
+
+		/// <summary>
+		/// Map of short project file names to the relative or complete project file name
+		/// </summary>
 		static Dictionary<string, FileReference> ShortProjectNameDictionary = new Dictionary<string, FileReference>(StringComparer.InvariantCultureIgnoreCase);
-		/** Map of target names to the relative or complete project file name */
+
+		/// <summary>
+		/// Map of target names to the relative or complete project file name
+		/// </summary>
 		static Dictionary<string, FileReference> TargetToProjectDictionary = new Dictionary<string, FileReference>(StringComparer.InvariantCultureIgnoreCase);
 
 		public static bool FindTargetFilesInFolder(DirectoryReference InTargetFolder)
diff --git a/Engine/Source/Programs/UnrealBuildTool/System/Unity.cs b/Engine/Source/Programs/UnrealBuildTool/System/Unity.cs
--- a/Engine/Source/Programs/UnrealBuildTool/System/Unity.cs
+++ b/Engine/Source/Programs/UnrealBuildTool/System/Unity.cs
@@ -11,8 +11,10 @@ namespace UnrealBuildTool
 {
 	public class Unity
 	{
-		/** The set of source files that UnrealBuildTool determined to be part of the programmer's "working set".
-		    This feature is only used when bUseAdaptiveUnityBuild is enabled */
+		/// <summary>
+		/// The set of source files that UnrealBuildTool determined to be part of the programmer's "working set".
+		/// This feature is only used when bUseAdaptiveUnityBuild is enabled
+		/// </summary>
 		public static HashSet<FileItem> SourceFileWorkingSet = new HashSet<FileItem>();
 
 
diff --git a/Engine/Source/Programs/UnrealBuildTool/System/UnrealBuildTool.cs b/Engine/Source/Programs/UnrealBuildTool/System/UnrealBuildTool.cs
--- a/Engine/Source/Programs/UnrealBuildTool/System/UnrealBuildTool.cs
+++ b/Engine/Source/Programs/UnrealBuildTool/System/UnrealBuildTool.cs
@@ -16,75 +16,117 @@ namespace UnrealBuildTool
 {
 	public partial class UnrealBuildTool
 	{
-		/** Time at which control fist passes to UBT. */
+		/// <summary>
+		/// Time at which control fist passes to UBT.
+		/// </summary>
 		static public DateTime StartTime;
 
 		static public string BuildGuid = Guid.NewGuid().ToString();
 
-		/** Total time spent building projects. */
+		/// <summary>
+		/// Total time spent building projects.
+		/// </summary>
 		static public double TotalBuildProjectTime = 0;
 
-		/** Total time spent compiling. */
+		/// <summary>
+		/// Total time spent compiling.
+		/// </summary>
 		static public double TotalCompileTime = 0;
 
-		/** Total time spent creating app bundles. */
+		/// <summary>
+		/// Total time spent creating app bundles.
+		/// </summary>
 		static public double TotalCreateAppBundleTime = 0;
 
-		/** Total time spent generating debug information. */
+		/// <summary>
+		/// Total time spent generating debug information.
+		/// </summary>
 		static public double TotalGenerateDebugInfoTime = 0;
 
-		/** Total time spent linking. */
+		/// <summary>
+		/// Total time spent linking.
+		/// </summary>
 		static public double TotalLinkTime = 0;
 
-		/** Total time spent on other actions. */
+		/// <summary>
+		/// Total time spent on other actions.
+		/// </summary>
 		static public double TotalOtherActionsTime = 0;
 
 		/// How much time was spent scanning for include dependencies for outdated C++ files
 		public static double TotalDeepIncludeScanTime = 0.0;
 
-		/** The command line */
+		/// <summary>
+		/// The command line
+		/// </summary>
 		static public List<string> CmdLine = new List<string>();
 
-		/** The environment at boot time. */
+		/// <summary>
+		/// The environment at boot time.
+		/// </summary>
 		static public System.Collections.IDictionary InitialEnvironment;
 
-		/** Should we skip all extra modules for platforms? */
+		/// <summary>
+		/// Should we skip all extra modules for platforms?
+		/// </summary>
 		static bool bSkipNonHostPlatforms = false;
 
 		static UnrealTargetPlatform OnlyPlatformSpecificFor = UnrealTargetPlatform.Unknown;
 
-		/** Are we running for Rocket */
+		/// <summary>
+		/// Are we running for Rocket
+		/// </summary>
 		static public bool bRunningRocket = false;
 
-		/** The project file */
+		/// <summary>
+		/// The project file
+		/// </summary>
 		static FileReference UProjectFile;
 
-		/** The Unreal project file path.  This should always be valid when compiling game projects.  Engine targets and program targets may not have a UProject set. */
+		/// <summary>
+		/// The Unreal project file path.  This should always be valid when compiling game projects.  Engine targets and program targets may not have a UProject set.
+		/// </summary>
 		static bool bHasUProjectFile = false;
 
-		/** The Unreal project directory.  This should always be valid when compiling game projects.  Engine targets and program targets may not have a UProject set. */
+		/// <summary>
+		/// The Unreal project directory.  This should always be valid when compiling game projects.  Engine targets and program targets may not have a UProject set.
+		/// </summary>
 		static DirectoryReference UProjectPath;
 
-		/** The full name of the Root UE4 directory */
+		/// <summary>
+		/// The full name of the Root UE4 directory
+		/// </summary>
 		public static readonly DirectoryReference RootDirectory = new DirectoryReference(Path.Combine(Path.GetDirectoryName(Assembly.GetExecutingAssembly().GetOriginalLocation()), "..", "..", ".."));
 
-		/** The full name of the Engine directory */
+		/// <summary>
+		/// The full name of the Engine directory
+		/// </summary>
 		public static readonly DirectoryReference EngineDirectory = DirectoryReference.Combine(RootDirectory, "Engine");
 
-		/** The full name of the Engine/Source directory */
+		/// <summary>
+		/// The full name of the Engine/Source directory
+		/// </summary>
 		public static readonly DirectoryReference EngineSourceDirectory = DirectoryReference.Combine(EngineDirectory, "Source");
 
-		/** This is set to true during UBT startup, after it is safe to be accessing the IsGatheringBuild and IsAssemblingBuild properties.  Never access this directly. */
+		/// <summary>
+		/// This is set to true during UBT startup, after it is safe to be accessing the IsGatheringBuild and IsAssemblingBuild properties.  Never access this directly.
+		/// </summary>
 		private static bool bIsSafeToCheckIfGatheringOrAssemblingBuild = false;
 
-		/** True if this run of UBT should only invalidate Makefile. */
+		/// <summary>
+		/// True if this run of UBT should only invalidate Makefile.
+		/// </summary>
 		static public bool bIsInvalidatingMakefilesOnly = false;
 
-		/** Cached array of all the supported target platforms */
+		/// <summary>
+		/// Cached array of all the supported target platforms
+		/// </summary>
 		static public UnrealTargetPlatform[] AllPlatforms = (UnrealTargetPlatform[])Enum.GetValues(typeof(UnrealTargetPlatform));
 
-		/** True if we should gather module dependencies for building in this run.  If this is false, then we'll expect to be loading information from disk about
-			the target's modules before we'll be able to build anything.  One or both of IsGatheringBuild or IsAssemblingBuild must be true. */
+		/// <summary>
+		/// True if we should gather module dependencies for building in this run.  If this is false, then we'll expect to be loading information from disk about
+		/// the target's modules before we'll be able to build anything.  One or both of IsGatheringBuild or IsAssemblingBuild must be true.
+		/// </summary>
 		static public bool IsGatheringBuild
 		{
 			get
@@ -101,8 +143,10 @@ namespace UnrealBuildTool
 		}
 		static private bool bIsGatheringBuild_Unsafe = true;
 
-		/** True if we should go ahead and actually compile and link in this run.  If this is false, no actions will be executed and we'll stop after gathering
-			and saving information about dependencies.  One or both of IsGatheringBuild or IsAssemblingBuild must be true. */
+		/// <summary>
+		/// True if we should go ahead and actually compile and link in this run.  If this is false, no actions will be executed and we'll stop after gathering
+		/// and saving information about dependencies.  One or both of IsGatheringBuild or IsAssemblingBuild must be true.
+		/// </summary>
 		static public bool IsAssemblingBuild
 		{
 			get
@@ -119,9 +163,11 @@ namespace UnrealBuildTool
 		}
 		static private bool bIsAssemblingBuild_Unsafe = true;
 
-		/** Used when BuildConfiguration.bUseUBTMakefiles is enabled.  If true, it means that our cached includes may not longer be
-			valid (or never were), and we need to recover by forcibly scanning included headers for all build prerequisites to make sure that our
-			cached set of includes is actually correct, before determining which files are outdated. */
+		/// <summary>
+		/// Used when BuildConfiguration.bUseUBTMakefiles is enabled.  If true, it means that our cached includes may not longer be
+		/// valid (or never were), and we need to recover by forcibly scanning included headers for all build prerequisites to make sure that our
+		/// cached set of includes is actually correct, before determining which files are outdated.
+		/// </summary>
 		static public bool bNeedsFullCPPIncludeRescan = false;
 
 
@@ -2206,12 +2252,11 @@ namespace UnrealBuildTool
 			return bIsRunning;
 		}
 
-		/**
-		 * Parses the passed in command line for build configuration overrides.
-		 * 
-		 * @param	Arguments	List of arguments to parse
-		 * @return	List of build target settings
-		 */
+		/// <summary>
+		/// Parses the passed in command line for build configuration overrides.
+		/// </summary>
+		/// <param name="Arguments">List of arguments to parse</param>
+		/// <returns>List of build target settings</returns>
 		private static List<string[]> ParseCommandLineFlags(string[] Arguments)
 		{
 			var TargetSettings = new List<string[]>();
@@ -2283,24 +2328,36 @@ namespace UnrealBuildTool
 
 			public int Version;
 
-			/** Every action in the action graph */
+			/// <summary>
+			/// Every action in the action graph
+			/// </summary>
 			public List<Action> AllActions;
 
-			/** List of the actions that need to be run in order to build the targets' final output items */
+			/// <summary>
+			/// List of the actions that need to be run in order to build the targets' final output items
+			/// </summary>
 			public Action[] PrerequisiteActions;
 
-			/** Environment variables that we'll need in order to invoke the platform's compiler and linker */
+			/// <summary>
+			/// Environment variables that we'll need in order to invoke the platform's compiler and linker
+			/// </summary>
 			// @todo ubtmake: Really we want to allow a different set of environment variables for every Action.  This would allow for targets on multiple platforms to be built in a single assembling phase.  We'd only have unique variables for each platform that has actions, so we'd want to make sure we only store the minimum set.
 			public readonly List<Tuple<string, string>> EnvironmentVariables = new List<Tuple<string, string>>();
 
-			/** Maps each target to a list of UObject module info structures */
+			/// <summary>
+			/// Maps each target to a list of UObject module info structures
+			/// </summary>
 			public Dictionary<string, List<UHTModuleInfo>> TargetNameToUObjectModules;
 
-			/** List of targets being built */
+			/// <summary>
+			/// List of targets being built
+			/// </summary>
 			public List<UEBuildTarget> Targets;
 
 
-			/** Current working set of source files, for when bUseAdaptiveUnityBuild is enabled */
+			/// <summary>
+			/// Current working set of source files, for when bUseAdaptiveUnityBuild is enabled
+			/// </summary>
 			public HashSet<FileItem> SourceFileWorkingSet = new HashSet<FileItem>();
 
 			public UBTMakefile()
@@ -2337,7 +2394,7 @@ namespace UnrealBuildTool
 			}
 
 
-			/** @return	 True if this makefile's contents look valid.  Called after loading the file to make sure it is legit. */
+			/// <returns> True if this makefile's contents look valid.  Called after loading the file to make sure it is legit.</returns>
 			public bool IsValidMakefile()
 			{
 				return
diff --git a/Engine/Source/Programs/UnrealBuildTool/System/Utils.cs b/Engine/Source/Programs/UnrealBuildTool/System/Utils.cs
--- a/Engine/Source/Programs/UnrealBuildTool/System/Utils.cs
+++ b/Engine/Source/Programs/UnrealBuildTool/System/Utils.cs
@@ -110,7 +110,9 @@ namespace UnrealBuildTool
 		/// </summary>
 		public static readonly bool IsRunningOnMono = Type.GetType("Mono.Runtime") != null;
 
-		/** Searches for a flag in a set of command-line arguments. */
+		/// <summary>
+		/// Searches for a flag in a set of command-line arguments.
+		/// </summary>
 		public static bool ParseCommandLineFlag(string[] Arguments, string FlagName, out int ArgumentIndex)
 		{
 			// Find an argument with the given name.
@@ -125,15 +127,16 @@ namespace UnrealBuildTool
 			return false;
 		}
 
-		/** Regular expression to match $(ENV) and/ or %ENV% environment variables. */
+		/// <summary>
+		/// Regular expression to match $(ENV) and/ or %ENV% environment variables.
+		/// </summary>
 		static Regex EnvironmentVariableRegex = new Regex(@"\$\((.*?)\)|\%(.*?)\%", RegexOptions.None);
 
-		/**
-		 * Resolves $(ENV) and/ or %ENV% to the value of the environment variable in the passed in string.
-		 * 
-		 * @param	InString	String to resolve environment variable in.
-		 * @return	String with environment variable expanded/ resolved.
-		 */
+		/// <summary>
+		/// Resolves $(ENV) and/ or %ENV% to the value of the environment variable in the passed in string.
+		/// </summary>
+		/// <param name="InString">String to resolve environment variable in.</param>
+		/// <returns>String with environment variable expanded/ resolved.</returns>
 		public static string ResolveEnvironmentVariable(string InString)
 		{
 			string Result = InString;
@@ -206,15 +209,13 @@ namespace UnrealBuildTool
 			return Result;
 		}
 
-		/**
-		 * This is a faster replacement of File.ReadAllText. Code snippet based on code 
-		 * and analysis by Sam Allen
-		 * 
-		 * http://dotnetperls.com/Content/File-Handling.aspx
-		 * 
-		 * @param	SourceFile		Source file to fully read and convert to string
-		 * @return	Textual representation of file.
-		 */
+		/// <summary>
+		/// This is a faster replacement of File.ReadAllText. Code snippet based on code
+		/// and analysis by Sam Allen
+		/// http://dotnetperls.com/Content/File-Handling.aspx
+		/// </summary>
+		/// <param name="SourceFile"> Source file to fully read and convert to string</param>
+		/// <returns>Textual representation of file.</returns>
 		public static string ReadAllText(string SourceFile)
 		{
 			using (StreamReader Reader = new StreamReader(SourceFile, System.Text.Encoding.UTF8))
@@ -223,13 +224,12 @@ namespace UnrealBuildTool
 			}
 		}
 
-		/**
-		 * Reads the specified environment variable
-		 *
-		 * @param	VarName		the environment variable to read
-		 * @param	bDefault	the default value to use if missing
-		 * @return	the value of the environment variable if found and the default value if missing
-		 */
+		/// <summary>
+		/// Reads the specified environment variable
+		/// </summary>
+		/// <param name="VarName"> the environment variable to read</param>
+		/// <param name="bDefault">the default value to use if missing</param>
+		/// <returns>the value of the environment variable if found and the default value if missing</returns>
 		public static bool GetEnvironmentVariable(string VarName, bool bDefault)
 		{
 			string Value = Environment.GetEnvironmentVariable(VarName);
@@ -241,13 +241,12 @@ namespace UnrealBuildTool
 			return bDefault;
 		}
 
-		/**
-		 * Reads the specified environment variable
-		 *
-		 * @param	VarName		the environment variable to read
-		 * @param	Default	the default value to use if missing
-		 * @return	the value of the environment variable if found and the default value if missing
-		 */
+		/// <summary>
+		/// Reads the specified environment variable
+		/// </summary>
+		/// <param name="VarName"> the environment variable to read</param>
+		/// <param name="Default">the default value to use if missing</param>
+		/// <returns>the value of the environment variable if found and the default value if missing</returns>
 		public static string GetStringEnvironmentVariable(string VarName, string Default)
 		{
 			string Value = Environment.GetEnvironmentVariable(VarName);
@@ -258,13 +257,12 @@ namespace UnrealBuildTool
 			return Default;
 		}
 
-		/**
-		 * Reads the specified environment variable
-		 *
-		 * @param	VarName		the environment variable to read
-		 * @param	Default	the default value to use if missing
-		 * @return	the value of the environment variable if found and the default value if missing
-		 */
+		/// <summary>
+		/// Reads the specified environment variable
+		/// </summary>
+		/// <param name="VarName"> the environment variable to read</param>
+		/// <param name="Default">the default value to use if missing</param>
+		/// <returns>the value of the environment variable if found and the default value if missing</returns>
 		public static double GetEnvironmentVariable(string VarName, double Default)
 		{
 			string Value = Environment.GetEnvironmentVariable(VarName);
@@ -275,13 +273,12 @@ namespace UnrealBuildTool
 			return Default;
 		}
 
-		/**
-		 * Reads the specified environment variable
-		 *
-		 * @param	VarName		the environment variable to read
-		 * @param	Default	the default value to use if missing
-		 * @return	the value of the environment variable if found and the default value if missing
-		 */
+		/// <summary>
+		/// Reads the specified environment variable
+		/// </summary>
+		/// <param name="VarName"> the environment variable to read</param>
+		/// <param name="Default">the default value to use if missing</param>
+		/// <returns>the value of the environment variable if found and the default value if missing</returns>
 		public static string GetEnvironmentVariable(string VarName, string Default)
 		{
 			string Value = Environment.GetEnvironmentVariable(VarName);
@@ -292,12 +289,11 @@ namespace UnrealBuildTool
 			return Default;
 		}
 
-		/**
-		 * Sets the environment variables from the passed in batch file
-		 * 
-		 * @param	BatchFileName	Name of the batch file to parse
-		 * @param	Parameters		Optional command-line parameters to pass to the batch file when running it
-		 */
+		/// <summary>
+		/// Sets the environment variables from the passed in batch file
+		/// </summary>
+		/// <param name="BatchFileName">Name of the batch file to parse</param>
+		/// <param name="Parameters"> Optional command-line parameters to pass to the batch file when running it</param>
 		public static void SetEnvironmentVariablesFromBatchFile(string BatchFileName, string Parameters = "")
 		{
 			// @todo ubtmake: Experiment with changing this to run asynchronously at startup, and only blocking if accessed before the .bat file finishes
@@ -318,9 +314,10 @@ namespace UnrealBuildTool
 		}
 
 
-		/**
-		/* Try to launch a local process, and produce a friendly error message if it fails.
-		/*/
+		/// <summary>
+		/// Try to launch a local process, and produce a friendly error message if it fails.
+		/// /
+		/// </summary>
 		public static int RunLocalProcess(Process LocalProcess)
 		{
 			int ExitCode = -1;
@@ -350,9 +347,9 @@ namespace UnrealBuildTool
 			return ExitCode;
 		}
 
-		/**
-		 * Runs a local process and pipes the output to the log
-		 */
+		/// <summary>
+		/// Runs a local process and pipes the output to the log
+		/// </summary>
 		public static int RunLocalProcessAndLogOutput(ProcessStartInfo StartInfo)
 		{
 			Process LocalProcess = new Process();
@@ -907,13 +904,17 @@ namespace UnrealBuildTool
 			}
 		}
 
-		/** Helper function to get the minimum of two comparable objects */
+		/// <summary>
+		/// Helper function to get the minimum of two comparable objects
+		/// </summary>
 		public static T Min<T>(T a, T b) where T : IComparable
 		{
 			return a.CompareTo(b) < 0 ? a : b;
 		}
 
-		/** Helper function to get the maximum of two comparable objects */
+		/// <summary>
+		/// Helper function to get the maximum of two comparable objects
+		/// </summary>
 		public static T Max<T>(T a, T b) where T : IComparable
 		{
 			return a.CompareTo(b) > 0 ? a : b;
diff --git a/Engine/Source/Programs/UnrealBuildTool/System/VCProject.cs b/Engine/Source/Programs/UnrealBuildTool/System/VCProject.cs
--- a/Engine/Source/Programs/UnrealBuildTool/System/VCProject.cs
+++ b/Engine/Source/Programs/UnrealBuildTool/System/VCProject.cs
@@ -1144,7 +1144,9 @@ namespace UnrealBuildTool
 	}
 
 
-	/** A Visual C# project. */
+	/// <summary>
+	/// A Visual C# project.
+	/// </summary>
 	public class VCSharpProjectFile : MSBuildProjectFile
 	{
 		// This is the GUID that Visual Studio uses to identify a C# project file in the solution
@@ -1163,7 +1165,9 @@ namespace UnrealBuildTool
 		}
 
 
-		/** Reads the list of dependencies from the specified project file. */
+		/// <summary>
+		/// Reads the list of dependencies from the specified project file.
+		/// </summary>
 		public List<string> GetCSharpDependencies()
 		{
 			var RelativeFilePaths = new List<string>();
@@ -1183,11 +1187,10 @@ namespace UnrealBuildTool
 			return RelativeFilePaths;
 		}
 
-		/**
-		 * Adds a C# dot net (system) assembly reference to this project
-		 *
-		 * @param	AssemblyReference	The full path to the assembly file on disk
-		 */
+		/// <summary>
+		/// Adds a C# dot net (system) assembly reference to this project
+		/// </summary>
+		/// <param name="AssemblyReference">The full path to the assembly file on disk</param>
 		public void AddDotNetAssemblyReference(string AssemblyReference)
 		{
 			if (!DotNetAssemblyReferences.Contains(AssemblyReference))
@@ -1196,11 +1199,10 @@ namespace UnrealBuildTool
 			}
 		}
 
-		/**
-		 * Adds a C# assembly reference to this project, such as a third party assembly needed for this project to compile
-		 *
-		 * @param	AssemblyReference	The full path to the assembly file on disk
-		 */
+		/// <summary>
+		/// Adds a C# assembly reference to this project, such as a third party assembly needed for this project to compile
+		/// </summary>
+		/// <param name="AssemblyReference">The full path to the assembly file on disk</param>
 		public void AddAssemblyReference(FileReference AssemblyReference)
 		{
 			AssemblyReferences.Add(AssemblyReference);
diff --git a/Engine/Source/Programs/UnrealBuildTool/System/XGE.cs b/Engine/Source/Programs/UnrealBuildTool/System/XGE.cs
--- a/Engine/Source/Programs/UnrealBuildTool/System/XGE.cs
+++ b/Engine/Source/Programs/UnrealBuildTool/System/XGE.cs
@@ -162,7 +162,9 @@ namespace UnrealBuildTool
 			return EventArgs;
 		}
 
-		/** Writes a XGE task file containing the specified actions to the specified file path. */
+		/// <summary>
+		/// Writes a XGE task file containing the specified actions to the specified file path.
+		/// </summary>
 		public static void WriteTaskFile(List<Action> InActions, string TaskFilePath, bool bProgressMarkup)
 		{
 			Dictionary<string, string> ExportEnv = new Dictionary<string, string>();
@@ -407,7 +409,9 @@ namespace UnrealBuildTool
 			}
 		}
 
-		/** The possible result of executing tasks with XGE. */
+		/// <summary>
+		/// The possible result of executing tasks with XGE.
+		/// </summary>
 		public enum ExecutionResult
 		{
 			Unavailable,
@@ -415,11 +419,11 @@ namespace UnrealBuildTool
 			TasksSucceeded,
 		}
 
-		/**
-		 * Executes the tasks in the specified file.
-		 * @param TaskFilePath - The path to the file containing the tasks to execute in XGE XML format.
-		 * @return Indicates whether the tasks were successfully executed.
-		 */
+		/// <summary>
+		/// Executes the tasks in the specified file.
+		/// </summary>
+		/// <param name="TaskFilePath">- The path to the file containing the tasks to execute in XGE XML format.</param>
+		/// <returns>Indicates whether the tasks were successfully executed.</returns>
 		public static ExecutionResult ExecuteTaskFile(string TaskFilePath, DataReceivedEventHandler OutputEventHandler)
 		{
 			bool bSilentCompileOutput = false;
@@ -477,9 +481,9 @@ namespace UnrealBuildTool
 			}
 		}
 
-		/**
-		 * Executes the tasks in the specified file, parsing progress markup as part of the output.
-		 */
+		/// <summary>
+		/// Executes the tasks in the specified file, parsing progress markup as part of the output.
+		/// </summary>
 		public static ExecutionResult ExecuteTaskFileWithProgressMarkup(string TaskFilePath, int NumActions, DataReceivedEventHandler OutputEventHandler)
 		{
 			using (ProgressWriter Writer = new ProgressWriter("Compiling C++ source files...", false))
diff --git a/Engine/Source/Programs/UnrealBuildTool/System/XcodeProject.cs b/Engine/Source/Programs/UnrealBuildTool/System/XcodeProject.cs
--- a/Engine/Source/Programs/UnrealBuildTool/System/XcodeProject.cs
+++ b/Engine/Source/Programs/UnrealBuildTool/System/XcodeProject.cs
@@ -131,7 +131,9 @@ namespace UnrealBuildTool
 		{
 		}
 
-		/** Gets Xcode file category based on its extension */
+		/// <summary>
+		/// Gets Xcode file category based on its extension
+		/// </summary>
 		private string GetFileCategory(string Extension)
 		{
 			// @todo Mac: Handle more categories
@@ -144,7 +146,9 @@ namespace UnrealBuildTool
 			}
 		}
 
-		/** Gets Xcode file type based on its extension */
+		/// <summary>
+		/// Gets Xcode file type based on its extension
+		/// </summary>
 		private string GetFileType(string Extension)
 		{
 			// @todo Mac: Handle more file types
@@ -168,7 +172,9 @@ namespace UnrealBuildTool
 			}
 		}
 
-		/** Returns true if Extension is a known extension for files containing source code */
+		/// <summary>
+		/// Returns true if Extension is a known extension for files containing source code
+		/// </summary>
 		private bool IsSourceCode(string Extension)
 		{
 			return Extension == ".c" || Extension == ".cc" || Extension == ".cpp" || Extension == ".m" || Extension == ".mm";
@@ -205,10 +211,10 @@ namespace UnrealBuildTool
 			return false;
 		}
 
-		/**
-		 * Returns a project navigator group to which the file should belong based on its path.
-		 * Creates a group tree if it doesn't exist yet.
-		 */
+		/// <summary>
+		/// Returns a project navigator group to which the file should belong based on its path.
+		/// Creates a group tree if it doesn't exist yet.
+		/// </summary>
 		public XcodeFileGroup FindGroupByFullPath(ref Dictionary<string, XcodeFileGroup> Groups, DirectoryReference FullPath)
 		{
 			string RelativePath = (FullPath == XcodeProjectFileGenerator.MasterProjectPath) ? "" : FullPath.MakeRelativeTo(XcodeProjectFileGenerator.MasterProjectPath);
@@ -269,9 +275,9 @@ namespace UnrealBuildTool
 			return new XcodeSourceFile(InitFilePath, InitProjectSubFolder);
 		}
 
-		/**
-		 * Generates bodies of all sections that contain a list of source files plus a dictionary of project navigator groups.
-		 */
+		/// <summary>
+		/// Generates bodies of all sections that contain a list of source files plus a dictionary of project navigator groups.
+		/// </summary>
 		public void GenerateSectionsContents(ref string PBXBuildFileSection, ref string PBXFileReferenceSection,
 											 ref string PBXSourcesBuildPhaseSection, ref Dictionary<string, XcodeFileGroup> Groups)
 		{
diff --git a/Engine/Source/Programs/UnrealBuildTool/System/XcodeProjectFileGenerator.cs b/Engine/Source/Programs/UnrealBuildTool/System/XcodeProjectFileGenerator.cs
--- a/Engine/Source/Programs/UnrealBuildTool/System/XcodeProjectFileGenerator.cs
+++ b/Engine/Source/Programs/UnrealBuildTool/System/XcodeProjectFileGenerator.cs
@@ -197,9 +197,9 @@ namespace UnrealBuildTool
 		{
 		}
 
-		/**
-		 * Make a random Guid string usable by Xcode (24 characters exactly)
-		 */
+		/// <summary>
+		/// Make a random Guid string usable by Xcode (24 characters exactly)
+		/// </summary>
 		public static string MakeXcodeGuid()
 		{
 			string Guid = "";
diff --git a/Engine/Source/Programs/UnrealBuildTool/ToolChain/RemoteToolChain.cs b/Engine/Source/Programs/UnrealBuildTool/ToolChain/RemoteToolChain.cs
--- a/Engine/Source/Programs/UnrealBuildTool/ToolChain/RemoteToolChain.cs
+++ b/Engine/Source/Programs/UnrealBuildTool/ToolChain/RemoteToolChain.cs
@@ -15,9 +15,9 @@ namespace UnrealBuildTool
 {
 	public abstract class RemoteToolChain : UEToolChain
 	{
-		/**
-		 * Common error codes reported by Remote Tool Chain and its actions.
-		 */
+		/// <summary>
+		/// Common error codes reported by Remote Tool Chain and its actions.
+		/// </summary>
 		public enum RemoteToolChainErrorCode
 		{
 			NoError = 0,
@@ -38,54 +38,78 @@ namespace UnrealBuildTool
 			UEToolChain.RegisterPlatformToolChain(CPPPlatform, this);
 		}
 
-		/** These two variables will be loaded from XML config file in XmlConfigLoader.Init() */
+		/// <summary>
+		/// These two variables will be loaded from XML config file in XmlConfigLoader.Init()
+		/// </summary>
 		[XmlConfig]
 		public static string RemoteServerName = "";
 		[XmlConfig]
 		public static string[] PotentialServerNames = new string[] { };
 
-		/** Keep a list of remote files that are potentially copied from local to remote */
+		/// <summary>
+		/// Keep a list of remote files that are potentially copied from local to remote
+		/// </summary>
 		private static Dictionary<FileItem, FileItem> CachedRemoteFileItems = new Dictionary<FileItem, FileItem>();
 
-		/** The base path (on the Mac) to the your particular development directory, where files will be copied to from the PC */
+		/// <summary>
+		/// The base path (on the Mac) to the your particular development directory, where files will be copied to from the PC
+		/// </summary>
 		public static string UserDevRootMacBase = "/UE4/Builds/";
 
-		/** The final path (on the Mac) to your particular development directory, where files will be copied to from the PC */
+		/// <summary>
+		/// The final path (on the Mac) to your particular development directory, where files will be copied to from the PC
+		/// </summary>
 		public static string UserDevRootMac = "/UE4/Builds";
 
-		/** Whether or not to connect to UnrealRemoteTool using RPCUtility */
+		/// <summary>
+		/// Whether or not to connect to UnrealRemoteTool using RPCUtility
+		/// </summary>
 		[XmlConfig]
 		public static bool bUseRPCUtil = true;
 
-		/** The user has specified a deltacopy install path */
+		/// <summary>
+		/// The user has specified a deltacopy install path
+		/// </summary>
 		private static string OverrideDeltaCopyInstallPath = null;
 
-		/** Path to rsync executable and parameters for your rsync utility */
+		/// <summary>
+		/// Path to rsync executable and parameters for your rsync utility
+		/// </summary>
 		[XmlConfig]
 		public static string RSyncExe = "${PROGRAM_FILES}\\DeltaCopy\\rsync.exe";
 		public static string ResolvedRSyncExe = null;
 
-		/** Path to rsync executable and parameters for your rsync utility */
+		/// <summary>
+		/// Path to rsync executable and parameters for your rsync utility
+		/// </summary>
 		[XmlConfig]
 		public static string SSHExe = "${PROGRAM_FILES}\\DeltaCopy\\ssh.exe";
 		public static string ResolvedSSHExe = null;
 
-		/** Instead of looking for RemoteToolChainPrivate.key in the usual places (Documents/Unreal Engine/UnrealBuildTool/SSHKeys, Engine/Build/SSHKeys), this private key will be used if set */
+		/// <summary>
+		/// Instead of looking for RemoteToolChainPrivate.key in the usual places (Documents/Unreal Engine/UnrealBuildTool/SSHKeys, Engine/Build/SSHKeys), this private key will be used if set
+		/// </summary>
 		[XmlConfig]
 		public static string SSHPrivateKeyOverridePath = "";
 		public static string ResolvedSSHPrivateKey = null;
 
-		/** The authentication used for Rsync (for the -e rsync flag) */
+		/// <summary>
+		/// The authentication used for Rsync (for the -e rsync flag)
+		/// </summary>
 		[XmlConfig]
 		public static string RsyncAuthentication = "ssh -i '${CYGWIN_SSH_PRIVATE_KEY}'";
 		public static string ResolvedRsyncAuthentication = null;
 
-		/** The authentication used for SSH (probably similar to RsyncAuthentication) */
+		/// <summary>
+		/// The authentication used for SSH (probably similar to RsyncAuthentication)
+		/// </summary>
 		[XmlConfig]
 		public static string SSHAuthentication = "-i '${CYGWIN_SSH_PRIVATE_KEY}'";
 		public static string ResolvedSSHAuthentication = null;
 
-		/** Username on the remote machine to connect to with RSync */
+		/// <summary>
+		/// Username on the remote machine to connect to with RSync
+		/// </summary>
 		[XmlConfig]
 		public static string RSyncUsername = "${CURRENT_USER}";
 		public static string ResolvedRSyncUsername = null;
@@ -93,17 +117,25 @@ namespace UnrealBuildTool
 		// has the toolchain initialized remote execution yet? no need to do it multiple times
 		private static bool bHasBeenInitialized = false;
 
-		/** The directory that this local branch is in, without drive information (strip off X:\ from X:\UE4\iOS) */
+		/// <summary>
+		/// The directory that this local branch is in, without drive information (strip off X:\ from X:\UE4\iOS)
+		/// </summary>
 		public static string BranchDirectory = Path.GetFullPath(".\\");
 
 
-		/** Substrings that indicate a line contains an error */
+		/// <summary>
+		/// Substrings that indicate a line contains an error
+		/// </summary>
 		protected static List<string> ErrorMessageTokens;
 
-		/** The platform this toolchain is compiling for */
+		/// <summary>
+		/// The platform this toolchain is compiling for
+		/// </summary>
 		protected UnrealTargetPlatform RemoteToolChainPlatform;
 
-		/** The average amound of memory a compile takes, used so that we don't compile too many things at once */
+		/// <summary>
+		/// The average amound of memory a compile takes, used so that we don't compile too many things at once
+		/// </summary>
 		public static int MemoryPerCompileMB = 1000;
 
 		static RemoteToolChain()
@@ -473,7 +505,9 @@ namespace UnrealBuildTool
 			}
 		}
 
-		/** Converts the passed in path from UBT host to compiler native format. */
+		/// <summary>
+		/// Converts the passed in path from UBT host to compiler native format.
+		/// </summary>
 		public override string ConvertPath(string OriginalPath)
 		{
 			if (BuildHostPlatform.Current.Platform != UnrealTargetPlatform.Mac)
@@ -560,9 +594,9 @@ namespace UnrealBuildTool
 			return RemoteFileItem;
 		}
 
-		/**
-		 * Helper function to sync source files to and from the local system and a remote Mac
-		 */
+		/// <summary>
+		/// Helper function to sync source files to and from the local system and a remote Mac
+		/// </summary>
 		//This chunk looks to be required to pipe output to VS giving information on the status of a remote build.
 		public static bool OutputReceivedDataEventHandlerEncounteredError = false;
 		public static string OutputReceivedDataEventHandlerEncounteredErrorMessage = "";
@@ -950,12 +984,11 @@ namespace UnrealBuildTool
 			return RemoteAvailableCommandSlotCount;
 		}
 
-		/**
-		 * Translates clang output warning/error messages into vs-clickable messages
-		 * 
-		 * @param	sender		Sending object
-		 * @param	e			Event arguments (In this case, the line of string output)
-		 */
+		/// <summary>
+		/// Translates clang output warning/error messages into vs-clickable messages
+		/// </summary>
+		/// <param name="sender"> Sending object</param>
+		/// <param name="e">  Event arguments (In this case, the line of string output)</param>
 		protected void RemoteOutputReceivedEventHandler(object sender, DataReceivedEventArgs e)
 		{
 			var Output = e.Data;
@@ -1005,10 +1038,10 @@ namespace UnrealBuildTool
 			}
 		}
 
-		/**
-		 * Queries the remote compile server for CPU information
-		 * and computes the proper ProcessorCountMultiplier.
-		 */
+		/// <summary>
+		/// Queries the remote compile server for CPU information
+		/// and computes the proper ProcessorCountMultiplier.
+		/// </summary>
 		static private Int32 QueryResult = 0;
 		static public void OutputReceivedForQuery(Object Sender, DataReceivedEventArgs Line)
 		{
diff --git a/Engine/Source/Programs/UnrealBuildTool/ToolChain/UEToolChain.cs b/Engine/Source/Programs/UnrealBuildTool/ToolChain/UEToolChain.cs
--- a/Engine/Source/Programs/UnrealBuildTool/ToolChain/UEToolChain.cs
+++ b/Engine/Source/Programs/UnrealBuildTool/ToolChain/UEToolChain.cs
@@ -21,7 +21,9 @@ namespace UnrealBuildTool
 
 		void CompileCSharpProject(CSharpEnvironment CompileEnvironment, FileReference ProjectFileName, FileReference DestinationFile);
 
-		/** Converts the passed in path from UBT host to compiler native format. */
+		/// <summary>
+		/// Converts the passed in path from UBT host to compiler native format.
+		/// </summary>
 		String ConvertPath(String OriginalPath);
 
 		/// <summary>
@@ -116,7 +118,9 @@ namespace UnrealBuildTool
 			return FileReference.Combine(LinkEnvironment.Config.IntermediateDirectory, OutputFile.Reference.GetFileName() + ".response");
 		}
 
-		/** Converts the passed in path from UBT host to compiler native format. */
+		/// <summary>
+		/// Converts the passed in path from UBT host to compiler native format.
+		/// </summary>
 		public virtual String ConvertPath(String OriginalPath)
 		{
 			return OriginalPath;
diff --git a/Engine/Source/Programs/UnrealBuildTool/UWP/UEBuildUWP.cs b/Engine/Source/Programs/UnrealBuildTool/UWP/UEBuildUWP.cs
--- a/Engine/Source/Programs/UnrealBuildTool/UWP/UEBuildUWP.cs
+++ b/Engine/Source/Programs/UnrealBuildTool/UWP/UEBuildUWP.cs
@@ -55,7 +55,9 @@ namespace UnrealBuildTool
 		/// True if we should only build against the app-local CRT and set the API-set flags
 		public static readonly bool bWinApiFamilyApp = true;
 
-		/** True if VS EnvDTE is available (false when building using Visual Studio Express) */
+		/// <summary>
+		/// True if VS EnvDTE is available (false when building using Visual Studio Express)
+		/// </summary>
 		public static bool bHasVisualStudioDTE
 		{
 			get
@@ -83,23 +85,20 @@ namespace UnrealBuildTool
 			}
 		}
 
-		/**
-		 * Returns VisualStudio common tools path for current compiler.
-		 * 
-		 * @return Common tools path.
-		 */
+		/// <summary>
+		/// Returns VisualStudio common tools path for current compiler.
+		/// </summary>
+		/// <returns>Common tools path.</returns>
 		public static string GetVSComnToolsPath()
 		{
 			return GetVSComnToolsPath(Compiler);
 		}
 
-		/**
-		 * Returns VisualStudio common tools path for given compiler.
-		 * 
-		 * @param Compiler Compiler for which to return tools path.
-		 * 
-		 * @return Common tools path.
-		 */
+		/// <summary>
+		/// Returns VisualStudio common tools path for given compiler.
+		/// </summary>
+		/// <param name="Compiler">Compiler for which to return tools path.</param>
+		/// <returns>Common tools path.</returns>
 		public static string GetVSComnToolsPath(WindowsCompiler Compiler)
 		{
 			int VSVersion;
@@ -141,9 +140,9 @@ namespace UnrealBuildTool
 		}
 
 
-		/**
-		 *	Register the platform with the UEBuildPlatform class
-		 */
+		/// <summary>
+		/// Register the platform with the UEBuildPlatform class
+		/// </summary>
 		protected override void RegisterBuildPlatformInternal()
 		{
 			// Register this build platform for UWP
@@ -152,13 +151,11 @@ namespace UnrealBuildTool
 			UEBuildPlatform.RegisterPlatformWithGroup(UnrealTargetPlatform.UWP, UnrealPlatformGroup.Microsoft);
 		}
 
-		/**
-		 *	Retrieve the CPPTargetPlatform for the given UnrealTargetPlatform
-		 *
-		 *	@param	InUnrealTargetPlatform		The UnrealTargetPlatform being build
-		 *	
-		 *	@return	CPPTargetPlatform			The CPPTargetPlatform to compile for
-		 */
+		/// <summary>
+		/// Retrieve the CPPTargetPlatform for the given UnrealTargetPlatform
+		/// </summary>
+		/// <param name="InUnrealTargetPlatform"> The UnrealTargetPlatform being build</param>
+		/// <returns>CPPTargetPlatform   The CPPTargetPlatform to compile for</returns>
 		public override CPPTargetPlatform GetCPPTargetPlatform(UnrealTargetPlatform InUnrealTargetPlatform)
 		{
 			switch (InUnrealTargetPlatform)
@@ -169,13 +166,11 @@ namespace UnrealBuildTool
 			throw new BuildException("UWPPlatform::GetCPPTargetPlatform: Invalid request for {0}", InUnrealTargetPlatform.ToString());
 		}
 
-		/**
-		 *	Get the extension to use for the given binary type
-		 *	
-		 *	@param	InBinaryType		The binrary type being built
-		 *	
-		 *	@return	string				The binary extenstion (ie 'exe' or 'dll')
-		 */
+		/// <summary>
+		/// Get the extension to use for the given binary type
+		/// </summary>
+		/// <param name="InBinaryType"> The binrary type being built</param>
+		/// <returns>string    The binary extenstion (ie 'exe' or 'dll')</returns>
 		public override string GetBinaryExtension(UEBuildBinaryType InBinaryType)
 		{
 			switch (InBinaryType)
@@ -211,13 +206,11 @@ namespace UnrealBuildTool
 		}
 
 
-		/**
-		 *	Get the extension to use for debug info for the given binary type
-		 *	
-		 *	@param	InBinaryType		The binary type being built
-		 *	
-		 *	@return	string				The debug info extension (i.e. 'pdb')
-		 */
+		/// <summary>
+		/// Get the extension to use for debug info for the given binary type
+		/// </summary>
+		/// <param name="InBinaryType"> The binary type being built</param>
+		/// <returns>string    The debug info extension (i.e. 'pdb')</returns>
 		public override string GetDebugInfoExtension(UEBuildBinaryType InBinaryType)
 		{
 			switch (InBinaryType)
@@ -229,28 +222,24 @@ namespace UnrealBuildTool
 			return "";
 		}
 
-		/**
-		 *	Whether incremental linking should be used
-		 *	
-		 *	@param	InPlatform			The CPPTargetPlatform being built
-		 *	@param	InConfiguration		The CPPTargetConfiguration being built
-		 *	
-		 *	@return	bool	true if incremental linking should be used, false if not
-		 */
+		/// <summary>
+		/// Whether incremental linking should be used
+		/// </summary>
+		/// <param name="InPlatform">  The CPPTargetPlatform being built</param>
+		/// <param name="InConfiguration"> The CPPTargetConfiguration being built</param>
+		/// <returns>bool true if incremental linking should be used, false if not</returns>
 		public override bool ShouldUseIncrementalLinking(CPPTargetPlatform Platform, CPPTargetConfiguration Configuration)
 		{
 			return (Configuration == CPPTargetConfiguration.Debug);
 		}
 
-		/**
-		 *	Whether PDB files should be used
-		 *	
-		 *	@param	InPlatform			The CPPTargetPlatform being built
-		 *	@param	InConfiguration		The CPPTargetConfiguration being built
-		 *	@param	bInCreateDebugInfo	true if debug info is getting create, false if not
-		 *	
-		 *	@return	bool	true if PDB files should be used, false if not
-		 */
+		/// <summary>
+		/// Whether PDB files should be used
+		/// </summary>
+		/// <param name="InPlatform">  The CPPTargetPlatform being built</param>
+		/// <param name="InConfiguration"> The CPPTargetConfiguration being built</param>
+		/// <param name="bInCreateDebugInfo">true if debug info is getting create, false if not</param>
+		/// <returns>bool true if PDB files should be used, false if not</returns>
 		public override bool ShouldUsePDBFiles(CPPTargetPlatform Platform, CPPTargetConfiguration Configuration, bool bCreateDebugInfo)
 		{
 			// Only supported on PC.
@@ -261,13 +250,12 @@ namespace UnrealBuildTool
 			return false;
 		}
 
-		/**
-		 *	Whether the editor should be built for this platform or not
-		 *	
-		 *	@param	InPlatform		The UnrealTargetPlatform being built
-		 *	@param	InConfiguration	The UnrealTargetConfiguration being built
-		 *	@return	bool			true if the editor should be built, false if not
-		 */
+		/// <summary>
+		/// Whether the editor should be built for this platform or not
+		/// </summary>
+		/// <param name="InPlatform"> The UnrealTargetPlatform being built</param>
+		/// <param name="InConfiguration">The UnrealTargetConfiguration being built</param>
+		/// <returns>bool   true if the editor should be built, false if not</returns>
 		public override bool ShouldNotBuildEditor(UnrealTargetPlatform InPlatform, UnrealTargetConfiguration InConfiguration)
 		{
 			return true;
@@ -332,11 +320,10 @@ namespace UnrealBuildTool
 			}
 		}
 
-		/**
-		 *	Setup the target environment for building
-		 *	
-		 *	@param	InBuildTarget		The target being built
-		 */
+		/// <summary>
+		/// Setup the target environment for building
+		/// </summary>
+		/// <param name="InBuildTarget"> The target being built</param>
 		public override void SetUpEnvironment(UEBuildTarget InBuildTarget)
 		{
 
@@ -399,11 +386,10 @@ namespace UnrealBuildTool
 			}
 		}
 
-		/**
-		 *	Setup the configuration environment for building
-		 *	
-		 *	@param	InBuildTarget		The target being built
-		 */
+		/// <summary>
+		/// Setup the configuration environment for building
+		/// </summary>
+		/// <param name="InBuildTarget"> The target being built</param>
 		public override void SetUpConfigurationEnvironment(UEBuildTarget InBuildTarget)
 		{
 			// Determine the C++ compile/link configuration based on the Unreal configuration.
@@ -460,14 +446,12 @@ namespace UnrealBuildTool
 			InBuildTarget.GlobalLinkEnvironment.Config.bCreateDebugInfo = true;
 		}
 
-		/**
-		 *	Whether this platform should create debug information or not
-		 *	
-		 *	@param	InPlatform			The UnrealTargetPlatform being built
-		 *	@param	InConfiguration		The UnrealTargetConfiguration being built
-		 *	
-		 *	@return	bool				true if debug info should be generated, false if not
-		 */
+		/// <summary>
+		/// Whether this platform should create debug information or not
+		/// </summary>
+		/// <param name="InPlatform">  The UnrealTargetPlatform being built</param>
+		/// <param name="InConfiguration"> The UnrealTargetConfiguration being built</param>
+		/// <returns>bool    true if debug info should be generated, false if not</returns>
 		public override bool ShouldCreateDebugInfo(UnrealTargetPlatform Platform, UnrealTargetConfiguration Configuration)
 		{
 			switch (Configuration)
@@ -483,9 +467,9 @@ namespace UnrealBuildTool
 			};
 		}
 
-		/**
-		 *	Return whether this platform has uniquely named binaries across multiple games
-		 */
+		/// <summary>
+		/// Return whether this platform has uniquely named binaries across multiple games
+		/// </summary>
 		public override bool HasUniqueBinaries()
 		{
 			// Windows applications have many shared binaries between games
diff --git a/Engine/Source/Programs/UnrealBuildTool/UWP/UWPDeploy.cs b/Engine/Source/Programs/UnrealBuildTool/UWP/UWPDeploy.cs
--- a/Engine/Source/Programs/UnrealBuildTool/UWP/UWPDeploy.cs
+++ b/Engine/Source/Programs/UnrealBuildTool/UWP/UWPDeploy.cs
@@ -15,9 +15,9 @@ namespace UnrealBuildTool
 	/// </summary>
 	class UWPDeploy : UEBuildDeploy
 	{
-		/**
-		 *	Register the platform with the UEBuildDeploy class
-		 */
+		/// <summary>
+		/// Register the platform with the UEBuildDeploy class
+		/// </summary>
 		public override void RegisterBuildDeploy()
 		{
 			// Register this deployment handle for UWP
@@ -25,9 +25,9 @@ namespace UnrealBuildTool
 			UEBuildDeploy.RegisterBuildDeploy(UnrealTargetPlatform.UWP, this);
 		}
 
-		/**
-		 *  Utility function to delete a file
-		 */
+		/// <summary>
+		/// Utility function to delete a file
+		/// </summary>
 		void DeployHelper_DeleteFile(string InFileToDelete)
 		{
 			Log.TraceInformation("UWPDeploy.DeployHelper_DeleteFile({0})", InFileToDelete);
@@ -43,10 +43,9 @@ namespace UnrealBuildTool
 			}
 		}
 
-		/**
-		 *	Copy the contents of the given source directory to the given destination directory
-		 *	
-		 */
+		/// <summary>
+		/// Copy the contents of the given source directory to the given destination directory
+		/// </summary>
 		bool CopySourceToDestDir(string InSourceDirectory, string InDestinationDirectory, string InWildCard,
 			bool bInIncludeSubDirectories, bool bInRemoveDestinationOrphans)
 		{
@@ -127,9 +126,9 @@ namespace UnrealBuildTool
 		}
 
 
-		/**
-		 *	Helper function for copying files
-		 */
+		/// <summary>
+		/// Helper function for copying files
+		/// </summary>
 		void CopyFile(string InSource, string InDest, bool bForce)
 		{
 			Log.TraceInformation("UWPDeploy.CopyFile({0}, {1}, {2},...)", InSource, InDest, bForce);
@@ -150,9 +149,9 @@ namespace UnrealBuildTool
 			}
 		}
 
-		/**
-		 *	Helper function for copying a tree files
-		 */
+		/// <summary>
+		/// Helper function for copying a tree files
+		/// </summary>
 		//void CopyDirectory(string InSource, string InDest, bool bForce, bool bRecurse)
 		//{
 		//    if (Directory.Exists(InSource))
diff --git a/Engine/Source/Programs/UnrealBuildTool/UWP/UWPProjectGenerator.cs b/Engine/Source/Programs/UnrealBuildTool/UWP/UWPProjectGenerator.cs
--- a/Engine/Source/Programs/UnrealBuildTool/UWP/UWPProjectGenerator.cs
+++ b/Engine/Source/Programs/UnrealBuildTool/UWP/UWPProjectGenerator.cs
@@ -10,14 +10,14 @@ using System.IO;
 
 namespace UnrealBuildTool
 {
-	/**
-	 *	Base class for platform-specific project generators 
-	 */
+	/// <summary>
+	/// Base class for platform-specific project generators
+	/// </summary>
 	public class UWPProjectGenerator : UEPlatformProjectGenerator
 	{
-		/**
-		 *	Register the platform with the UEPlatformProjectGenerator class
-		 */
+		/// <summary>
+		/// Register the platform with the UEPlatformProjectGenerator class
+		/// </summary>
 		public override void RegisterPlatformProjectGenerator()
 		{
 			if (UWPPlatform.bEnableUWPSupport)
@@ -74,24 +74,21 @@ namespace UnrealBuildTool
 		///
 		///	VisualStudio project generation functions
 		///	
-		/**
-		 *	Whether this build platform has native support for VisualStudio
-		 *	
-		 *	@param	InPlatform			The UnrealTargetPlatform being built
-		 *	@param	InConfiguration		The UnrealTargetConfiguration being built
-		 *	
-		 *	@return	bool				true if native VisualStudio support (or custom VSI) is available
-		 */
+		/// <summary>
+		/// Whether this build platform has native support for VisualStudio
+		/// </summary>
+		/// <param name="InPlatform">  The UnrealTargetPlatform being built</param>
+		/// <param name="InConfiguration"> The UnrealTargetConfiguration being built</param>
+		/// <returns>bool    true if native VisualStudio support (or custom VSI) is available</returns>
 		public override bool HasVisualStudioSupport(UnrealTargetPlatform InPlatform, UnrealTargetConfiguration InConfiguration)
 		{
 			return false;
 		}
 
-		/**
-		 * Get whether this platform deploys 
-		 * 
-		 * @return	bool		true if the 'Deploy' option should be enabled
-		 */
+		/// <summary>
+		/// Get whether this platform deploys
+		/// </summary>
+		/// <returns>bool  true if the 'Deploy' option should be enabled</returns>
 		public override bool GetVisualStudioDeploymentEnabled(UnrealTargetPlatform InPlatform, UnrealTargetConfiguration InConfiguration)
 		{
 			return true;
diff --git a/Engine/Source/Programs/UnrealBuildTool/UWP/UWPToolChain.cs b/Engine/Source/Programs/UnrealBuildTool/UWP/UWPToolChain.cs
--- a/Engine/Source/Programs/UnrealBuildTool/UWP/UWPToolChain.cs
+++ b/Engine/Source/Programs/UnrealBuildTool/UWP/UWPToolChain.cs
@@ -1022,7 +1022,9 @@ namespace UnrealBuildTool
 		}
 
 
-		/** Gets the default include paths for the given platform. */
+		/// <summary>
+		/// Gets the default include paths for the given platform.
+		/// </summary>
 		public static string GetVCIncludePaths(CPPTargetPlatform Platform)
 		{
 			Debug.Assert(Platform == CPPTargetPlatform.UWP);
diff --git a/Engine/Source/Programs/UnrealBuildTool/Utilities/CopyrightVerify.cs b/Engine/Source/Programs/UnrealBuildTool/Utilities/CopyrightVerify.cs
--- a/Engine/Source/Programs/UnrealBuildTool/Utilities/CopyrightVerify.cs
+++ b/Engine/Source/Programs/UnrealBuildTool/Utilities/CopyrightVerify.cs
@@ -25,13 +25,11 @@ namespace UnrealBuildTool
 				"ENGINE\\SHADERS\\SHADERVERSION.USF",
 			};
 
-		/**
-		 * See if the given file is whitelisted
-		 * 
-		 * @param	InFilename		The name of the file to check
-		 * 
-		 * @return	bool			true if it is whitelisted, false if not
-		 */
+		/// <summary>
+		/// See if the given file is whitelisted
+		/// </summary>
+		/// <param name="InFilename"> The name of the file to check</param>
+		/// <returns>bool   true if it is whitelisted, false if not</returns>
 		static bool IsWhitelistedFile(string InFilename)
 		{
 			string FilenameUpper = InFilename.ToUpperInvariant();
@@ -48,13 +46,12 @@ namespace UnrealBuildTool
 
 		static List<string> BadFiles = new List<string>();
 
-		/**
-		 * Process the given list of files, checking each one has the proper copyright.
-		 * This REQUIRES that the first line in any file is the copyright line!!!
-		 * 
-		 * @param	InFiles			The list of files to process
-		 * @param	OutBadFiles		OUTPUT - List of files that failed the copyright check (and why)
-		 */
+		/// <summary>
+		/// Process the given list of files, checking each one has the proper copyright.
+		/// This REQUIRES that the first line in any file is the copyright line!!!
+		/// </summary>
+		/// <param name="InFiles">  The list of files to process</param>
+		/// <param name="OutBadFiles"> OUTPUT - List of files that failed the copyright check (and why)</param>
 		static void ProcessFileList(List<FileInfo> InFiles, ref Dictionary<string, string> OutBadFiles)
 		{
 			foreach (FileInfo CheckFile in InFiles)
