Convert everything in UBT to use standard C# XML comments, so IntelliSense can display useful information in tooltips.

diff --git a/Engine/Source/Programs/UnrealBuildTool/WinRT/UEBuildWinRT.cs b/Engine/Source/Programs/UnrealBuildTool/WinRT/UEBuildWinRT.cs
--- a/Engine/Source/Programs/UnrealBuildTool/WinRT/UEBuildWinRT.cs
+++ b/Engine/Source/Programs/UnrealBuildTool/WinRT/UEBuildWinRT.cs
@@ -40,17 +40,17 @@ namespace UnrealBuildTool
 			return false;
 		}
 
-		/**
-		 *	Whether the required external SDKs are installed for this platform
-		 */
+		/// <summary>
+		/// Whether the required external SDKs are installed for this platform
+		/// </summary>
 		protected override SDKStatus HasRequiredManualSDKInternal()
 		{
 			return !Utils.IsRunningOnMono && IsVisualStudioInstalled() ? SDKStatus.Valid : SDKStatus.Invalid;
 		}
 
-		/**
-		 *	Register the platform with the UEBuildPlatform class
-		 */
+		/// <summary>
+		/// Register the platform with the UEBuildPlatform class
+		/// </summary>
 		protected override void RegisterBuildPlatformInternal()
 		{
 			[REDACTED_EMAIL]: Add platform support
@@ -91,13 +91,11 @@ namespace UnrealBuildTool
 			}
 		}
 
-		/**
-		 *	Retrieve the CPPTargetPlatform for the given UnrealTargetPlatform
-		 *
-		 *	@param	InUnrealTargetPlatform		The UnrealTargetPlatform being build
-		 *	
-		 *	@return	CPPTargetPlatform			The CPPTargetPlatform to compile for
-		 */
+		/// <summary>
+		/// Retrieve the CPPTargetPlatform for the given UnrealTargetPlatform
+		/// </summary>
+		/// <param name="InUnrealTargetPlatform"> The UnrealTargetPlatform being build</param>
+		/// <returns>CPPTargetPlatform   The CPPTargetPlatform to compile for</returns>
 		public override CPPTargetPlatform GetCPPTargetPlatform(UnrealTargetPlatform InUnrealTargetPlatform)
 		{
 			switch (InUnrealTargetPlatform)
@@ -110,13 +108,11 @@ namespace UnrealBuildTool
 			throw new BuildException("WinRTPlatform::GetCPPTargetPlatform: Invalid request for {0}", InUnrealTargetPlatform.ToString());
 		}
 
-		/**
-		 *	Get the extension to use for the given binary type
-		 *	
-		 *	@param	InBinaryType		The binary type being built
-		 *	
-		 *	@return	string				The binary extension (i.e. 'exe' or 'dll')
-		 */
+		/// <summary>
+		/// Get the extension to use for the given binary type
+		/// </summary>
+		/// <param name="InBinaryType"> The binary type being built</param>
+		/// <returns>string    The binary extension (i.e. 'exe' or 'dll')</returns>
 		public override string GetBinaryExtension(UEBuildBinaryType InBinaryType)
 		{
 			switch (InBinaryType)
@@ -135,37 +131,33 @@ namespace UnrealBuildTool
 			return base.GetBinaryExtension(InBinaryType);
 		}
 
-		/**
-		 *	Get the extension to use for debug info for the given binary type
-		 *	
-		 *	@param	InBinaryType		The binary type being built
-		 *	
-		 *	@return	string				The debug info extension (i.e. 'pdb')
-		 */
+		/// <summary>
+		/// Get the extension to use for debug info for the given binary type
+		/// </summary>
+		/// <param name="InBinaryType"> The binary type being built</param>
+		/// <returns>string    The debug info extension (i.e. 'pdb')</returns>
 		public override string GetDebugInfoExtension(UEBuildBinaryType InBinaryType)
 		{
 			return ".pdb";
 		}
 
-		/**
-		 *	Gives the platform a chance to 'override' the configuration settings 
-		 *	that are overridden on calls to RunUBT.
-		 *	
-		 *	@param	InPlatform			The UnrealTargetPlatform being built
-		 *	@param	InConfiguration		The UnrealTargetConfiguration being built
-		 *	
-		 *	@return	bool				true if debug info should be generated, false if not
-		 */
+		/// <summary>
+		/// Gives the platform a chance to 'override' the configuration settings
+		/// that are overridden on calls to RunUBT.
+		/// </summary>
+		/// <param name="InPlatform">  The UnrealTargetPlatform being built</param>
+		/// <param name="InConfiguration"> The UnrealTargetConfiguration being built</param>
+		/// <returns>bool    true if debug info should be generated, false if not</returns>
 		public override void ResetBuildConfiguration(UnrealTargetPlatform InPlatform, UnrealTargetConfiguration InConfiguration)
 		{
 			ValidateUEBuildConfiguration();
 		}
 
-		/**
-		 *	Validate the UEBuildConfiguration for this platform
-		 *	This is called BEFORE calling UEBuildConfiguration to allow setting 
-		 *	various fields used in that function such as CompileLeanAndMean...
-		 */
+		/// <summary>
+		/// Validate the UEBuildConfiguration for this platform
+		/// This is called BEFORE calling UEBuildConfiguration to allow setting
+		/// various fields used in that function such as CompileLeanAndMean...
+		/// </summary>
 		public override void ValidateUEBuildConfiguration()
 		{
 			UEBuildConfiguration.bCompileLeanAndMeanUE = true;
@@ -179,60 +171,55 @@ namespace UnrealBuildTool
 			BuildConfiguration.bStopXGECompilationAfterErrors = false;
 		}
 
-		/**
-		 * Whether this platform should build a monolithic binary
-		 */
+		/// <summary>
+		/// Whether this platform should build a monolithic binary
+		/// </summary>
 		public override bool ShouldCompileMonolithicBinary(UnrealTargetPlatform InPlatform)
 		{
 			return true;
 		}
 
-		/**
-		 *	Whether PDB files should be used
-		 *	
-		 *	@param	InPlatform			The CPPTargetPlatform being built
-		 *	@param	InConfiguration		The CPPTargetConfiguration being built
-		 *	@param	bInCreateDebugInfo	true if debug info is getting create, false if not
-		 *	
-		 *	@return	bool	true if PDB files should be used, false if not
-		 */
+		/// <summary>
+		/// Whether PDB files should be used
+		/// </summary>
+		/// <param name="InPlatform">  The CPPTargetPlatform being built</param>
+		/// <param name="InConfiguration"> The CPPTargetConfiguration being built</param>
+		/// <param name="bInCreateDebugInfo">true if debug info is getting create, false if not</param>
+		/// <returns>bool true if PDB files should be used, false if not</returns>
 		public override bool ShouldUsePDBFiles(CPPTargetPlatform Platform, CPPTargetConfiguration Configuration, bool bCreateDebugInfo)
 		{
 			return true;
 		}
 
-		/**
-		 *	Whether the editor should be built for this platform or not
-		 *	
-		 *	@param	InPlatform		The UnrealTargetPlatform being built
-		 *	@param	InConfiguration	The UnrealTargetConfiguration being built
-		 *	@return	bool			true if the editor should be built, false if not
-		 */
+		/// <summary>
+		/// Whether the editor should be built for this platform or not
+		/// </summary>
+		/// <param name="InPlatform"> The UnrealTargetPlatform being built</param>
+		/// <param name="InConfiguration">The UnrealTargetConfiguration being built</param>
+		/// <returns>bool   true if the editor should be built, false if not</returns>
 		public override bool ShouldNotBuildEditor(UnrealTargetPlatform InPlatform, UnrealTargetConfiguration InConfiguration)
 		{
 			return true;
 		}
 
-		/**
-		 *	Whether the platform requires cooked data
-		 *	
-		 *	@param	InPlatform		The UnrealTargetPlatform being built
-		 *	@param	InConfiguration	The UnrealTargetConfiguration being built
-		 *	@return	bool			true if the platform requires cooked data, false if not
-		 */
+		/// <summary>
+		/// Whether the platform requires cooked data
+		/// </summary>
+		/// <param name="InPlatform"> The UnrealTargetPlatform being built</param>
+		/// <param name="InConfiguration">The UnrealTargetConfiguration being built</param>
+		/// <returns>bool   true if the platform requires cooked data, false if not</returns>
 		public override bool BuildRequiresCookedData(UnrealTargetPlatform InPlatform, UnrealTargetConfiguration InConfiguration)
 		{
 			return true;
 		}
 
-		/**
-		 *	Get a list of extra modules the platform requires.
-		 *	This is to allow undisclosed platforms to add modules they need without exposing information about the platfomr.
-		 *	
-		 *	@param	Target						The target being build
-		 *	@param	BuildTarget					The UEBuildTarget getting build
-		 *	@param	PlatformExtraModules		OUTPUT the list of extra modules the platform needs to add to the target
-		 */
+		/// <summary>
+		/// Get a list of extra modules the platform requires.
+		/// This is to allow undisclosed platforms to add modules they need without exposing information about the platfomr.
+		/// </summary>
+		/// <param name="Target">     The target being build</param>
+		/// <param name="BuildTarget">    The UEBuildTarget getting build</param>
+		/// <param name="PlatformExtraModules"> OUTPUT the list of extra modules the platform needs to add to the target</param>
 		public override void GetExtraModules(TargetInfo Target, UEBuildTarget BuildTarget, ref List<string> PlatformExtraModules)
 		{
 			if (Target.Platform == UnrealTargetPlatform.WinRT)
@@ -241,15 +228,14 @@ namespace UnrealBuildTool
 			}
 		}
 
-		/**
-		 *	Modify the newly created module passed in for this platform.
-		 *	This is not required - but allows for hiding details of a
-		 *	particular platform.
-		 *	
-		 *  @param  Name			The name of the module
-		 *	@param	Rules			The module rules
-		 *	@param	Target			The target being build
-		 */
+		/// <summary>
+		/// Modify the newly created module passed in for this platform.
+		/// This is not required - but allows for hiding details of a
+		/// particular platform.
+		/// </summary>
+		/// <param name="">Name   The name of the module</param>
+		/// <param name="Rules">  The module rules</param>
+		/// <param name="Target">  The target being build</param>
 		public override void ModifyModuleRules(string ModuleName, ModuleRules Rules, TargetInfo Target)
 		{
 			if ((Target.Platform == UnrealTargetPlatform.WinRT) ||
@@ -512,11 +498,10 @@ namespace UnrealBuildTool
 			}
 		}
 
-		/**
-		 *	Setup the target environment for building
-		 *	
-		 *	@param	InBuildTarget		The target being built
-		 */
+		/// <summary>
+		/// Setup the target environment for building
+		/// </summary>
+		/// <param name="InBuildTarget"> The target being built</param>
 		public override void SetUpEnvironment(UEBuildTarget InBuildTarget)
 		{
 			InBuildTarget.GlobalCompileEnvironment.Config.Definitions.Add("PLATFORM_DESKTOP=0");
@@ -611,14 +596,12 @@ namespace UnrealBuildTool
 			//InBuildTarget.GlobalCompileEnvironment.Config.FrameworkAssemblyDependencies.Add("Windows.winmd");
 		}
 
-		/**
-		 *	Whether this platform should create debug information or not
-		 *	
-		 *	@param	InPlatform			The UnrealTargetPlatform being built
-		 *	@param	InConfiguration		The UnrealTargetConfiguration being built
-		 *	
-		 *	@return	bool				true if debug info should be generated, false if not
-		 */
+		/// <summary>
+		/// Whether this platform should create debug information or not
+		/// </summary>
+		/// <param name="InPlatform">  The UnrealTargetPlatform being built</param>
+		/// <param name="InConfiguration"> The UnrealTargetConfiguration being built</param>
+		/// <returns>bool    true if debug info should be generated, false if not</returns>
 		public override bool ShouldCreateDebugInfo(UnrealTargetPlatform Platform, UnrealTargetConfiguration Configuration)
 		{
 			switch (Configuration)
@@ -632,11 +615,10 @@ namespace UnrealBuildTool
 			};
 		}
 
-		/**
-		 *	Setup the binaries for this specific platform.
-		 *	
-		 *	@param	InBuildTarget		The target being built
-		 */
+		/// <summary>
+		/// Setup the binaries for this specific platform.
+		/// </summary>
+		/// <param name="InBuildTarget"> The target being built</param>
 		public override void SetupBinaries(UEBuildTarget InBuildTarget)
 		{
 			InBuildTarget.ExtraModuleNames.Add("XInput");
@@ -645,9 +627,9 @@ namespace UnrealBuildTool
 		//
 		// WinRT specific functions
 		//
-		/**
-		 *	Should the app be compiled as WinRT
-		 */
+		/// <summary>
+		/// Should the app be compiled as WinRT
+		/// </summary>
 		public static bool ShouldCompileWinRT()
 		{
 			return WinRTPlatform.bCompileWinRT;
diff --git a/Engine/Source/Programs/UnrealBuildTool/WinRT/WinRTDeploy.cs b/Engine/Source/Programs/UnrealBuildTool/WinRT/WinRTDeploy.cs
--- a/Engine/Source/Programs/UnrealBuildTool/WinRT/WinRTDeploy.cs
+++ b/Engine/Source/Programs/UnrealBuildTool/WinRT/WinRTDeploy.cs
@@ -13,9 +13,9 @@ namespace UnrealBuildTool
 	/// </summary>
 	class WinRTDeploy : UEBuildDeploy
 	{
-		/**
-		 *	Register the platform with the UEBuildDeploy class
-		 */
+		/// <summary>
+		/// Register the platform with the UEBuildDeploy class
+		/// </summary>
 		public override void RegisterBuildDeploy()
 		{
 			// Register this deployment handle for WinRT
@@ -23,9 +23,9 @@ namespace UnrealBuildTool
 			UEBuildDeploy.RegisterBuildDeploy(UnrealTargetPlatform.WinRT, this);
 		}
 
-		/**
-		 *  Utility function to delete a file
-		 */
+		/// <summary>
+		/// Utility function to delete a file
+		/// </summary>
 		void DeployHelper_DeleteFile(string InFileToDelete)
 		{
 			if (File.Exists(InFileToDelete) == true)
@@ -40,10 +40,9 @@ namespace UnrealBuildTool
 			}
 		}
 
-		/**
-		 *	Copy the contents of the given source directory to the given destination directory
-		 *	
-		 */
+		/// <summary>
+		/// Copy the contents of the given source directory to the given destination directory
+		/// </summary>
 		bool CopySourceToDestDir(string InSourceDirectory, string InDestinationDirectory, string InWildCard,
 			bool bInIncludeSubDirectories, bool bInRemoveDestinationOrphans)
 		{
@@ -121,16 +120,14 @@ namespace UnrealBuildTool
 			return true;
 		}
 
-		/**
-		 *	Copy the game contents via reading in the <GAME>\Saved\DeployList\ContentList.txt file.
-		 *	If the file doesn't exist, simply mirror ALL content to the deployment directory.
-		 *	
-		 *	@param	InAppName			The name of the target app - could be the GameName
-		 *	@param	InPlatform			UnrealTargetPlatform of the target
-		 *	@param	InConfiguration		UnrealTargetConfiguration of the target
-		 *	
-		 *	@return	bool				true if successful
-		 */
+		/// <summary>
+		/// Copy the game contents via reading in the <GAME>\Saved\DeployList\ContentList.txt file.
+		/// If the file doesn't exist, simply mirror ALL content to the deployment directory.
+		/// </summary>
+		/// <param name="InAppName">  The name of the target app - could be the GameName</param>
+		/// <param name="InPlatform">  UnrealTargetPlatform of the target</param>
+		/// <param name="InConfiguration"> UnrealTargetConfiguration of the target</param>
+		/// <returns>bool    true if successful</returns>
 		bool CopyGameContent(string InAppName, UnrealTargetPlatform InPlatform, UnrealTargetConfiguration InConfiguration)
 		{
 			string DeployDirectory = "../../" + InAppName + "/Binaries/WinRT/Image/Loose/";
@@ -206,9 +203,9 @@ namespace UnrealBuildTool
 			return true;
 		}
 
-		/**
-		 *	Helper function for copying files
-		 */
+		/// <summary>
+		/// Helper function for copying files
+		/// </summary>
 		void CopyFile(string InSource, string InDest, bool bForce)
 		{
 			if (File.Exists(InSource) == true)
diff --git a/Engine/Source/Programs/UnrealBuildTool/WinRT/WinRTProjectGenerator.cs b/Engine/Source/Programs/UnrealBuildTool/WinRT/WinRTProjectGenerator.cs
--- a/Engine/Source/Programs/UnrealBuildTool/WinRT/WinRTProjectGenerator.cs
+++ b/Engine/Source/Programs/UnrealBuildTool/WinRT/WinRTProjectGenerator.cs
@@ -8,14 +8,14 @@ using System.IO;
 
 namespace UnrealBuildTool
 {
-	/**
-	 *	Base class for platform-specific project generators 
-	 */
+	/// <summary>
+	/// Base class for platform-specific project generators
+	/// </summary>
 	class WinRTProjectGenerator : UEPlatformProjectGenerator
 	{
-		/**
-		 *	Register the platform with the UEPlatformProjectGenerator class
-		 */
+		/// <summary>
+		/// Register the platform with the UEPlatformProjectGenerator class
+		/// </summary>
 		public override void RegisterPlatformProjectGenerator()
 		{
 			// Register this project generator for WinRT
@@ -34,27 +34,23 @@ namespace UnrealBuildTool
 		///
 		///	VisualStudio project generation functions
 		///	
-		/**
-		 *	Whether this build platform has native support for VisualStudio
-		 *	
-		 *	@param	InPlatform			The UnrealTargetPlatform being built
-		 *	@param	InConfiguration		The UnrealTargetConfiguration being built
-		 *	
-		 *	@return	bool				true if native VisualStudio support (or custom VSI) is available
-		 */
+		/// <summary>
+		/// Whether this build platform has native support for VisualStudio
+		/// </summary>
+		/// <param name="InPlatform">  The UnrealTargetPlatform being built</param>
+		/// <param name="InConfiguration"> The UnrealTargetConfiguration being built</param>
+		/// <returns>bool    true if native VisualStudio support (or custom VSI) is available</returns>
 		public override bool HasVisualStudioSupport(UnrealTargetPlatform InPlatform, UnrealTargetConfiguration InConfiguration)
 		{
 			return false;
 		}
 
-		/**
-		 *	Return the VisualStudio platform name for this build platform
-		 *	
-		 *	@param	InPlatform			The UnrealTargetPlatform being built
-		 *	@param	InConfiguration		The UnrealTargetConfiguration being built
-		 *	
-		 *	@return	string				The name of the platform that VisualStudio recognizes
-		 */
+		/// <summary>
+		/// Return the VisualStudio platform name for this build platform
+		/// </summary>
+		/// <param name="InPlatform">  The UnrealTargetPlatform being built</param>
+		/// <param name="InConfiguration"> The UnrealTargetConfiguration being built</param>
+		/// <returns>string    The name of the platform that VisualStudio recognizes</returns>
 		public override string GetVisualStudioPlatformName(UnrealTargetPlatform InPlatform, UnrealTargetConfiguration InConfiguration)
 		{
 			if (InPlatform == UnrealTargetPlatform.WinRT)
@@ -64,24 +60,21 @@ namespace UnrealBuildTool
 			return InPlatform.ToString();
 		}
 
-		/**
-		 *	Return the platform toolset string to write into the project configuration
-		 *	
-		 *	@param	InPlatform			The UnrealTargetPlatform being built
-		 *	@param	InConfiguration		The UnrealTargetConfiguration being built
-		 *	
-		 *	@return	string				The custom configuration section for the project file; Empty string if it doesn't require one
-		 */
+		/// <summary>
+		/// Return the platform toolset string to write into the project configuration
+		/// </summary>
+		/// <param name="InPlatform">  The UnrealTargetPlatform being built</param>
+		/// <param name="InConfiguration"> The UnrealTargetConfiguration being built</param>
+		/// <returns>string    The custom configuration section for the project file; Empty string if it doesn't require one</returns>
 		public override string GetVisualStudioPlatformToolsetString(UnrealTargetPlatform InPlatform, UnrealTargetConfiguration InConfiguration, VCProjectFile InProjectFile)
 		{
 			return "		<PlatformToolset>v110</PlatformToolset>" + ProjectFileGenerator.NewLine;
 		}
 
-		/**
-		 * Get whether this platform deploys 
-		 * 
-		 * @return	bool		true if the 'Deploy' option should be enabled
-		 */
+		/// <summary>
+		/// Get whether this platform deploys
+		/// </summary>
+		/// <returns>bool  true if the 'Deploy' option should be enabled</returns>
 		public override bool GetVisualStudioDeploymentEnabled(UnrealTargetPlatform InPlatform, UnrealTargetConfiguration InConfiguration)
 		{
 			return true;
diff --git a/Engine/Source/Programs/UnrealBuildTool/WinRT/WinRTToolChain.cs b/Engine/Source/Programs/UnrealBuildTool/WinRT/WinRTToolChain.cs
--- a/Engine/Source/Programs/UnrealBuildTool/WinRT/WinRTToolChain.cs
+++ b/Engine/Source/Programs/UnrealBuildTool/WinRT/WinRTToolChain.cs
@@ -850,7 +850,9 @@ namespace UnrealBuildTool
 			return OutputFile;
 		}
 
-		/** Accesses the bin directory for the VC toolchain for the specified platform. */
+		/// <summary>
+		/// Accesses the bin directory for the VC toolchain for the specified platform.
+		/// </summary>
 		static string GetVCToolPath(CPPTargetPlatform Platform, CPPTargetConfiguration Configuration, string ToolName)
 		{
 			// Initialize environment variables required for spawned tools.
@@ -903,7 +905,9 @@ namespace UnrealBuildTool
 			return VCToolPath;
 		}
 
-		/** Accesses the directory for .NET Framework binaries such as MSBuild */
+		/// <summary>
+		/// Accesses the directory for .NET Framework binaries such as MSBuild
+		/// </summary>
 		static string GetDotNetFrameworkToolPath(CPPTargetPlatform Platform, string ToolName)
 		{
 			// Initialize environment variables required for spawned tools.
@@ -920,15 +924,19 @@ namespace UnrealBuildTool
 			return ToolPath;
 		}
 
-		/** Helper to only initialize environment variables once. */
+		/// <summary>
+		/// Helper to only initialize environment variables once.
+		/// </summary>
 		static bool bAreEnvironmentVariablesAlreadyInitialized = false;
 
-		/** Installation folder of the Windows SDK, e.g. C:\Program Files\Microsoft SDKs\Windows\v6.0A\ */
+		/// <summary>
+		/// Installation folder of the Windows SDK, e.g. C:\Program Files\Microsoft SDKs\Windows\v6.0A\
+		/// </summary>
 		static string WindowsSDKDir = "";
 
-		/**
-		 * Initializes environment variables required by toolchain. Different for 32 and 64 bit.
-		 */
+		/// <summary>
+		/// Initializes environment variables required by toolchain. Different for 32 and 64 bit.
+		/// </summary>
 		static void InitializeEnvironmentVariables(CPPTargetPlatform Platform)
 		{
 			if (!bAreEnvironmentVariablesAlreadyInitialized)
diff --git a/Engine/Source/Programs/UnrealBuildTool/Windows/UEBuildWindows.cs b/Engine/Source/Programs/UnrealBuildTool/Windows/UEBuildWindows.cs
--- a/Engine/Source/Programs/UnrealBuildTool/Windows/UEBuildWindows.cs
+++ b/Engine/Source/Programs/UnrealBuildTool/Windows/UEBuildWindows.cs
@@ -111,7 +111,9 @@ namespace UnrealBuildTool
 		/// True if we allow using addresses larger than 2GB on 32 bit builds
 		public static bool bBuildLargeAddressAwareBinary = true;
 
-		/** True if VS EnvDTE is available (false when building using Visual Studio Express) */
+		/// <summary>
+		/// True if VS EnvDTE is available (false when building using Visual Studio Express)
+		/// </summary>
 		public static bool bHasVisualStudioDTE
 		{
 			get
@@ -147,7 +149,9 @@ namespace UnrealBuildTool
 			}
 		}
 
-		/** The current architecture */
+		/// <summary>
+		/// The current architecture
+		/// </summary>
 		public override string GetActiveArchitecture()
 		{
 			return IsWindowsXPSupported() ? "_XP" : base.GetActiveArchitecture();
@@ -158,23 +162,20 @@ namespace UnrealBuildTool
 			return SDKStatus.Valid;
 		}
 
-		/**
-		 * Returns VisualStudio common tools path for current compiler.
-		 * 
-		 * @return Common tools path.
-		 */
+		/// <summary>
+		/// Returns VisualStudio common tools path for current compiler.
+		/// </summary>
+		/// <returns>Common tools path.</returns>
 		public static string GetVSComnToolsPath()
 		{
 			return GetVSComnToolsPath(Compiler);
 		}
 
-		/**
-		 * Returns VisualStudio common tools path for given compiler.
-		 * 
-		 * @param Compiler Compiler for which to return tools path.
-		 * 
-		 * @return Common tools path.
-		 */
+		/// <summary>
+		/// Returns VisualStudio common tools path for given compiler.
+		/// </summary>
+		/// <param name="Compiler">Compiler for which to return tools path.</param>
+		/// <returns>Common tools path.</returns>
 		public static string GetVSComnToolsPath(WindowsCompiler Compiler)
 		{
 			int VSVersion;
@@ -221,9 +222,9 @@ namespace UnrealBuildTool
 			return new DirectoryInfo(Path.Combine(VSPath, "..", "Tools")).FullName;
 		}
 
-		/**
-         *	If this platform can be compiled with SN-DBS
-         */
+		/// <summary>
+		/// If this platform can be compiled with SN-DBS
+		/// </summary>
 		public override bool CanUseSNDBS()
 		{
 			// Check that SN-DBS is available
@@ -240,9 +241,9 @@ namespace UnrealBuildTool
 			}
 		}
 
-		/**
-		 *	Register the platform with the UEBuildPlatform class
-		 */
+		/// <summary>
+		/// Register the platform with the UEBuildPlatform class
+		/// </summary>
 		protected override void RegisterBuildPlatformInternal()
 		{
 			// Register this build platform for both Win64 and Win32
@@ -257,13 +258,11 @@ namespace UnrealBuildTool
 			UEBuildPlatform.RegisterPlatformWithGroup(UnrealTargetPlatform.Win32, UnrealPlatformGroup.Microsoft);
 		}
 
-		/**
-		 *	Retrieve the CPPTargetPlatform for the given UnrealTargetPlatform
-		 *
-		 *	@param	InUnrealTargetPlatform		The UnrealTargetPlatform being build
-		 *	
-		 *	@return	CPPTargetPlatform			The CPPTargetPlatform to compile for
-		 */
+		/// <summary>
+		/// Retrieve the CPPTargetPlatform for the given UnrealTargetPlatform
+		/// </summary>
+		/// <param name="InUnrealTargetPlatform"> The UnrealTargetPlatform being build</param>
+		/// <returns>CPPTargetPlatform   The CPPTargetPlatform to compile for</returns>
 		public override CPPTargetPlatform GetCPPTargetPlatform(UnrealTargetPlatform InUnrealTargetPlatform)
 		{
 			switch (InUnrealTargetPlatform)
@@ -277,13 +276,11 @@ namespace UnrealBuildTool
 			throw new BuildException("WindowsPlatform::GetCPPTargetPlatform: Invalid request for {0}", InUnrealTargetPlatform.ToString());
 		}
 
-		/**
-		 *	Get the extension to use for the given binary type
-		 *	
-		 *	@param	InBinaryType		The binrary type being built
-		 *	
-		 *	@return	string				The binary extenstion (ie 'exe' or 'dll')
-		 */
+		/// <summary>
+		/// Get the extension to use for the given binary type
+		/// </summary>
+		/// <param name="InBinaryType"> The binrary type being built</param>
+		/// <returns>string    The binary extenstion (ie 'exe' or 'dll')</returns>
 		public override string GetBinaryExtension(UEBuildBinaryType InBinaryType)
 		{
 			switch (InBinaryType)
@@ -333,13 +330,11 @@ namespace UnrealBuildTool
 		}
 
 
-		/**
-		 *	Get the extension to use for debug info for the given binary type
-		 *	
-		 *	@param	InBinaryType		The binary type being built
-		 *	
-		 *	@return	string				The debug info extension (i.e. 'pdb')
-		 */
+		/// <summary>
+		/// Get the extension to use for debug info for the given binary type
+		/// </summary>
+		/// <param name="InBinaryType"> The binary type being built</param>
+		/// <returns>string    The debug info extension (i.e. 'pdb')</returns>
 		public override string GetDebugInfoExtension(UEBuildBinaryType InBinaryType)
 		{
 			switch (InBinaryType)
@@ -351,28 +346,24 @@ namespace UnrealBuildTool
 			return "";
 		}
 
-		/**
-		 *	Whether incremental linking should be used
-		 *	
-		 *	@param	InPlatform			The CPPTargetPlatform being built
-		 *	@param	InConfiguration		The CPPTargetConfiguration being built
-		 *	
-		 *	@return	bool	true if incremental linking should be used, false if not
-		 */
+		/// <summary>
+		/// Whether incremental linking should be used
+		/// </summary>
+		/// <param name="InPlatform">  The CPPTargetPlatform being built</param>
+		/// <param name="InConfiguration"> The CPPTargetConfiguration being built</param>
+		/// <returns>bool true if incremental linking should be used, false if not</returns>
 		public override bool ShouldUseIncrementalLinking(CPPTargetPlatform Platform, CPPTargetConfiguration Configuration)
 		{
 			return (Configuration == CPPTargetConfiguration.Debug);
 		}
 
-		/**
-		 *	Whether PDB files should be used
-		 *	
-		 *	@param	InPlatform			The CPPTargetPlatform being built
-		 *	@param	InConfiguration		The CPPTargetConfiguration being built
-		 *	@param	bInCreateDebugInfo	true if debug info is getting create, false if not
-		 *	
-		 *	@return	bool	true if PDB files should be used, false if not
-		 */
+		/// <summary>
+		/// Whether PDB files should be used
+		/// </summary>
+		/// <param name="InPlatform">  The CPPTargetPlatform being built</param>
+		/// <param name="InConfiguration"> The CPPTargetConfiguration being built</param>
+		/// <param name="bInCreateDebugInfo">true if debug info is getting create, false if not</param>
+		/// <returns>bool true if PDB files should be used, false if not</returns>
 		public override bool ShouldUsePDBFiles(CPPTargetPlatform Platform, CPPTargetConfiguration Configuration, bool bCreateDebugInfo)
 		{
 			// Only supported on PC.
@@ -383,13 +374,12 @@ namespace UnrealBuildTool
 			return false;
 		}
 
-		/**
-		 *	Whether the editor should be built for this platform or not
-		 *	
-		 *	@param	InPlatform		The UnrealTargetPlatform being built
-		 *	@param	InConfiguration	The UnrealTargetConfiguration being built
-		 *	@return	bool			true if the editor should be built, false if not
-		 */
+		/// <summary>
+		/// Whether the editor should be built for this platform or not
+		/// </summary>
+		/// <param name="InPlatform"> The UnrealTargetPlatform being built</param>
+		/// <param name="InConfiguration">The UnrealTargetConfiguration being built</param>
+		/// <returns>bool   true if the editor should be built, false if not</returns>
 		public override bool ShouldNotBuildEditor(UnrealTargetPlatform InPlatform, UnrealTargetConfiguration InConfiguration)
 		{
 			return false;
@@ -470,11 +460,11 @@ namespace UnrealBuildTool
 			}
 		}
 
-		/**
-		 *	Validate the UEBuildConfiguration for this platform
-		 *	This is called BEFORE calling UEBuildConfiguration to allow setting 
-		 *	various fields used in that function such as CompileLeanAndMean...
-		 */
+		/// <summary>
+		/// Validate the UEBuildConfiguration for this platform
+		/// This is called BEFORE calling UEBuildConfiguration to allow setting
+		/// various fields used in that function such as CompileLeanAndMean...
+		/// </summary>
 		public override void ValidateUEBuildConfiguration()
 		{
 			UEBuildConfiguration.bCompileICU = true;
@@ -536,11 +526,10 @@ namespace UnrealBuildTool
 			}
 		}
 
-		/**
-		 *	Setup the target environment for building
-		 *	
-		 *	@param	InBuildTarget		The target being built
-		 */
+		/// <summary>
+		/// Setup the target environment for building
+		/// </summary>
+		/// <param name="InBuildTarget"> The target being built</param>
 		public override void SetUpEnvironment(UEBuildTarget InBuildTarget)
 		{
 			InBuildTarget.GlobalCompileEnvironment.Config.Definitions.Add("WIN32=1");
@@ -709,11 +698,10 @@ namespace UnrealBuildTool
 			}
 		}
 
-		/**
-		 *	Setup the configuration environment for building
-		 *	
-		 *	@param	InBuildTarget		The target being built
-		 */
+		/// <summary>
+		/// Setup the configuration environment for building
+		/// </summary>
+		/// <param name="InBuildTarget"> The target being built</param>
 		public override void SetUpConfigurationEnvironment(UEBuildTarget InBuildTarget)
 		{
 			// Determine the C++ compile/link configuration based on the Unreal configuration.
@@ -771,14 +759,12 @@ namespace UnrealBuildTool
 			InBuildTarget.GlobalLinkEnvironment.Config.bCreateDebugInfo = true;
 		}
 
-		/**
-		 *	Whether this platform should create debug information or not
-		 *	
-		 *	@param	InPlatform			The UnrealTargetPlatform being built
-		 *	@param	InConfiguration		The UnrealTargetConfiguration being built
-		 *	
-		 *	@return	bool				true if debug info should be generated, false if not
-		 */
+		/// <summary>
+		/// Whether this platform should create debug information or not
+		/// </summary>
+		/// <param name="InPlatform">  The UnrealTargetPlatform being built</param>
+		/// <param name="InConfiguration"> The UnrealTargetConfiguration being built</param>
+		/// <returns>bool    true if debug info should be generated, false if not</returns>
 		public override bool ShouldCreateDebugInfo(UnrealTargetPlatform Platform, UnrealTargetConfiguration Configuration)
 		{
 			switch (Configuration)
@@ -794,9 +780,9 @@ namespace UnrealBuildTool
 			};
 		}
 
-		/**
-		 *	Return whether this platform has uniquely named binaries across multiple games
-		 */
+		/// <summary>
+		/// Return whether this platform has uniquely named binaries across multiple games
+		/// </summary>
 		public override bool HasUniqueBinaries()
 		{
 			// Windows applications have many shared binaries between games
diff --git a/Engine/Source/Programs/UnrealBuildTool/Windows/VCEnvironment.cs b/Engine/Source/Programs/UnrealBuildTool/Windows/VCEnvironment.cs
--- a/Engine/Source/Programs/UnrealBuildTool/Windows/VCEnvironment.cs
+++ b/Engine/Source/Programs/UnrealBuildTool/Windows/VCEnvironment.cs
@@ -44,9 +44,9 @@ namespace UnrealBuildTool
 			}
 		}
 
-		/**
-		 * Initializes environment variables required by toolchain. Different for 32 and 64 bit.
-		 */
+		/// <summary>
+		/// Initializes environment variables required by toolchain. Different for 32 and 64 bit.
+		/// </summary>
 		public static VCEnvironment SetEnvironment(CPPTargetPlatform Platform)
 		{
 			if (EnvVars != null && EnvVars.Platform == Platform)
@@ -141,7 +141,9 @@ namespace UnrealBuildTool
 			}
 		}
 
-		/** Check that an environment variable is the same as one derived from keys in the registry, and send a telemetry event if it's not */
+		/// <summary>
+		/// Check that an environment variable is the same as one derived from keys in the registry, and send a telemetry event if it's not
+		/// </summary>
 		private void CompareEnvironmentVariable(string VariableName, string RegistryValue)
 		{
 			string EnvironmentValue = Environment.GetEnvironmentVariable(VariableName);
@@ -214,7 +216,9 @@ namespace UnrealBuildTool
 			return FinalResult;
 		}
 
-		/** Gets the version of the Windows SDK libraries to use. As per VCVarsQueryRegistry.bat, this is the directory name that sorts last. */
+		/// <summary>
+		/// Gets the version of the Windows SDK libraries to use. As per VCVarsQueryRegistry.bat, this is the directory name that sorts last.
+		/// </summary>
 		static string FindWindowsSDKLibVersion(string WindowsSDKDir)
 		{
 			string WindowsSDKLibVersion;
@@ -339,13 +343,17 @@ namespace UnrealBuildTool
 			return LatestVersion;
 		}
 
-		/** Gets the path to the 32bit tool binaries. */
+		/// <summary>
+		/// Gets the path to the 32bit tool binaries.
+		/// </summary>
 		static string GetVSToolPath32Bit(string BaseVSToolPath)
 		{
 			return Path.Combine(BaseVSToolPath, "../../VC/bin");
 		}
 
-		/** Gets the path to the 64bit tool binaries. */
+		/// <summary>
+		/// Gets the path to the 64bit tool binaries.
+		/// </summary>
 		static string GetVSToolPath64Bit(string BaseVSToolPath)
 		{
 			// Use the native 64-bit compiler if present, otherwise use the amd64-on-x86 compiler. VS2012 Express only includes the latter.
@@ -358,7 +366,9 @@ namespace UnrealBuildTool
 			return Path.Combine(BaseVSToolPath, "../../VC/bin/x86_amd64");
 		}
 
-		/** Gets the path to the compiler. */
+		/// <summary>
+		/// Gets the path to the compiler.
+		/// </summary>
 		static string GetCompilerToolPath(string PlatformVSToolPath)
 		{
 			// If we were asked to use Clang, then we'll redirect the path to the compiler to the LLVM installation directory
@@ -406,7 +416,9 @@ namespace UnrealBuildTool
 			return new Version(ExeVersionInfo.FileMajorPart, ExeVersionInfo.FileMinorPart, ExeVersionInfo.FileBuildPart, ExeVersionInfo.FilePrivatePart);
 		}
 
-		/** Gets the path to the linker. */
+		/// <summary>
+		/// Gets the path to the linker.
+		/// </summary>
 		static string GetLinkerToolPath(string PlatformVSToolPath)
 		{
 			// If we were asked to use Clang, then we'll redirect the path to the compiler to the LLVM installation directory
@@ -428,7 +440,9 @@ namespace UnrealBuildTool
 			return Path.Combine(PlatformVSToolPath, "link.exe");
 		}
 
-		/** Gets the path to the library linker. */
+		/// <summary>
+		/// Gets the path to the library linker.
+		/// </summary>
 		static string GetLibraryLinkerToolPath(string PlatformVSToolPath)
 		{
 			// Regardless of the target, if we're linking on a 64 bit machine, we want to use the 64 bit linker (it's faster than the 32 bit linker)
@@ -436,7 +450,9 @@ namespace UnrealBuildTool
 			return Path.Combine(PlatformVSToolPath, "lib.exe");
 		}
 
-		/** Gets the path to the resource compiler's rc.exe for the specified platform. */
+		/// <summary>
+		/// Gets the path to the resource compiler's rc.exe for the specified platform.
+		/// </summary>
 		string GetResourceCompilerToolPath(CPPTargetPlatform Platform)
 		{
 			// 64 bit -- we can use the 32 bit version to target 64 bit on 32 bit OS.
@@ -467,7 +483,9 @@ namespace UnrealBuildTool
 			return Path.Combine(WindowsSDKDir, "bin/rc.exe");
 		}
 
-		/** Gets the path to MSBuild. */
+		/// <summary>
+		/// Gets the path to MSBuild.
+		/// </summary>
 		static string GetMSBuildToolPath()
 		{
 			string FrameworkDirectory = Environment.GetEnvironmentVariable("FrameworkDir");
@@ -480,7 +498,9 @@ namespace UnrealBuildTool
 			return Path.Combine(FrameworkDirectory, FrameworkVersion, "MSBuild.exe");
 		}
 
-		/** Gets the Visual C++ installation folder from the registry */
+		/// <summary>
+		/// Gets the Visual C++ installation folder from the registry
+		/// </summary>
 		static string FindVisualCppInstallationFolder(WindowsCompiler Version)
 		{
 			// Get the version string
@@ -515,7 +535,9 @@ namespace UnrealBuildTool
 			return InstallDir;
 		}
 
-		/** Finds the directory containing the Universal CRT installation. Returns null for Visual Studio versions before 2015 */
+		/// <summary>
+		/// Finds the directory containing the Universal CRT installation. Returns null for Visual Studio versions before 2015
+		/// </summary>
 		static string FindUniversalCRTInstallationFolder()
 		{
 			if (WindowsPlatform.Compiler != WindowsCompiler.VisualStudio2015)
@@ -532,7 +554,9 @@ namespace UnrealBuildTool
 			return (Value as string);
 		}
 
-		/** Gets the version of the Universal CRT to use. As per VCVarsQueryRegistry.bat, this is the directory name that sorts last. */
+		/// <summary>
+		/// Gets the version of the Universal CRT to use. As per VCVarsQueryRegistry.bat, this is the directory name that sorts last.
+		/// </summary>
 		static string FindUniversalCRTVersion(string UniversalCRTDir)
 		{
 			string UniversalCRTVersion = null;
@@ -551,7 +575,9 @@ namespace UnrealBuildTool
 			return UniversalCRTVersion;
 		}
 
-		/** Sets the Visual C++ INCLUDE environment variable */
+		/// <summary>
+		/// Sets the Visual C++ INCLUDE environment variable
+		/// </summary>
 		static List<string> GetVisualCppIncludePaths(string VisualCppDir, string UniversalCRTDir, string UniversalCRTVersion, string NetFXSDKDir, string WindowsSDKDir, string WindowsSDKLibVersion)
 		{
 			List<string> IncludePaths = new List<string>();
@@ -605,7 +631,9 @@ namespace UnrealBuildTool
 			return IncludePaths;
 		}
 
-		/** Sets the Visual C++ LIB environment variable */
+		/// <summary>
+		/// Sets the Visual C++ LIB environment variable
+		/// </summary>
 		static List<string> GetVisualCppLibraryPaths(string VisualCppDir, string UniversalCRTDir, string UniversalCRTVersion, string NetFXSDKDir, string WindowsSDKDir, string WindowsSDKLibVersion, CPPTargetPlatform Platform)
 		{
 			List<string> LibraryPaths = new List<string>();
diff --git a/Engine/Source/Programs/UnrealBuildTool/Windows/VCToolChain.cs b/Engine/Source/Programs/UnrealBuildTool/Windows/VCToolChain.cs
--- a/Engine/Source/Programs/UnrealBuildTool/Windows/VCToolChain.cs
+++ b/Engine/Source/Programs/UnrealBuildTool/Windows/VCToolChain.cs
@@ -1754,7 +1754,9 @@ namespace UnrealBuildTool
 			BuildProjectAction.ProducedItems.Add(PDBFile);
 		}
 
-		/** Gets the default include paths for the given platform. */
+		/// <summary>
+		/// Gets the default include paths for the given platform.
+		/// </summary>
 		public static string GetVCIncludePaths(CPPTargetPlatform Platform)
 		{
 			Debug.Assert(Platform == CPPTargetPlatform.Win32 || Platform == CPPTargetPlatform.Win64);
@@ -1784,7 +1786,9 @@ namespace UnrealBuildTool
 		}
 
 
-		/** Formats compiler output from Clang so that it is clickable in Visual Studio */
+		/// <summary>
+		/// Formats compiler output from Clang so that it is clickable in Visual Studio
+		/// </summary>
 		protected static void ClangCompilerOutputFormatter(object sender, DataReceivedEventArgs e)
 		{
 			var Output = e.Data;
diff --git a/Engine/Source/Programs/UnrealBuildTool/Windows/WindowsProjectGenerator.cs b/Engine/Source/Programs/UnrealBuildTool/Windows/WindowsProjectGenerator.cs
--- a/Engine/Source/Programs/UnrealBuildTool/Windows/WindowsProjectGenerator.cs
+++ b/Engine/Source/Programs/UnrealBuildTool/Windows/WindowsProjectGenerator.cs
@@ -8,14 +8,14 @@ using System.IO;
 
 namespace UnrealBuildTool
 {
-	/**
-	 *	Base class for platform-specific project generators 
-	 */
+	/// <summary>
+	/// Base class for platform-specific project generators
+	/// </summary>
 	public class WindowsProjectGenerator : UEPlatformProjectGenerator
 	{
-		/**
-		 *	Register the platform with the UEPlatformProjectGenerator class
-		 */
+		/// <summary>
+		/// Register the platform with the UEPlatformProjectGenerator class
+		/// </summary>
 		public override void RegisterPlatformProjectGenerator()
 		{
 			// Register this project generator for Win32 and Win64
@@ -28,14 +28,12 @@ namespace UnrealBuildTool
 		///
 		///	VisualStudio project generation functions
 		///	
-		/**
-		 *	Return the VisualStudio platform name for this build platform
-		 *	
-		 *	@param	InPlatform			The UnrealTargetPlatform being built
-		 *	@param	InConfiguration		The UnrealTargetConfiguration being built
-		 *	
-		 *	@return	string				The name of the platform that VisualStudio recognizes
-		 */
+		/// <summary>
+		/// Return the VisualStudio platform name for this build platform
+		/// </summary>
+		/// <param name="InPlatform">  The UnrealTargetPlatform being built</param>
+		/// <param name="InConfiguration"> The UnrealTargetConfiguration being built</param>
+		/// <returns>string    The name of the platform that VisualStudio recognizes</returns>
 		public override string GetVisualStudioPlatformName(UnrealTargetPlatform InPlatform, UnrealTargetConfiguration InConfiguration)
 		{
 			if (InPlatform == UnrealTargetPlatform.Win64)
