Add telemetry event for the time taken to perform a build, including the number of compile actions, link actions, pch actions, and list of the targets.

diff --git a/Engine/Source/Programs/UnrealBuildTool/System/ActionGraph.cs b/Engine/Source/Programs/UnrealBuildTool/System/ActionGraph.cs
--- a/Engine/Source/Programs/UnrealBuildTool/System/ActionGraph.cs
+++ b/Engine/Source/Programs/UnrealBuildTool/System/ActionGraph.cs
@@ -368,13 +368,15 @@ namespace UnrealBuildTool
 		/// <summary>
 		/// Executes a list of actions.
 		/// </summary>
-		public static bool ExecuteActions(List<Action> ActionsToExecute, out string ExecutorName)
+		public static bool ExecuteActions(List<Action> ActionsToExecute, out string ExecutorName, string TargetInfoForTelemetry)
 		{
 			bool Result = true;
 			bool bUsedXGE = false;
 			ExecutorName = "";
 			if (ActionsToExecute.Count > 0)
 			{
+				DateTime StartTime = DateTime.UtcNow;
+
 				if (BuildConfiguration.bAllowXGE || BuildConfiguration.bXGEExport)
 				{
 					Log.TraceInformation("{0} {1} action{2} to XGE",
@@ -443,6 +445,22 @@ namespace UnrealBuildTool
 					Result = LocalExecutor.ExecuteActions(ActionsToExecute);
 				}
 
+				if(Telemetry.IsAvailable() && !BuildConfiguration.bXGEExport)
+				{
+					Telemetry.SendEvent("BuildTime",
+						"ExecutorType", ExecutorName,
+						"OS", Environment.OSVersion.ToString(),
+						"MachineName", Environment.MachineName,
+						"NumLogicalCores", Environment.ProcessorCount.ToString(),
+						"NumPhysicalCores", Utils.GetPhysicalProcessorCount().ToString(),
+						"Targets", TargetInfoForTelemetry,
+						"NumActions", ActionsToExecute.Count.ToString(),
+						"NumCompileActions", ActionsToExecute.Count(x => x.ActionType == ActionType.Compile).ToString(),
+						"NumPchCompileActions", ActionsToExecute.Count(x => x.ActionType == ActionType.Compile && x.ProducedItems.Any(y => y.Reference.HasExtension(".pch") || y.Reference.HasExtension(".gch"))).ToString(),
+						"NumLinkActions", ActionsToExecute.Count(x => x.ActionType == ActionType.Link).ToString(),
+						"ElapsedTime", (DateTime.UtcNow - StartTime).TotalSeconds.ToString());
+				}
+
 				if (bUsedXGE && BuildConfiguration.bXGEExport)
 				{
 					// we exported xge here, we do not test build products
diff --git a/Engine/Source/Programs/UnrealBuildTool/System/LocalExecutor.cs b/Engine/Source/Programs/UnrealBuildTool/System/LocalExecutor.cs
--- a/Engine/Source/Programs/UnrealBuildTool/System/LocalExecutor.cs
+++ b/Engine/Source/Programs/UnrealBuildTool/System/LocalExecutor.cs
@@ -348,28 +348,8 @@ namespace UnrealBuildTool
 			const float LoopSleepTime = 0.1f;
 
 			// Use WMI to figure out physical cores, excluding hyper threading.
-			int NumCores = 0;
-			if (!Utils.IsRunningOnMono)
-			{
-				try
-				{
-					using (var Mos = new System.Management.ManagementObjectSearcher("Select * from Win32_Processor"))
-					{
-						var MosCollection = Mos.Get();
-						foreach (var Item in MosCollection)
-						{
-							NumCores += int.Parse(Item["NumberOfCores"].ToString());
-						}
-					}
-				}
-				catch (Exception Ex)
-				{
-					Log.TraceWarning("Unable to get the number of Cores: {0}", Ex.ToString());
-					Log.TraceWarning("Falling back to processor count.");
-				}
-			}
-			// On some systems this requires a hot fix to work so we fall back to using the (logical) processor count.
-			if (NumCores == 0)
+			int NumCores = Utils.GetPhysicalProcessorCount();
+			if (NumCores == -1)
 			{
 				NumCores = System.Environment.ProcessorCount;
 			}
diff --git a/Engine/Source/Programs/UnrealBuildTool/System/UnrealBuildTool.cs b/Engine/Source/Programs/UnrealBuildTool/System/UnrealBuildTool.cs
--- a/Engine/Source/Programs/UnrealBuildTool/System/UnrealBuildTool.cs
+++ b/Engine/Source/Programs/UnrealBuildTool/System/UnrealBuildTool.cs
@@ -53,7 +53,9 @@ namespace UnrealBuildTool
 		/// </summary>
 		static public double TotalOtherActionsTime = 0;
 
+		/// <summary>
 		/// How much time was spent scanning for include dependencies for outdated C++ files
+		/// </summary>
 		public static double TotalDeepIncludeScanTime = 0.0;
 
 		/// <summary>
@@ -379,8 +381,7 @@ namespace UnrealBuildTool
 		}
 
 		/// <summary>
-		/// Is this a valid configuration
-		/// Used primarily for Rocket vs non-Rocket
+		/// Is this a valid configuration. Used primarily for Rocket vs non-Rocket.
 		/// </summary>
 		/// <param name="InConfiguration"></param>
 		/// <returns>true if valid, false if not</returns>
@@ -401,8 +402,7 @@ namespace UnrealBuildTool
 		}
 
 		/// <summary>
-		/// Is this a valid platform
-		/// Used primarily for Rocket vs non-Rocket
+		/// Is this a valid platform. Used primarily for Rocket vs non-Rocket.
 		/// </summary>
 		/// <param name="InPlatform"></param>
 		/// <returns>true if valid, false if not</returns>
@@ -430,8 +430,7 @@ namespace UnrealBuildTool
 
 
 		/// <summary>
-		/// Is this a valid platform
-		/// Used primarily for Rocket vs non-Rocket
+		/// Is this a valid platform. Used primarily for Rocket vs non-Rocket.
 		/// </summary>
 		/// <param name="InPlatform"></param>
 		/// <returns>true if valid, false if not</returns>
@@ -2059,7 +2058,8 @@ namespace UnrealBuildTool
 							}
 
 							// Execute the actions.
-							bSuccess = ActionGraph.ExecuteActions(ActionsToExecute, out ExecutorName);
+							string TargetInfoForTelemetry = String.Join("|", Targets.Select(x => String.Format("{0} {1} {2}{3}", x.TargetName, x.Platform, x.Configuration, BuildConfiguration.bUseUnityBuild? "" : " NonUnity")));
+							bSuccess = ActionGraph.ExecuteActions(ActionsToExecute, out ExecutorName, TargetInfoForTelemetry);
 
 							// if the build succeeded, write the receipts and do any needed syncing
 							if (bSuccess)
diff --git a/Engine/Source/Programs/UnrealBuildTool/System/Utils.cs b/Engine/Source/Programs/UnrealBuildTool/System/Utils.cs
--- a/Engine/Source/Programs/UnrealBuildTool/System/Utils.cs
+++ b/Engine/Source/Programs/UnrealBuildTool/System/Utils.cs
@@ -919,6 +919,41 @@ namespace UnrealBuildTool
 		{
 			return a.CompareTo(b) > 0 ? a : b;
 		}
+
+		/// <summary>
+		/// Gets the number of physical cores, excluding hyper threading.
+		/// </summary>
+		/// <param name="NumCores">The number of physical cores, excluding hyper threading</param>
+		/// <returns>The number of physical cores, or -1 if it could not be obtained</returns>
+		public static int GetPhysicalProcessorCount()
+		{
+			// Can't use WMI queries on Mono; just fail.
+			if (Utils.IsRunningOnMono)
+			{
+				return -1;
+			}
+
+			// On some systems this requires a hot fix to work so catch any exceptions
+			try
+			{
+				int NumCores = 0;
+				using (var Mos = new System.Management.ManagementObjectSearcher("Select * from Win32_Processor"))
+				{
+					var MosCollection = Mos.Get();
+					foreach (var Item in MosCollection)
+					{
+						NumCores += int.Parse(Item["NumberOfCores"].ToString());
+					}
+				}
+				return NumCores;
+			}
+			catch (Exception Ex)
+			{
+				Log.TraceWarning("Unable to get the number of Cores: {0}", Ex.ToString());
+				Log.TraceWarning("Falling back to processor count.");
+				return -1;
+			}
+		}
 	}
 
 	/// <summary>
