nicer comments

diff --git a/Engine/Source/Programs/UnrealBuildTool/System/RulesCompiler.cs b/Engine/Source/Programs/UnrealBuildTool/System/RulesCompiler.cs
--- a/Engine/Source/Programs/UnrealBuildTool/System/RulesCompiler.cs
+++ b/Engine/Source/Programs/UnrealBuildTool/System/RulesCompiler.cs
@@ -241,17 +241,17 @@ namespace UnrealBuildTool
 		public bool bBuildLocallyWithSNDBS = false;
 
 		/// <summary>
-		/// List of modules with header files that our module's public headers needs access to, but we don't need to "import" or link against.
+        /// List of modules names (no path needed) with header files that our module's public headers needs access to, but we don't need to "import" or link against.
 		/// </summary>
 		public List<string> PublicIncludePathModuleNames = new List<string>();
 
 		/// <summary>
-		/// List of public dependency module names.  These are modules that are required by our public source files.
+        /// List of public dependency module names (no path needed) (automatically does the private/public include). These are modules that are required by our public source files.
 		/// </summary>
 		public List<string> PublicDependencyModuleNames = new List<string>();
 
 		/// <summary>
-		/// List of modules with header files that our module's private code files needs access to, but we don't need to "import" or link against.
+        /// List of modules name (no path needed) with header files that our module's private code files needs access to, but we don't need to "import" or link against.
 		/// </summary>
 		public List<string> PrivateIncludePathModuleNames = new List<string>();
 
@@ -262,38 +262,38 @@ namespace UnrealBuildTool
 		public List<string> PrivateDependencyModuleNames = new List<string>();
 
 		/// <summary>
-		/// List of module dependencies that should be treated as circular references.  This modules must have already been added to
-		/// either the public or private dependent module list.
+        /// Only for legacy reason, should not be used in new code. List of module dependencies that should be treated as circular references.  This modules must have already been added to
+        /// either the public or private dependent module list.
 		/// </summary>
 		public List<string> CircularlyReferencedDependentModules = new List<string>();
 
 		/// <summary>
-		/// System include paths.  These are public stable header file directories that are not checked when resolving header dependencies.
+        /// List of system/library include paths - typically used for External (third party) modules.  These are public stable header file directories that are not checked when resolving header dependencies.
 		/// </summary>
 		public List<string> PublicSystemIncludePaths = new List<string>();
 
 		/// <summary>
-		/// List of all paths to include files that are exposed to other modules
+        /// (This setting is currently not need as we discover all files from the 'Public' folder) List of all paths to include files that are exposed to other modules
 		/// </summary>
 		public List<string> PublicIncludePaths = new List<string>();
 
 		/// <summary>
-		/// List of all paths to this module's internal include files, not exposed to other modules
+        /// List of all paths to this module's internal include files, not exposed to other modules (at least one include to the 'Private' path, more if we want to avoid relative paths)
 		/// </summary>
 		public List<string> PrivateIncludePaths = new List<string>();
 
 		/// <summary>
-		/// List of library paths - typically used for External (third party) modules
+        /// List of system/library paths (directory of .lib files) - typically used for External (third party) modules
 		/// </summary>
 		public List<string> PublicLibraryPaths = new List<string>();
 
 		/// <summary>
-		/// List of addition libraries - typically used for External (third party) modules
+        /// List of additional libraries (names of the .lib files including extension) - typically used for External (third party) modules
 		/// </summary>
 		public List<string> PublicAdditionalLibraries = new List<string>();
 
 		/// <summary>
-		// List of frameworks
+        // List of XCode frameworks (iOS and MacOS)
 		/// </summary>
 		public List<string> PublicFrameworks = new List<string>();
 
