Toolchains are now instanced to allow storing project and target specific configuration details. Required for UBT to support multiple targets. Removed IUEToolChain, which wasn't being used directly.

diff --git a/Engine/Source/Programs/AutomationTool/Android/AndroidPlatform.Automation.cs b/Engine/Source/Programs/AutomationTool/Android/AndroidPlatform.Automation.cs
--- a/Engine/Source/Programs/AutomationTool/Android/AndroidPlatform.Automation.cs
+++ b/Engine/Source/Programs/AutomationTool/Android/AndroidPlatform.Automation.cs
@@ -112,11 +112,12 @@ public class AndroidPlatform : Platform
 
 	public override void Package(ProjectParams Params, DeploymentContext SC, int WorkingCL)
 	{
-		var Architectures = UnrealBuildTool.AndroidToolChain.GetAllArchitectures();
-		var GPUArchitectures = UnrealBuildTool.AndroidToolChain.GetAllGPUArchitectures();
+		AndroidToolChain ToolChain = new AndroidToolChain(Params.RawProjectPath);
+		var Architectures = ToolChain.GetAllArchitectures();
+		var GPUArchitectures = ToolChain.GetAllGPUArchitectures();
 		bool bMakeSeparateApks = UnrealBuildTool.UEDeployAndroid.ShouldMakeSeparateApks();
 
-		var Deploy = new UEDeployAndroid();
+		var Deploy = new UEDeployAndroid(Params.RawProjectPath);
 		bool bPackageDataInsideApk = Deploy.PackageDataInsideApk(false);
 
 		string BaseApkName = GetFinalApkName(Params, SC.StageExecutables[0], true, "", "");
@@ -176,7 +177,7 @@ public class AndroidPlatform : Platform
 				{
 					string CookFlavor = SC.FinalCookPlatform.IndexOf("_") > 0 ? SC.FinalCookPlatform.Substring(SC.FinalCookPlatform.IndexOf("_")) : "";
 					string SOName = GetSONameWithoutArchitecture(Params, SC.StageExecutables[0]);
-					Deploy.PrepForUATPackageOrDeploy(Params.ShortProjectName, SC.ProjectRoot, SOName, SC.LocalRoot + "/Engine", Params.Distribution, CookFlavor, false);
+					Deploy.PrepForUATPackageOrDeploy(Params.RawProjectPath, Params.ShortProjectName, SC.ProjectRoot, SOName, SC.LocalRoot + "/Engine", Params.Distribution, CookFlavor, false);
 				}
 
 			    // Create APK specific OBB in case we have a detached OBB.
@@ -317,10 +318,11 @@ public class AndroidPlatform : Platform
 			throw new AutomationException(ExitCode.Error_OnlyOneTargetConfigurationSupported, "Android is currently only able to package one target configuration at a time, but StageTargetConfigurations contained {0} configurations", SC.StageTargetConfigurations.Count);
 		}
 
-		var Architectures = UnrealBuildTool.AndroidToolChain.GetAllArchitectures();
-		var GPUArchitectures = UnrealBuildTool.AndroidToolChain.GetAllGPUArchitectures();
+		AndroidToolChain ToolChain = new AndroidToolChain(Params.RawProjectPath);
+		var Architectures = ToolChain.GetAllArchitectures();
+		var GPUArchitectures = ToolChain.GetAllGPUArchitectures();
 		bool bMakeSeparateApks = UnrealBuildTool.UEDeployAndroid.ShouldMakeSeparateApks();
-		bool bPackageDataInsideApk = new UnrealBuildTool.UEDeployAndroid().PackageDataInsideApk(false);
+		bool bPackageDataInsideApk = new UnrealBuildTool.UEDeployAndroid(Params.RawProjectPath).PackageDataInsideApk(false);
 
 		bool bAddedOBB = false;
 		foreach (string Architecture in Architectures)
@@ -510,12 +512,12 @@ public class AndroidPlatform : Platform
 		string ApkName = GetFinalApkName(Params, SC.StageExecutables[0], true, DeviceArchitecture, GPUArchitecture);
 
 		// make sure APK is up to date (this is fast if so)
-		var Deploy = new UEDeployAndroid();
+		var Deploy = new UEDeployAndroid(Params.RawProjectPath);
 		if (!Params.Prebuilt)
 		{
 			string CookFlavor = SC.FinalCookPlatform.IndexOf("_") > 0 ? SC.FinalCookPlatform.Substring(SC.FinalCookPlatform.IndexOf("_")) : "";
 			string SOName = GetSONameWithoutArchitecture(Params, SC.StageExecutables[0]);
-			Deploy.PrepForUATPackageOrDeploy(Params.ShortProjectName, SC.ProjectRoot, SOName, SC.LocalRoot + "/Engine", Params.Distribution, CookFlavor, true);
+			Deploy.PrepForUATPackageOrDeploy(Params.RawProjectPath, Params.ShortProjectName, SC.ProjectRoot, SOName, SC.LocalRoot + "/Engine", Params.Distribution, CookFlavor, true);
 		}
 
 		// now we can use the apk to get more info
@@ -976,7 +978,7 @@ public class AndroidPlatform : Platform
 			return "";
 		}
 
-		var AppArchitectures = AndroidToolChain.GetAllArchitectures();
+		var AppArchitectures = new AndroidToolChain(Params.RawProjectPath).GetAllArchitectures();
 
 		// ask the device
 		ProcessResult ABIResult = RunAdbCommand(Params, " shell getprop ro.product.cpu.abi", null, ERunOptions.AppMustExist);
@@ -1041,7 +1043,7 @@ public class AndroidPlatform : Platform
 			return "";
 		}
 
-		var AppGPUArchitectures = AndroidToolChain.GetAllGPUArchitectures();
+		var AppGPUArchitectures = new AndroidToolChain(Params.RawProjectPath).GetAllGPUArchitectures();
 
 		// get the device extensions
 		ProcessResult ExtensionsResult = RunAdbCommand(Params, "shell dumpsys SurfaceFlinger", null, ERunOptions.AppMustExist);
diff --git a/Engine/Source/Programs/AutomationTool/AutomationUtils/CommandletUtils.cs b/Engine/Source/Programs/AutomationTool/AutomationUtils/CommandletUtils.cs
--- a/Engine/Source/Programs/AutomationTool/AutomationUtils/CommandletUtils.cs
+++ b/Engine/Source/Programs/AutomationTool/AutomationUtils/CommandletUtils.cs
@@ -212,7 +212,7 @@ namespace AutomationTool
 				CommandUtils.LogWarning("Commandlet {0} failed to copy the local log file from {1} to {2}. The log file will be lost.", Commandlet, LocalDDCLogFile, DestDDCLogFile);
 			}
 
-            string ProjectStatsDirectory = CombinePaths(Path.GetDirectoryName(ProjectName.FullName), "Saved", "Stats");
+            string ProjectStatsDirectory = CombinePaths((ProjectName == null)? CombinePaths(CmdEnv.LocalRoot, "Engine") : Path.GetDirectoryName(ProjectName.FullName), "Saved", "Stats");
             if (Directory.Exists(ProjectStatsDirectory))
             {
                 string DestCookerStats = CmdEnv.LogFolder;
diff --git a/Engine/Source/Programs/AutomationTool/IOS/IOSPlatform.Automation.cs b/Engine/Source/Programs/AutomationTool/IOS/IOSPlatform.Automation.cs
--- a/Engine/Source/Programs/AutomationTool/IOS/IOSPlatform.Automation.cs
+++ b/Engine/Source/Programs/AutomationTool/IOS/IOSPlatform.Automation.cs
@@ -148,7 +148,8 @@ public class IOSPlatform : Platform
 			// copy in all of the artwork and plist
 			var DeployHandler = new UEDeployIOS();
 
-			DeployHandler.PrepForUATPackageOrDeploy(Params.ShortProjectName,
+			DeployHandler.PrepForUATPackageOrDeploy(Params.RawProjectPath,
+				Params.ShortProjectName,
 				Path.GetDirectoryName(Params.RawProjectPath.FullName),
 				CombinePaths(Path.GetDirectoryName(Params.ProjectGameExeFilename), SC.StageExecutables[0]),
 				CombinePaths(SC.LocalRoot, "Engine"),
diff --git a/Engine/Source/Programs/AutomationTool/Scripts/RocketBuild.Automation.cs b/Engine/Source/Programs/AutomationTool/Scripts/RocketBuild.Automation.cs
--- a/Engine/Source/Programs/AutomationTool/Scripts/RocketBuild.Automation.cs
+++ b/Engine/Source/Programs/AutomationTool/Scripts/RocketBuild.Automation.cs
@@ -427,7 +427,7 @@ namespace Rocket
 		public static void StripSymbols(UnrealTargetPlatform TargetPlatform, string[] SourceFileNames, string[] TargetFileNames)
 		{
 			IUEBuildPlatform Platform = UEBuildPlatform.GetBuildPlatform(TargetPlatform);
-			IUEToolChain ToolChain = UEToolChain.GetPlatformToolChain(Platform.GetCPPTargetPlatform(TargetPlatform));
+			UEToolChain ToolChain = Platform.CreateToolChain(Platform.GetCPPTargetPlatform(TargetPlatform), null);
 			for (int Idx = 0; Idx < SourceFileNames.Length; Idx++)
 			{
 				CommandUtils.CreateDirectory(Path.GetDirectoryName(TargetFileNames[Idx]));
diff --git a/Engine/Source/Programs/UnrealBuildTool/Android/AndroidToolChain.cs b/Engine/Source/Programs/UnrealBuildTool/Android/AndroidToolChain.cs
--- a/Engine/Source/Programs/UnrealBuildTool/Android/AndroidToolChain.cs
+++ b/Engine/Source/Programs/UnrealBuildTool/Android/AndroidToolChain.cs
@@ -12,15 +12,17 @@
 {
 	public class AndroidToolChain : UEToolChain
 	{
+		private FileReference ProjectFile;
+
 		// the number of the clang version being used to compile 
-		static private float ClangVersionFloat = 0;
+		private float ClangVersionFloat = 0;
 
 		// the list of architectures we will compile for
-		static private List<string> Arches = null;
+		private List<string> Arches = null;
 		// the list of GPU architectures we will compile for
-		static private List<string> GPUArchitectures = null;
+		private List<string> GPUArchitectures = null;
 		// a list of all architecture+GPUArchitecture names (-armv7-es2, etc)
-		static private List<string> AllComboNames = null;
+		private List<string> AllComboNames = null;
 
 		static private Dictionary<string, string[]> AllArchNames = new Dictionary<string, string[]> {
 			{ "-armv7", new string[] { "armv7", "armeabi-v7a", } }, 
@@ -43,10 +45,107 @@
 			{ "-x64",   new string[] { "OnlineSubsystemGooglePlay", } }, 
 		};
 
-		public static void ParseArchitectures()
+		public AndroidToolChain(FileReference InProjectFile) 
+			: base(CPPTargetPlatform.Android)
 		{
+			ProjectFile = InProjectFile;
+
+			string NDKPath = Environment.GetEnvironmentVariable("NDKROOT");
+
+			// don't register if we don't have an NDKROOT specified
+			if (String.IsNullOrEmpty(NDKPath))
+			{
+				throw new BuildException("NDKROOT is not specified; cannot use Android toolchain.");
+			}
+
+			NDKPath = NDKPath.Replace("\"", "");
+
+			string ClangVersion = "";
+			string GccVersion = "";
+
+			// prefer clang 3.6, but fall back if needed for now
+			if (Directory.Exists(Path.Combine(NDKPath, @"toolchains/llvm-3.6")))
+			{
+				ClangVersion = "3.6";
+				GccVersion = "4.9";
+			}
+			else if (Directory.Exists(Path.Combine(NDKPath, @"toolchains/llvm-3.5")))
+			{
+				ClangVersion = "3.5";
+				GccVersion = "4.9";
+			}
+			else if (Directory.Exists(Path.Combine(NDKPath, @"toolchains/llvm-3.3")))
+			{
+				ClangVersion = "3.3";
+				GccVersion = "4.8";
+			}
+			else if (Directory.Exists(Path.Combine(NDKPath, @"toolchains/llvm-3.1")))
+			{
+				ClangVersion = "3.1";
+				GccVersion = "4.6";
+			}
+			else
+			{
+				throw new BuildException("Cannot find supported Android toolchain");
+			}
+
+			ClangVersionFloat = float.Parse(ClangVersion, System.Globalization.CultureInfo.InvariantCulture);
+			// Console.WriteLine("Compiling with clang {0}", ClangVersionFloat);
+
+			string ArchitecturePath = "";
+			string ArchitecturePathWindows32 = @"prebuilt/windows";
+			string ArchitecturePathWindows64 = @"prebuilt/windows-x86_64";
+			string ArchitecturePathMac = @"prebuilt/darwin-x86_64";
+			string ExeExtension = ".exe";
+
+			if (Directory.Exists(Path.Combine(NDKPath, ArchitecturePathWindows64)))
+			{
+				Log.TraceVerbose("        Found Windows 64 bit versions of toolchain");
+				ArchitecturePath = ArchitecturePathWindows64;
+			}
+			else if (Directory.Exists(Path.Combine(NDKPath, ArchitecturePathWindows32)))
+			{
+				Log.TraceVerbose("        Found Windows 32 bit versions of toolchain");
+				ArchitecturePath = ArchitecturePathWindows32;
+			}
+			else if (Directory.Exists(Path.Combine(NDKPath, ArchitecturePathMac)))
+			{
+				Log.TraceVerbose("        Found Mac versions of toolchain");
+				ArchitecturePath = ArchitecturePathMac;
+				ExeExtension = "";
+			}
+			else
+			{
+				throw new BuildException("Couldn't find 32-bit or 64-bit versions of the Android toolchain");
+			}
+
+			// set up the path to our toolchains
+			ClangPath = Path.Combine(NDKPath, @"toolchains/llvm-" + ClangVersion, ArchitecturePath, @"bin/clang++" + ExeExtension);
+			ArPathArm = Path.Combine(NDKPath, @"toolchains/arm-linux-androideabi-" + GccVersion, ArchitecturePath, @"bin/arm-linux-androideabi-ar" + ExeExtension);		//@todo android: use llvm-ar.exe instead?
+			ArPathArm64 = Path.Combine(NDKPath, @"toolchains/aarch64-linux-android-" + GccVersion, ArchitecturePath, @"bin/aarch64-linux-android-ar" + ExeExtension);	//@todo android: use llvm-ar.exe instead?
+			ArPathx86 = Path.Combine(NDKPath, @"toolchains/x86-" + GccVersion, ArchitecturePath, @"bin/i686-linux-android-ar" + ExeExtension);							//@todo android: verify x86 toolchain
+			ArPathx64 = Path.Combine(NDKPath, @"toolchains/x86_64-" + GccVersion, ArchitecturePath, @"bin/x86_64-linux-android-ar" + ExeExtension);						//@todo android: verify x64 toolchain
+
+
+			// toolchain params
+			ToolchainParamsArm = " -target armv7-none-linux-androideabi" +
+								   " --sysroot=\"" + Path.Combine(NDKPath, "platforms", GetNdkApiLevel(), "arch-arm") + "\"" +
+								   " -gcc-toolchain \"" + Path.Combine(NDKPath, @"toolchains/arm-linux-androideabi-" + GccVersion, ArchitecturePath) + "\"";
+			ToolchainParamsArm64 = " -target aarch64-none-linux-android" +
+								   " --sysroot=\"" + Path.Combine(NDKPath, "platforms", GetNdkApiLevel(), "arch-arm64") + "\"" +
+								   " -gcc-toolchain \"" + Path.Combine(NDKPath, @"toolchains/aarch64-linux-android-" + GccVersion, ArchitecturePath) + "\"";
+			ToolchainParamsx86 = " -target i686-none-linux-android" +
+								   " --sysroot=\"" + Path.Combine(NDKPath, "platforms", GetNdkApiLevel(), "arch-x86") + "\"" +
+								   " -gcc-toolchain \"" + Path.Combine(NDKPath, @"toolchains/x86-" + GccVersion, ArchitecturePath) + "\"";
+			ToolchainParamsx64 = " -target x86_64-none-linux-android" +
+								   " --sysroot=\"" + Path.Combine(NDKPath, "platforms", GetNdkApiLevel(), "arch-x86_64") + "\"" +
+								   " -gcc-toolchain \"" + Path.Combine(NDKPath, @"toolchains\x86_64-" + GccVersion, ArchitecturePath) + "\"";
+		}
+
+		public void ParseArchitectures()
+		{
 			// look in ini settings for what platforms to compile for
-			ConfigCacheIni Ini = ConfigCacheIni.CreateConfigCacheIni(UnrealTargetPlatform.Android, "Engine", UnrealBuildTool.GetUProjectPath());
+			ConfigCacheIni Ini = ConfigCacheIni.CreateConfigCacheIni(UnrealTargetPlatform.Android, "Engine", DirectoryReference.FromFile(ProjectFile));
 			Arches = new List<string>();
 			bool bBuild = true;
 			if (Ini.GetBool("/Script/AndroidRuntimeSettings.AndroidRuntimeSettings", "bBuildForArmV7", out bBuild) && bBuild)
@@ -123,7 +222,7 @@
 			ParseArchitectures();
 		}
 
-		static public List<string> GetAllArchitectures()
+		public List<string> GetAllArchitectures()
 		{
 			if (Arches == null)
 			{
@@ -133,7 +232,7 @@
 			return Arches;
 		}
 
-		static public List<string> GetAllGPUArchitectures()
+		public List<string> GetAllGPUArchitectures()
 		{
 			if (GPUArchitectures == null)
 			{
@@ -143,9 +242,9 @@
 			return GPUArchitectures;
 		}
 
-		static public int GetNdkApiLevelInt(int MinNdk = 19)
+		public int GetNdkApiLevelInt(int MinNdk = 19)
 		{
-			string NDKVersion = AndroidToolChain.GetNdkApiLevel();
+			string NDKVersion = GetNdkApiLevel();
 			int NDKVersionInt = MinNdk;
 			if (NDKVersion.Contains("-"))
 			{
@@ -159,10 +258,10 @@
 			return NDKVersionInt;
 		}
 
-		static public string GetNdkApiLevel()
+		public string GetNdkApiLevel()
 		{
 			// ask the .ini system for what version to use
-			ConfigCacheIni Ini = ConfigCacheIni.CreateConfigCacheIni(UnrealTargetPlatform.Android, "Engine", UnrealBuildTool.GetUProjectPath());
+			ConfigCacheIni Ini = ConfigCacheIni.CreateConfigCacheIni(UnrealTargetPlatform.Android, "Engine", DirectoryReference.FromFile(ProjectFile));
 			string NDKLevel;
 			Ini.GetString("/Script/AndroidPlatformEditor.AndroidSDKSettings", "NDKAPILevel", out NDKLevel);
 
@@ -182,7 +281,7 @@
 			return NDKLevel;
 		}
 
-		static public string GetLargestApiLevel(string[] ApiLevels)
+		public string GetLargestApiLevel(string[] ApiLevels)
 		{
 			int LargestLevel = 0;
 			string LargestString = null;
@@ -214,106 +313,7 @@
 			return LargestString;
 		}
 
-		public override void RegisterToolChain()
+		string GetCLArguments_Global(CPPEnvironment CompileEnvironment, string Architecture)
 		{
-			string NDKPath = Environment.GetEnvironmentVariable("NDKROOT");
-
-			// don't register if we don't have an NDKROOT specified
-			if (String.IsNullOrEmpty(NDKPath))
-			{
-				return;
-			}
-
-			NDKPath = NDKPath.Replace("\"", "");
-
-			string ClangVersion = "";
-			string GccVersion = "";
-
-			// prefer clang 3.6, but fall back if needed for now
-			if (Directory.Exists(Path.Combine(NDKPath, @"toolchains/llvm-3.6")))
-			{
-				ClangVersion = "3.6";
-				GccVersion = "4.9";
-			}
-			else if (Directory.Exists(Path.Combine(NDKPath, @"toolchains/llvm-3.5")))
-			{
-				ClangVersion = "3.5";
-				GccVersion = "4.9";
-			}
-			else if (Directory.Exists(Path.Combine(NDKPath, @"toolchains/llvm-3.3")))
-			{
-				ClangVersion = "3.3";
-				GccVersion = "4.8";
-			}
-			else if (Directory.Exists(Path.Combine(NDKPath, @"toolchains/llvm-3.1")))
-			{
-				ClangVersion = "3.1";
-				GccVersion = "4.6";
-			}
-			else
-			{
-				return;
-			}
-
-			ClangVersionFloat = float.Parse(ClangVersion, System.Globalization.CultureInfo.InvariantCulture);
-			// Console.WriteLine("Compiling with clang {0}", ClangVersionFloat);
-
-			string ArchitecturePath = "";
-			string ArchitecturePathWindows32 = @"prebuilt/windows";
-			string ArchitecturePathWindows64 = @"prebuilt/windows-x86_64";
-			string ArchitecturePathMac = @"prebuilt/darwin-x86_64";
-			string ExeExtension = ".exe";
-
-			if (Directory.Exists(Path.Combine(NDKPath, ArchitecturePathWindows64)))
-			{
-				Log.TraceVerbose("        Found Windows 64 bit versions of toolchain");
-				ArchitecturePath = ArchitecturePathWindows64;
-			}
-			else if (Directory.Exists(Path.Combine(NDKPath, ArchitecturePathWindows32)))
-			{
-				Log.TraceVerbose("        Found Windows 32 bit versions of toolchain");
-				ArchitecturePath = ArchitecturePathWindows32;
-			}
-			else if (Directory.Exists(Path.Combine(NDKPath, ArchitecturePathMac)))
-			{
-				Log.TraceVerbose("        Found Mac versions of toolchain");
-				ArchitecturePath = ArchitecturePathMac;
-				ExeExtension = "";
-			}
-			else
-			{
-				Log.TraceVerbose("        Did not find 32 bit or 64 bit versions of toolchain");
-				return;
-			}
-
-			// set up the path to our toolchains
-			ClangPath = Path.Combine(NDKPath, @"toolchains/llvm-" + ClangVersion, ArchitecturePath, @"bin/clang++" + ExeExtension);
-			ArPathArm = Path.Combine(NDKPath, @"toolchains/arm-linux-androideabi-" + GccVersion, ArchitecturePath, @"bin/arm-linux-androideabi-ar" + ExeExtension);		//@todo android: use llvm-ar.exe instead?
-			ArPathArm64 = Path.Combine(NDKPath, @"toolchains/aarch64-linux-android-" + GccVersion, ArchitecturePath, @"bin/aarch64-linux-android-ar" + ExeExtension);	//@todo android: use llvm-ar.exe instead?
-			ArPathx86 = Path.Combine(NDKPath, @"toolchains/x86-" + GccVersion, ArchitecturePath, @"bin/i686-linux-android-ar" + ExeExtension);							//@todo android: verify x86 toolchain
-			ArPathx64 = Path.Combine(NDKPath, @"toolchains/x86_64-" + GccVersion, ArchitecturePath, @"bin/x86_64-linux-android-ar" + ExeExtension);						//@todo android: verify x64 toolchain
-
-
-			// toolchain params
-			ToolchainParamsArm = " -target armv7-none-linux-androideabi" +
-								   " --sysroot=\"" + Path.Combine(NDKPath, "platforms", GetNdkApiLevel(), "arch-arm") + "\"" +
-								   " -gcc-toolchain \"" + Path.Combine(NDKPath, @"toolchains/arm-linux-androideabi-" + GccVersion, ArchitecturePath) + "\"";
-			ToolchainParamsArm64 = " -target aarch64-none-linux-android" +
-								   " --sysroot=\"" + Path.Combine(NDKPath, "platforms", GetNdkApiLevel(), "arch-arm64") + "\"" +
-								   " -gcc-toolchain \"" + Path.Combine(NDKPath, @"toolchains/aarch64-linux-android-" + GccVersion, ArchitecturePath) + "\"";
-			ToolchainParamsx86 = " -target i686-none-linux-android" +
-								   " --sysroot=\"" + Path.Combine(NDKPath, "platforms", GetNdkApiLevel(), "arch-x86") + "\"" +
-								   " -gcc-toolchain \"" + Path.Combine(NDKPath, @"toolchains/x86-" + GccVersion, ArchitecturePath) + "\"";
-			ToolchainParamsx64 = " -target x86_64-none-linux-android" +
-								   " --sysroot=\"" + Path.Combine(NDKPath, "platforms", GetNdkApiLevel(), "arch-x86_64") + "\"" +
-								   " -gcc-toolchain \"" + Path.Combine(NDKPath, @"toolchains\x86_64-" + GccVersion, ArchitecturePath) + "\"";
-
-			// Register this tool chain
-			Log.TraceVerbose("        Registered for {0}", CPPTargetPlatform.Android.ToString());
-			UEToolChain.RegisterPlatformToolChain(CPPTargetPlatform.Android, this);
-		}
-
-		static string GetCLArguments_Global(CPPEnvironment CompileEnvironment, string Architecture)
-		{
 			string Result = "";
 
@@ -536,7 +536,7 @@
 			return Result;
 		}
 
-		static string GetLinkArguments(LinkEnvironment LinkEnvironment, string Architecture)
+		string GetLinkArguments(LinkEnvironment LinkEnvironment, string Architecture)
 		{
 			string Result = "";
 
@@ -674,7 +674,7 @@
 			return false;
 		}
 
-		static bool ShouldSkipLib(string Lib, string Arch, string GPUArchitecture)
+		bool ShouldSkipLib(string Lib, string Arch, string GPUArchitecture)
 		{
 			// reject any libs we outright don't want to link with
 			foreach (var LibName in LibrariesToSkip[Arch])
@@ -788,7 +788,7 @@
 			}
 		}
 
-		static void GenerateEmptyLinkFunctionsForRemovedModules(List<FileItem> SourceFiles, string ModuleName, DirectoryReference OutputDirectory)
+		void GenerateEmptyLinkFunctionsForRemovedModules(List<FileItem> SourceFiles, string ModuleName, DirectoryReference OutputDirectory)
 		{
 			// Only add to UELinkerFixups module
 			if (!ModuleName.Equals("Launch"))
@@ -1107,12 +1107,12 @@
 			return Path.Combine(Path.GetDirectoryName(Pathname), Path.GetFileNameWithoutExtension(Pathname) + Arch + GPUArchitecture + Path.GetExtension(Pathname));
 		}
 
-		static public string RemoveArchName(string Pathname)
+		public string RemoveArchName(string Pathname)
 		{
 			// remove all architecture names
-			foreach (string Arch in Arches)
+			foreach (string Arch in GetAllArchitectures())
 			{
-				foreach (string GPUArchitecture in GPUArchitectures)
+				foreach (string GPUArchitecture in GetAllGPUArchitectures())
 				{
 					Pathname = Path.Combine(Path.GetDirectoryName(Pathname), Path.GetFileName(Pathname).Replace(Arch + GPUArchitecture, ""));
 				}
@@ -1283,11 +1283,6 @@
 			}
 		}
 
-		public override UnrealTargetPlatform GetPlatform()
-		{
-			return UnrealTargetPlatform.Android;
-		}
-
 		public override void StripSymbols(string SourceFileName, string TargetFileName)
 		{
 			File.Copy(SourceFileName, TargetFileName, true);
diff --git a/Engine/Source/Programs/UnrealBuildTool/Android/UEBuildAndroid.cs b/Engine/Source/Programs/UnrealBuildTool/Android/UEBuildAndroid.cs
--- a/Engine/Source/Programs/UnrealBuildTool/Android/UEBuildAndroid.cs
+++ b/Engine/Source/Programs/UnrealBuildTool/Android/UEBuildAndroid.cs
@@ -366,8 +366,10 @@ namespace UnrealBuildTool
 			string NDKPath = Environment.GetEnvironmentVariable("NDKROOT");
 			NDKPath = NDKPath.Replace("\"", "");
 
+			AndroidToolChain ToolChain = new AndroidToolChain(InBuildTarget.ProjectFile);
+
 			string GccVersion = "4.6";
-			int NDKVersionInt = AndroidToolChain.GetNdkApiLevelInt();
+			int NDKVersionInt = ToolChain.GetNdkApiLevelInt();
 			if (Directory.Exists(Path.Combine(NDKPath, @"sources/cxx-stl/gnu-libstdc++/4.8")))
 			{
 				GccVersion = "4.8";
@@ -459,10 +461,12 @@ namespace UnrealBuildTool
 			};
 		}
 
-		public override List<FileReference> FinalizeBinaryPaths(FileReference BinaryName)
+		public override List<FileReference> FinalizeBinaryPaths(FileReference BinaryName, FileReference ProjectFile)
 		{
-			var Architectures = AndroidToolChain.GetAllArchitectures();
-			var GPUArchitectures = AndroidToolChain.GetAllGPUArchitectures();
+			AndroidToolChain ToolChain = new AndroidToolChain(ProjectFile);
+
+			var Architectures = ToolChain.GetAllArchitectures();
+			var GPUArchitectures = ToolChain.GetAllGPUArchitectures();
 
 			// make multiple output binaries
 			List<FileReference> AllBinaries = new List<FileReference>();
@@ -477,9 +481,14 @@ namespace UnrealBuildTool
 			return AllBinaries;
 		}
 
+		public override UEToolChain CreateToolChain(CPPTargetPlatform Platform, FileReference ProjectFile)
+		{
+			return new AndroidToolChain(ProjectFile);
+		}
+
 		public override bool TryCreateDeploymentHandler(FileReference ProjectFile, out UEBuildDeploy DeploymentHandler)
 		{
-			DeploymentHandler = new UEDeployAndroid();
+			DeploymentHandler = new UEDeployAndroid(ProjectFile);
 			return true;
 		}
 	}
diff --git a/Engine/Source/Programs/UnrealBuildTool/Android/UEDeployAndroid.cs b/Engine/Source/Programs/UnrealBuildTool/Android/UEDeployAndroid.cs
--- a/Engine/Source/Programs/UnrealBuildTool/Android/UEDeployAndroid.cs
+++ b/Engine/Source/Programs/UnrealBuildTool/Android/UEDeployAndroid.cs
@@ -18,6 +18,13 @@ namespace UnrealBuildTool
 		/// </summary>
 		private List<string> PossibleApiLevels = null;
 
+		private FileReference ProjectFile;
+
+		public UEDeployAndroid(FileReference InProjectFile)
+		{
+			ProjectFile = InProjectFile;
+		}
+
 		/// <summary>
 		/// Simple function to pipe output asynchronously
 		/// </summary>
@@ -52,7 +59,7 @@ namespace UnrealBuildTool
 			ConfigCacheIni config = null;
 			if (!ConfigCache.TryGetValue(baseIniName, out config))
 			{
-				config = ConfigCacheIni.CreateConfigCacheIni(UnrealTargetPlatform.Android, "Engine", UnrealBuildTool.GetUProjectPath());
+				config = ConfigCacheIni.CreateConfigCacheIni(UnrealTargetPlatform.Android, "Engine", DirectoryReference.FromFile(ProjectFile));
 				ConfigCache.Add(baseIniName, config);
 			}
 
@@ -60,7 +67,7 @@ namespace UnrealBuildTool
 		}
 
 		private string CachedSDKLevel = null;
-		private string GetSdkApiLevel()
+		private string GetSdkApiLevel(AndroidToolChain ToolChain)
 		{
 			if (CachedSDKLevel == null)
 			{
@@ -72,7 +79,7 @@ namespace UnrealBuildTool
 				// if we want to use whatever version the ndk uses, then use that
 				if (SDKLevel == "matchndk")
 				{
-					SDKLevel = AndroidToolChain.GetNdkApiLevel();
+					SDKLevel = ToolChain.GetNdkApiLevel();
 				}
 
 				// run a command and capture output
@@ -94,7 +101,7 @@ namespace UnrealBuildTool
 
 					if (PossibleApiLevels != null && PossibleApiLevels.Count > 0)
 					{
-						SDKLevel = AndroidToolChain.GetLargestApiLevel(PossibleApiLevels.ToArray());
+						SDKLevel = ToolChain.GetLargestApiLevel(PossibleApiLevels.ToArray());
 					}
 					else
 					{
@@ -529,7 +536,7 @@ namespace UnrealBuildTool
 			Utils.RunLocalProcessAndLogOutput(StartInfo);
 		}
 
-		private static void CopySTL(string UE4BuildPath, string UE4Arch, bool bForDistribution)
+		private static void CopySTL(AndroidToolChain ToolChain, string UE4BuildPath, string UE4Arch, bool bForDistribution)
 		{
 			string Arch = GetNDKArch(UE4Arch);
 
@@ -539,7 +546,7 @@ namespace UnrealBuildTool
 				GccVersion = "4.8";
 			}
 			// only use 4.9 if NDK version > 19
-			if (AndroidToolChain.GetNdkApiLevelInt() > 19 && Directory.Exists(Environment.ExpandEnvironmentVariables("%NDKROOT%/sources/cxx-stl/gnu-libstdc++/4.9")))
+			if (ToolChain.GetNdkApiLevelInt() > 19 && Directory.Exists(Environment.ExpandEnvironmentVariables("%NDKROOT%/sources/cxx-stl/gnu-libstdc++/4.9")))
 			{
 				GccVersion = "4.9";
 			}
@@ -668,7 +675,7 @@ namespace UnrealBuildTool
 			return true;
 		}
 
-		private void UpdateProjectProperties(string UE4BuildPath, string ProjectName)
+		private void UpdateProjectProperties(AndroidToolChain ToolChain, string UE4BuildPath, string ProjectName)
 		{
 			Log.TraceInformation("\n===={0}====UPDATING BUILD CONFIGURATION FILES====================================================", DateTime.Now.ToString());
 
@@ -742,7 +749,7 @@ namespace UnrealBuildTool
 
 			// now update the project for each library
 			string AndroidCommandPath = Environment.ExpandEnvironmentVariables("%ANDROID_HOME%/tools/android" + (Utils.IsRunningOnMono ? "" : ".bat"));
-			string UpdateCommandLine = "--silent update project --subprojects --name " + ProjectName + " --path . --target " + GetSdkApiLevel();
+			string UpdateCommandLine = "--silent update project --subprojects --name " + ProjectName + " --path . --target " + GetSdkApiLevel(ToolChain);
 			foreach (string Lib in LibsToBeAdded)
 			{
 				string LocalUpdateCommandLine = UpdateCommandLine + " --library " + Lib;
@@ -751,7 +758,7 @@ namespace UnrealBuildTool
 				// and later code needs each lib to have a build.xml
 				if (!File.Exists(Path.Combine(Lib, "build.xml")))
 				{
-					RunCommandLineProgramAndThrowOnError(UE4BuildPath, AndroidCommandPath, "--silent update lib-project --path " + Lib + " --target " + GetSdkApiLevel(), "");
+					RunCommandLineProgramAndThrowOnError(UE4BuildPath, AndroidCommandPath, "--silent update lib-project --path " + Lib + " --target " + GetSdkApiLevel(ToolChain), "");
 				}
 				RunCommandLineProgramAndThrowOnError(UE4BuildPath, AndroidCommandPath, LocalUpdateCommandLine, "Updating project.properties, local.properties, and build.xml...");
 			}
@@ -759,7 +766,7 @@ namespace UnrealBuildTool
 		}
 
 
-		private string GetAllBuildSettings(string BuildPath, bool bForDistribution, bool bMakeSeparateApks, bool bPackageDataInsideApk, bool bDisableVerifyOBBOnStartUp)
+		private string GetAllBuildSettings(AndroidToolChain ToolChain, string BuildPath, bool bForDistribution, bool bMakeSeparateApks, bool bPackageDataInsideApk, bool bDisableVerifyOBBOnStartUp)
 		{
 			// make the settings string - this will be char by char compared against last time
 			StringBuilder CurrentSettings = new StringBuilder();
@@ -767,7 +774,7 @@ namespace UnrealBuildTool
 			CurrentSettings.AppendLine(string.Format("ANDROID_HOME={0}", Environment.GetEnvironmentVariable("ANDROID_HOME")));
 			CurrentSettings.AppendLine(string.Format("ANT_HOME={0}", Environment.GetEnvironmentVariable("ANT_HOME")));
 			CurrentSettings.AppendLine(string.Format("JAVA_HOME={0}", Environment.GetEnvironmentVariable("JAVA_HOME")));
-			CurrentSettings.AppendLine(string.Format("SDKVersion={0}", GetSdkApiLevel()));
+			CurrentSettings.AppendLine(string.Format("SDKVersion={0}", GetSdkApiLevel(ToolChain)));
 			CurrentSettings.AppendLine(string.Format("bForDistribution={0}", bForDistribution));
 			CurrentSettings.AppendLine(string.Format("bMakeSeparateApks={0}", bMakeSeparateApks));
 			CurrentSettings.AppendLine(string.Format("bPackageDataInsideApk={0}", bPackageDataInsideApk));
@@ -801,13 +808,13 @@ namespace UnrealBuildTool
 				}
 			}
 
-			var Arches = AndroidToolChain.GetAllArchitectures();
+			var Arches = ToolChain.GetAllArchitectures();
 			foreach (string Arch in Arches)
 			{
 				CurrentSettings.AppendFormat("Arch={0}{1}", Arch, Environment.NewLine);
 			}
 
-			var GPUArchitectures = AndroidToolChain.GetAllGPUArchitectures();
+			var GPUArchitectures = ToolChain.GetAllGPUArchitectures();
 			foreach (string GPUArch in GPUArchitectures)
 			{
 				CurrentSettings.AppendFormat("GPUArch={0}{1}", GPUArch, Environment.NewLine);
@@ -816,11 +823,11 @@ namespace UnrealBuildTool
 			return CurrentSettings.ToString();
 		}
 
-		private bool CheckDependencies(string ProjectName, string ProjectDirectory, string UE4BuildFilesPath, string GameBuildFilesPath, string EngineDirectory, List<string> SettingsFiles,
+		private bool CheckDependencies(AndroidToolChain ToolChain, string ProjectName, string ProjectDirectory, string UE4BuildFilesPath, string GameBuildFilesPath, string EngineDirectory, List<string> SettingsFiles,
 			string CookFlavor, string OutputPath, string UE4BuildPath, bool bMakeSeparateApks, bool bPackageDataInsideApk)
 		{
-			var Arches = AndroidToolChain.GetAllArchitectures();
-			var GPUArchitectures = AndroidToolChain.GetAllGPUArchitectures();
+			var Arches = ToolChain.GetAllArchitectures();
+			var GPUArchitectures = ToolChain.GetAllGPUArchitectures();
 
 			// check all input files (.so, java files, .ini files, etc)
 			bool bAllInputsCurrent = true;
@@ -1368,7 +1375,7 @@ namespace UnrealBuildTool
 			return false;
 		}
 
-		private void MakeApk(string ProjectName, string ProjectDirectory, string OutputPath, string EngineDirectory, bool bForDistribution, string CookFlavor, bool bMakeSeparateApks, bool bIncrementalPackage, bool bDisallowPackagingDataInApk)
+		private void MakeApk(AndroidToolChain ToolChain, string ProjectName, string ProjectDirectory, string OutputPath, string EngineDirectory, bool bForDistribution, string CookFlavor, bool bMakeSeparateApks, bool bIncrementalPackage, bool bDisallowPackagingDataInApk)
 		{
 			Log.TraceInformation("\n===={0}====PREPARING TO MAKE APK=================================================================", DateTime.Now.ToString());
 
@@ -1460,7 +1467,7 @@ namespace UnrealBuildTool
 			bool bDisableVerifyOBBOnStartUp = DisableVerifyOBBOnStartUp();
 
 			// check to see if any "meta information" is newer than last time we build
-			string CurrentBuildSettings = GetAllBuildSettings(UE4BuildPath, bForDistribution, bMakeSeparateApks, bPackageDataInsideApk, bDisableVerifyOBBOnStartUp);
+			string CurrentBuildSettings = GetAllBuildSettings(ToolChain, UE4BuildPath, bForDistribution, bMakeSeparateApks, bPackageDataInsideApk, bDisableVerifyOBBOnStartUp);
 			string BuildSettingsCacheFile = Path.Combine(UE4BuildPath, "UEBuildSettings.txt");
 
 			// do we match previous build settings?
@@ -1500,13 +1507,13 @@ namespace UnrealBuildTool
 				JavaFiles.AddRange(from t in templates select t.SourceFile);
 				JavaFiles.AddRange(from t in templates select t.DestinationFile);
 
-				bBuildSettingsMatch = CheckDependencies(ProjectName, ProjectDirectory, UE4BuildFilesPath, GameBuildFilesPath,
+				bBuildSettingsMatch = CheckDependencies(ToolChain, ProjectName, ProjectDirectory, UE4BuildFilesPath, GameBuildFilesPath,
 					EngineDirectory, JavaFiles, CookFlavor, OutputPath, UE4BuildPath, bMakeSeparateApks, bPackageDataInsideApk);
 
 			}
 
-			var Arches = AndroidToolChain.GetAllArchitectures();
-			var GPUArchitectures = AndroidToolChain.GetAllGPUArchitectures();
+			var Arches = ToolChain.GetAllArchitectures();
+			var GPUArchitectures = ToolChain.GetAllGPUArchitectures();
 
 			IEnumerable<Tuple<string, string, string>> BuildList = null;
 
@@ -1629,7 +1636,7 @@ namespace UnrealBuildTool
 				File.WriteAllText(ManifestFile, Manifest);
 
 				// update metadata files (like project.properties, build.xml) if we are missing a build.xml or if we just overwrote project.properties with a bad version in it (from game/engine dir)
-				UpdateProjectProperties(UE4BuildPath, ProjectName);
+				UpdateProjectProperties(ToolChain, UE4BuildPath, ProjectName);
 
 				// Copy the generated .so file from the binaries directory to the jni folder
 				if (!File.Exists(SourceSOName))
@@ -1704,7 +1711,7 @@ namespace UnrealBuildTool
 
 				// after ndk-build is called, we can now copy in the stl .so (ndk-build deletes old files)
 				// copy libgnustl_shared.so to library (use 4.8 if possible, otherwise 4.6)
-				CopySTL(UE4BuildPath, Arch, bForDistribution);
+				CopySTL(ToolChain, UE4BuildPath, Arch, bForDistribution);
 				CopyGfxDebugger(UE4BuildPath, Arch);
 				CopyPluginLibs(EngineDirectory, UE4BuildPath, Arch);
 
@@ -1767,11 +1774,13 @@ namespace UnrealBuildTool
 
 		public override bool PrepTargetForDeployment(UEBuildTarget InTarget)
 		{
+			AndroidToolChain ToolChain = new AndroidToolChain(InTarget.ProjectFile); 
+
 			// we need to strip architecture from any of the output paths
-			string BaseSoName = AndroidToolChain.RemoveArchName(InTarget.OutputPaths[0].FullName);
+			string BaseSoName = ToolChain.RemoveArchName(InTarget.OutputPaths[0].FullName);
 
 			// make an apk at the end of compiling, so that we can run without packaging (debugger, cook on the fly, etc)
-			MakeApk(InTarget.AppName, InTarget.ProjectDirectory.FullName, BaseSoName, BuildConfiguration.RelativeEnginePath, bForDistribution: false, CookFlavor: "",
+			MakeApk(ToolChain, InTarget.AppName, InTarget.ProjectDirectory.FullName, BaseSoName, BuildConfiguration.RelativeEnginePath, bForDistribution: false, CookFlavor: "",
 				bMakeSeparateApks: ShouldMakeSeparateApks(), bIncrementalPackage: true, bDisallowPackagingDataInApk: false);
 
 			// if we made any non-standard .apk files, the generated debugger settings may be wrong
@@ -1800,11 +1809,12 @@ namespace UnrealBuildTool
 			// 			return bSeparateApks;
 		}
 
-		public override bool PrepForUATPackageOrDeploy(string ProjectName, string ProjectDirectory, string ExecutablePath, string EngineDirectory, bool bForDistribution, string CookFlavor, bool bIsDataDeploy)
+		public override bool PrepForUATPackageOrDeploy(FileReference ProjectFile, string ProjectName, string ProjectDirectory, string ExecutablePath, string EngineDirectory, bool bForDistribution, string CookFlavor, bool bIsDataDeploy)
 		{
 			// note that we cannot allow the data packaged into the APK if we are doing something like Launch On that will not make an obb
 			// file and instead pushes files directly via deploy
-			MakeApk(ProjectName, ProjectDirectory, ExecutablePath, EngineDirectory, bForDistribution: bForDistribution, CookFlavor: CookFlavor,
+			AndroidToolChain ToolChain = new AndroidToolChain(ProjectFile);
+			MakeApk(ToolChain, ProjectName, ProjectDirectory, ExecutablePath, EngineDirectory, bForDistribution: bForDistribution, CookFlavor: CookFlavor,
 				bMakeSeparateApks: ShouldMakeSeparateApks(), bIncrementalPackage: false, bDisallowPackagingDataInApk: bIsDataDeploy);
 			return true;
 		}
diff --git a/Engine/Source/Programs/UnrealBuildTool/Configuration/UEBuildBinary.cs b/Engine/Source/Programs/UnrealBuildTool/Configuration/UEBuildBinary.cs
--- a/Engine/Source/Programs/UnrealBuildTool/Configuration/UEBuildBinary.cs
+++ b/Engine/Source/Programs/UnrealBuildTool/Configuration/UEBuildBinary.cs
@@ -212,7 +212,7 @@ namespace UnrealBuildTool
 		/// <param name="CompileEnvironment">The environment to compile the binary in</param>
 		/// <param name="LinkEnvironment">The environment to link the binary in</param>
 		/// <returns></returns>
-		public abstract IEnumerable<FileItem> Build(IUEToolChain ToolChain, CPPEnvironment CompileEnvironment, LinkEnvironment LinkEnvironment);
+		public abstract IEnumerable<FileItem> Build(UEToolChain ToolChain, CPPEnvironment CompileEnvironment, LinkEnvironment LinkEnvironment);
 
 		/// <summary>
 		/// Called to allow the binary to modify the link environment of a different binary containing 
@@ -289,7 +289,7 @@ namespace UnrealBuildTool
 		/// </summary>
 		/// <param name="ToolChain">The platform toolchain</param>
 		/// <param name="BuildProducts">Mapping of produced build product to type</param>
-		public virtual void GetBuildProducts(IUEToolChain ToolChain, Dictionary<FileReference, BuildProductType> BuildProducts)
+		public virtual void GetBuildProducts(UEToolChain ToolChain, Dictionary<FileReference, BuildProductType> BuildProducts)
 		{
 			// Get the type of build products we're creating
 			BuildProductType Type = BuildProductType.RequiredResource;
@@ -332,7 +332,7 @@ namespace UnrealBuildTool
 		/// <param name="OutputFile">Build product to add</param>
 		/// <param name="DebugExtension">Extension for the matching debug file (may be null).</param>
 		/// <param name="Receipt">Receipt to add to</param>
-		static void AddBuildProductAndDebugFile(FileReference OutputFile, BuildProductType OutputType, string DebugExtension, Dictionary<FileReference, BuildProductType> BuildProducts, IUEToolChain ToolChain)
+		static void AddBuildProductAndDebugFile(FileReference OutputFile, BuildProductType OutputType, string DebugExtension, Dictionary<FileReference, BuildProductType> BuildProducts, UEToolChain ToolChain)
 		{
 			BuildProducts.Add(OutputFile, OutputType);
 
@@ -549,7 +549,7 @@ namespace UnrealBuildTool
 		/// <param name="CompileEnvironment">The environment to compile the binary in</param>
 		/// <param name="LinkEnvironment">The environment to link the binary in</param>
 		/// <returns></returns>
-		public override IEnumerable<FileItem> Build(IUEToolChain ToolChain, CPPEnvironment CompileEnvironment, LinkEnvironment LinkEnvironment)
+		public override IEnumerable<FileItem> Build(UEToolChain ToolChain, CPPEnvironment CompileEnvironment, LinkEnvironment LinkEnvironment)
 		{
 			// UnrealCodeAnalyzer produces output files only for a specific module.
 			if (BuildConfiguration.bRunUnrealCodeAnalyzer && !(Modules.Any(x => x.Name == BuildConfiguration.UCAModuleToAnalyze)))
@@ -637,7 +637,7 @@ namespace UnrealBuildTool
 		/// </summary>
 		/// <param name="ToolChain">The platform toolchain</param>
 		/// <param name="BuildProducts">Mapping of produced build product to type</param>
-		public override void GetBuildProducts(IUEToolChain ToolChain, Dictionary<FileReference, BuildProductType> BuildProducts)
+		public override void GetBuildProducts(UEToolChain ToolChain, Dictionary<FileReference, BuildProductType> BuildProducts)
 		{
 			base.GetBuildProducts(ToolChain, BuildProducts);
 
@@ -660,7 +660,7 @@ namespace UnrealBuildTool
 			return Config.OutputFilePath.FullName;
 		}
 
-		private LinkEnvironment SetupBinaryLinkEnvironment(IUEToolChain ToolChain, LinkEnvironment LinkEnvironment, CPPEnvironment CompileEnvironment)
+		private LinkEnvironment SetupBinaryLinkEnvironment(UEToolChain ToolChain, LinkEnvironment LinkEnvironment, CPPEnvironment CompileEnvironment)
 		{
 			var BinaryLinkEnvironment = LinkEnvironment.DeepCopy();
 			var LinkEnvironmentVisitedModules = new HashSet<UEBuildModule>();
@@ -745,7 +745,7 @@ namespace UnrealBuildTool
 			return BinaryLinkEnvironment;
 		}
 
-		private List<FileItem> SetupOutputFiles(IUEToolChain ToolChain, ref LinkEnvironment BinaryLinkEnvironment)
+		private List<FileItem> SetupOutputFiles(UEToolChain ToolChain, ref LinkEnvironment BinaryLinkEnvironment)
 		{
 			// Early exits first
 			if (ProjectFileGenerator.bGenerateProjectFiles)
@@ -862,7 +862,7 @@ namespace UnrealBuildTool
 		/// <param name="CompileEnvironment">The environment to compile the binary in</param>
 		/// <param name="LinkEnvironment">The environment to link the binary in</param>
 		/// <returns></returns>
-		public override IEnumerable<FileItem> Build(IUEToolChain ToolChain, CPPEnvironment CompileEnvironment, LinkEnvironment LinkEnvironment)
+		public override IEnumerable<FileItem> Build(UEToolChain ToolChain, CPPEnvironment CompileEnvironment, LinkEnvironment LinkEnvironment)
 		{
 			var ProjectCSharpEnviroment = new CSharpEnvironment();
 			if (LinkEnvironment.Config.Target.Configuration == CPPTargetConfiguration.Debug)
diff --git a/Engine/Source/Programs/UnrealBuildTool/Configuration/UEBuildDeploy.cs b/Engine/Source/Programs/UnrealBuildTool/Configuration/UEBuildDeploy.cs
--- a/Engine/Source/Programs/UnrealBuildTool/Configuration/UEBuildDeploy.cs
+++ b/Engine/Source/Programs/UnrealBuildTool/Configuration/UEBuildDeploy.cs
@@ -26,7 +26,7 @@ namespace UnrealBuildTool
 		/// <summary>
 		/// Prepare the target for deployment
 		/// </summary>
-		public virtual bool PrepForUATPackageOrDeploy(string ProjectName, string ProjectDirectory, string ExecutablePath, string EngineDirectory, bool bForDistribution, string CookFlavor, bool bIsDataDeploy)
+		public virtual bool PrepForUATPackageOrDeploy(FileReference ProjectFile, string ProjectName, string ProjectDirectory, string ExecutablePath, string EngineDirectory, bool bForDistribution, string CookFlavor, bool bIsDataDeploy)
 		{
 			return true;
 		}
diff --git a/Engine/Source/Programs/UnrealBuildTool/Configuration/UEBuildModule.cs b/Engine/Source/Programs/UnrealBuildTool/Configuration/UEBuildModule.cs
--- a/Engine/Source/Programs/UnrealBuildTool/Configuration/UEBuildModule.cs
+++ b/Engine/Source/Programs/UnrealBuildTool/Configuration/UEBuildModule.cs
@@ -618,7 +618,7 @@ namespace UnrealBuildTool
 		/// <summary>
 		/// Compiles the module, and returns a list of files output by the compiler.
 		/// </summary>
-		public abstract List<FileItem> Compile(IUEToolChain ToolChain, CPPEnvironment GlobalCompileEnvironment, CPPEnvironment CompileEnvironment);
+		public abstract List<FileItem> Compile(UEToolChain ToolChain, CPPEnvironment GlobalCompileEnvironment, CPPEnvironment CompileEnvironment);
 
 		// Object interface.
 		public override string ToString()
@@ -736,7 +736,7 @@ namespace UnrealBuildTool
 		}
 
 		// UEBuildModule interface.
-		public override List<FileItem> Compile(IUEToolChain ToolChain, CPPEnvironment GlobalCompileEnvironment, CPPEnvironment CompileEnvironment)
+		public override List<FileItem> Compile(UEToolChain ToolChain, CPPEnvironment GlobalCompileEnvironment, CPPEnvironment CompileEnvironment)
 		{
 			return new List<FileItem>();
 		}
@@ -959,7 +959,7 @@ namespace UnrealBuildTool
 		}
 
 		// UEBuildModule interface.
-		public override List<FileItem> Compile(IUEToolChain ToolChain, CPPEnvironment GlobalCompileEnvironment, CPPEnvironment CompileEnvironment)
+		public override List<FileItem> Compile(UEToolChain ToolChain, CPPEnvironment GlobalCompileEnvironment, CPPEnvironment CompileEnvironment)
 		{
 			var BuildPlatform = UEBuildPlatform.GetBuildPlatformForCPPTargetPlatform(CompileEnvironment.Config.Target.Platform);
 
@@ -1868,7 +1868,7 @@ namespace UnrealBuildTool
 		}
 
 		// UEBuildModule interface.
-		public override List<FileItem> Compile(IUEToolChain ToolChain, CPPEnvironment GlobalCompileEnvironment, CPPEnvironment CompileEnvironment)
+		public override List<FileItem> Compile(UEToolChain ToolChain, CPPEnvironment GlobalCompileEnvironment, CPPEnvironment CompileEnvironment)
 		{
 			var ModuleCLREnvironment = CompileEnvironment.DeepCopy();
 
@@ -1997,7 +1997,7 @@ namespace UnrealBuildTool
 		/// <param name="ModuleName">Name of the module this PCH is being generated for</param>
 		/// <param name="bAllowDLLExports">True if we should allow DLLEXPORT definitions for this PCH</param>
 		/// <returns>the compilation output result of the created pch.</returns>
-		public static CPPOutput GeneratePCHCreationAction(IUEToolChain ToolChain, UEBuildTarget Target, string PCHHeaderNameInCode, FileItem PrecompiledHeaderIncludeFilename, CPPEnvironment ProjectCPPEnvironment, DirectoryReference OutputDirectory, string ModuleName, bool bAllowDLLExports)
+		public static CPPOutput GeneratePCHCreationAction(UEToolChain ToolChain, UEBuildTarget Target, string PCHHeaderNameInCode, FileItem PrecompiledHeaderIncludeFilename, CPPEnvironment ProjectCPPEnvironment, DirectoryReference OutputDirectory, string ModuleName, bool bAllowDLLExports)
 		{
 			// Find the header file to be precompiled. Don't skip external headers
 			if (PrecompiledHeaderIncludeFilename.bExists)
diff --git a/Engine/Source/Programs/UnrealBuildTool/Configuration/UEBuildPlatform.cs b/Engine/Source/Programs/UnrealBuildTool/Configuration/UEBuildPlatform.cs
--- a/Engine/Source/Programs/UnrealBuildTool/Configuration/UEBuildPlatform.cs
+++ b/Engine/Source/Programs/UnrealBuildTool/Configuration/UEBuildPlatform.cs
@@ -161,7 +161,7 @@ namespace UnrealBuildTool
 		/// For platforms that need to output multiple files per binary (ie Android "fat" binaries)
 		/// this will emit multiple paths. By default, it simply makes an array from the input
 		/// </summary>
-		List<FileReference> FinalizeBinaryPaths(FileReference BinaryName);
+		List<FileReference> FinalizeBinaryPaths(FileReference BinaryName, FileReference ProjectFile);
 
 		/// <summary>
 		/// Setup the configuration environment for building
@@ -247,6 +247,14 @@ namespace UnrealBuildTool
 		/// </summary>
 		bool HasDefaultBuildConfig(UnrealTargetPlatform InPlatform, DirectoryReference InProjectPath);
 
+		/// <summary>
+		/// Creates a toolchain instance for the given platform. There should be a single toolchain instance per-target, as their may be
+		/// state data and configuration cached between calls.
+		/// </summary>
+		/// <param name="Platform">The platform to create a toolchain for</param>
+		/// <returns>New toolchain instance.</returns>
+		UEToolChain CreateToolChain(CPPTargetPlatform Platform, FileReference ProjectFile);
+
 		/// <summary>
 		/// Create a build deployment handler
 		/// </summary>
@@ -638,7 +646,7 @@ namespace UnrealBuildTool
 		/// For platforms that need to output multiple files per binary (ie Android "fat" binaries)
 		/// this will emit multiple paths. By default, it simply makes an array from the input
 		/// </summary>
-		public virtual List<FileReference> FinalizeBinaryPaths(FileReference BinaryName)
+		public virtual List<FileReference> FinalizeBinaryPaths(FileReference BinaryName, FileReference ProjectFile)
 		{
 			List<FileReference> TempList = new List<FileReference>() { BinaryName };
 			return TempList;
@@ -953,6 +961,15 @@ namespace UnrealBuildTool
 				BoolKeys, null, null);
 		}
 
+		/// <summary>
+		/// Creates a toolchain instance for the given platform. There should be a single toolchain instance per-target, as their may be
+		/// state data and configuration cached between calls.
+		/// </summary>
+		/// <param name="Platform">The platform to create a toolchain for</param>
+		/// <param name="ProjectFile">The path to the project file for this target</param>
+		/// <returns>New toolchain instance.</returns>
+		public abstract UEToolChain CreateToolChain(CPPTargetPlatform Platform, FileReference ProjectFile);
+
 		/// <summary>
 		/// Create a build deployment handler
 		/// </summary>
diff --git a/Engine/Source/Programs/UnrealBuildTool/Configuration/UEBuildTarget.cs b/Engine/Source/Programs/UnrealBuildTool/Configuration/UEBuildTarget.cs
--- a/Engine/Source/Programs/UnrealBuildTool/Configuration/UEBuildTarget.cs
+++ b/Engine/Source/Programs/UnrealBuildTool/Configuration/UEBuildTarget.cs
@@ -1177,7 +1177,7 @@ namespace UnrealBuildTool
 			{
 				OutputDirectory = UnrealBuildTool.EngineDirectory;
 			}
-			OutputPaths = MakeExecutablePaths(OutputDirectory, bCompileMonolithic ? TargetName : AppName, Platform, Configuration, Rules.UndecoratedConfiguration, bCompileMonolithic && ProjectFile != null, Rules.ExeBinariesSubFolder);
+			OutputPaths = MakeExecutablePaths(OutputDirectory, bCompileMonolithic ? TargetName : AppName, Platform, Configuration, Rules.UndecoratedConfiguration, bCompileMonolithic && ProjectFile != null, Rules.ExeBinariesSubFolder, ProjectFile);
 
 			// handle some special case defines (so build system can pass -DEFINE as normal instead of needing
 			// to know about special parameters)
@@ -1821,7 +1821,7 @@ namespace UnrealBuildTool
 		/// <summary>
 		/// Prepare all the receipts this target (all the .target and .modules files). See the VersionManifest class for an explanation of what these files are.
 		/// </summary>
-		public void PrepareReceipts(IUEToolChain ToolChain)
+		public void PrepareReceipts(UEToolChain ToolChain)
 		{
 			// Read the version file
 			BuildVersion Version;
@@ -1966,7 +1966,7 @@ namespace UnrealBuildTool
 		/// <summary>
 		/// Builds the target, appending list of output files and returns building result.
 		/// </summary>
-		public ECompilationResult Build(IUEToolChain TargetToolChain, out List<FileItem> OutputItems, out List<UHTModuleInfo> UObjectModules, out string EULAViolationWarning)
+		public ECompilationResult Build(UEToolChain TargetToolChain, out List<FileItem> OutputItems, out List<UHTModuleInfo> UObjectModules, out string EULAViolationWarning)
 		{
 			OutputItems = new List<FileItem>();
 			UObjectModules = new List<UHTModuleInfo>();
@@ -2212,7 +2212,7 @@ namespace UnrealBuildTool
 		/// Setup target before build. This method finds dependencies, sets up global environment etc.
 		/// </summary>
 		/// <returns>Special Rocket lib files that are build products.</returns>
-		public void PreBuildSetup(IUEToolChain TargetToolChain)
+		public void PreBuildSetup(UEToolChain TargetToolChain)
 		{
 			// Set up the global compile and link environment in GlobalCompileEnvironment and GlobalLinkEnvironment.
 			SetupGlobalEnvironment(TargetToolChain);
@@ -2931,7 +2931,7 @@ namespace UnrealBuildTool
 
 			// Get the output filenames
 			FileReference BaseBinaryPath = FileReference.Combine(OutputDirectory, MakeBinaryFileName(AppName + "-" + Module.Name, Platform, ModuleConfiguration, Rules.UndecoratedConfiguration, BinaryType));
-			List<FileReference> OutputFilePaths = UEBuildPlatform.GetBuildPlatform(Platform).FinalizeBinaryPaths(BaseBinaryPath);
+			List<FileReference> OutputFilePaths = UEBuildPlatform.GetBuildPlatform(Platform).FinalizeBinaryPaths(BaseBinaryPath, ProjectFile);
 
 			// Prepare the configuration object
 			UEBuildBinaryConfiguration Config = new UEBuildBinaryConfiguration(
@@ -3022,7 +3022,7 @@ namespace UnrealBuildTool
 		/// <param name="bIncludesGameModules">Whether this executable contains game modules</param>
 		/// <param name="ExeSubFolder">Subfolder for executables. May be null.</param>
 		/// <returns>List of executable paths for this target</returns>
-		public static List<FileReference> MakeExecutablePaths(DirectoryReference BaseDirectory, string BinaryName, UnrealTargetPlatform Platform, UnrealTargetConfiguration Configuration, UnrealTargetConfiguration UndecoratedConfiguration, bool bIncludesGameModules, string ExeSubFolder)
+		public static List<FileReference> MakeExecutablePaths(DirectoryReference BaseDirectory, string BinaryName, UnrealTargetPlatform Platform, UnrealTargetConfiguration Configuration, UnrealTargetConfiguration UndecoratedConfiguration, bool bIncludesGameModules, string ExeSubFolder, FileReference ProjectFile)
 		{
 			// Get the configuration for the executable. If we're building DebugGame, and this executable only contains engine modules, use the same name as development.
 			UnrealTargetConfiguration ExeConfiguration = Configuration;
@@ -3040,7 +3040,7 @@ namespace UnrealBuildTool
 			FileReference BinaryFile = FileReference.Combine(BinaryDirectory, MakeBinaryFileName(BinaryName, Platform, ExeConfiguration, UndecoratedConfiguration, UEBuildBinaryType.Executable));
 
 			// Allow the platform to customize the output path (and output several executables at once if necessary)
-			return UEBuildPlatform.GetBuildPlatform(Platform).FinalizeBinaryPaths(BinaryFile);
+			return UEBuildPlatform.GetBuildPlatform(Platform).FinalizeBinaryPaths(BinaryFile, ProjectFile);
 		}
 
 		/// <summary>
@@ -3191,7 +3191,7 @@ namespace UnrealBuildTool
 		/// <summary>
 		/// Sets up the global compile and link environment for the target.
 		/// </summary>
-		public virtual void SetupGlobalEnvironment(IUEToolChain ToolChain)
+		public virtual void SetupGlobalEnvironment(UEToolChain ToolChain)
 		{
 			var BuildPlatform = UEBuildPlatform.GetBuildPlatform(Platform);
 
diff --git a/Engine/Source/Programs/UnrealBuildTool/HTML5/HTML5ToolChain.cs b/Engine/Source/Programs/UnrealBuildTool/HTML5/HTML5ToolChain.cs
--- a/Engine/Source/Programs/UnrealBuildTool/HTML5/HTML5ToolChain.cs
+++ b/Engine/Source/Programs/UnrealBuildTool/HTML5/HTML5ToolChain.cs
@@ -15,18 +15,18 @@ namespace UnrealBuildTool
 		static bool bEnableTracing = false;
 
 		// cache the location of SDK tools
-		public override void RegisterToolChain()
+		public HTML5ToolChain()
+			: base(CPPTargetPlatform.HTML5)
 		{
-			if (HTML5SDKInfo.IsSDKInstalled())
+			if (!HTML5SDKInfo.IsSDKInstalled())
 			{
-				// set some environment variable we'll need
-				//Environment.SetEnvironmentVariable("EMCC_DEBUG", "cache");
-				Environment.SetEnvironmentVariable("EMCC_CORES", "8");
-				Environment.SetEnvironmentVariable("EMCC_OPTIMIZE_NORMALLY", "1");
-				// finally register the toolchain that is now ready to go
-				Log.TraceVerbose("        Registered for {0}", CPPTargetPlatform.HTML5.ToString());
-				UEToolChain.RegisterPlatformToolChain(CPPTargetPlatform.HTML5, this);
+				throw new BuildException("HTML5 SDK is not installed; cannot use toolchain.");
 			}
+
+			// set some environment variable we'll need
+			//Environment.SetEnvironmentVariable("EMCC_DEBUG", "cache");
+			Environment.SetEnvironmentVariable("EMCC_CORES", "8");
+			Environment.SetEnvironmentVariable("EMCC_OPTIMIZE_NORMALLY", "1");
 		}
 
 		public override void PreBuildSync()
@@ -544,10 +544,5 @@ namespace UnrealBuildTool
 				BuildProducts.Add(Binary.Config.OutputFilePath + ".symbols", BuildProductType.RequiredResource);
 			}
 		}
-
-		public override UnrealTargetPlatform GetPlatform()
-		{
-			return UnrealTargetPlatform.HTML5;
-		}
 	};
 }
diff --git a/Engine/Source/Programs/UnrealBuildTool/HTML5/UEBuildHTML5.cs b/Engine/Source/Programs/UnrealBuildTool/HTML5/UEBuildHTML5.cs
--- a/Engine/Source/Programs/UnrealBuildTool/HTML5/UEBuildHTML5.cs
+++ b/Engine/Source/Programs/UnrealBuildTool/HTML5/UEBuildHTML5.cs
@@ -369,6 +369,16 @@ namespace UnrealBuildTool
 			// 			InBuildTarget.ExtraModuleNames.Add("ES2RHI");
 		}
 
+		/// <summary>
+		/// Creates a toolchain instance for the given platform.
+		/// </summary>
+		/// <param name="Platform">The platform to create a toolchain for</param>
+		/// <returns>New toolchain instance.</returns>
+		public override UEToolChain CreateToolChain(CPPTargetPlatform Platform, FileReference ProjectFile)
+		{
+			return new HTML5ToolChain();
+		}
+
 		/// <summary>
 		/// Create a build deployment handler
 		/// </summary>
diff --git a/Engine/Source/Programs/UnrealBuildTool/IOS/IOSToolChain.cs b/Engine/Source/Programs/UnrealBuildTool/IOS/IOSToolChain.cs
--- a/Engine/Source/Programs/UnrealBuildTool/IOS/IOSToolChain.cs
+++ b/Engine/Source/Programs/UnrealBuildTool/IOS/IOSToolChain.cs
@@ -17,12 +17,11 @@ namespace UnrealBuildTool
 {
 	class IOSToolChain : AppleToolChain
 	{
-		public override void RegisterToolChain()
+		public IOSToolChain(FileReference InProjectFile)
+			: base(CPPTargetPlatform.IOS, UnrealTargetPlatform.Mac, InProjectFile)
 		{
-			RegisterRemoteToolChain(UnrealTargetPlatform.Mac, CPPTargetPlatform.IOS);
 		}
 
-
 		/***********************************************************************
 		 * NOTE:
 		 *  Do NOT change the defaults to set your values, instead you should set the environment variables
@@ -715,7 +714,7 @@ namespace UnrealBuildTool
 			else
 			{
 				bool bIsUE4Game = LinkEnvironment.Config.OutputFilePath.FullName.Contains("UE4Game");
-				FileReference ResponsePath = FileReference.Combine((!bIsUE4Game && UnrealBuildTool.HasUProjectFile()) ? UnrealBuildTool.GetUProjectPath() : UnrealBuildTool.EngineDirectory, BuildConfiguration.PlatformIntermediateFolder, "LinkFileList_" + LinkEnvironment.Config.OutputFilePath.GetFileNameWithoutExtension() + ".tmp");
+				FileReference ResponsePath = FileReference.Combine((!bIsUE4Game && ProjectFile != null) ? ProjectFile.Directory : UnrealBuildTool.EngineDirectory, BuildConfiguration.PlatformIntermediateFolder, "LinkFileList_" + LinkEnvironment.Config.OutputFilePath.GetFileNameWithoutExtension() + ".tmp");
 				if (!Utils.IsRunningOnMono && BuildHostPlatform.Current.Platform != UnrealTargetPlatform.Mac)
 				{
 					ResponseFile.Create(ResponsePath, InputFileNames);
@@ -1002,7 +1001,7 @@ namespace UnrealBuildTool
 					// generate the dummy project so signing works
 					if (AppName == "UE4Game" || AppName == "UE4Client" || Utils.IsFileUnderDirectory(Target.ProjectDirectory + "/" + AppName + ".uproject", Path.GetFullPath("../..")))
 					{
-						UnrealBuildTool.GenerateProjectFiles(new XcodeProjectFileGenerator(UnrealBuildTool.GetUProjectFile()), new string[] { "-platforms=IOS", "-NoIntellIsense", "-iosdeployonly", "-ignorejunk" });
+						UnrealBuildTool.GenerateProjectFiles(new XcodeProjectFileGenerator(Target.ProjectFile), new string[] { "-platforms=IOS", "-NoIntellIsense", "-iosdeployonly", "-ignorejunk" });
 						Project = Path.GetFullPath("../..") + "/UE4_IOS.xcodeproj";
 					}
 					else
@@ -1144,7 +1143,7 @@ namespace UnrealBuildTool
 					if (!bUseDangerouslyFastMode)
 					{
 						// generate the dummy project so signing works
-						UnrealBuildTool.GenerateProjectFiles(new XcodeProjectFileGenerator(UnrealBuildTool.GetUProjectFile()), new string[] { "-NoIntellisense", "-iosdeployonly", (UnrealBuildTool.HasUProjectFile() ? "-game" : "") });
+						UnrealBuildTool.GenerateProjectFiles(new XcodeProjectFileGenerator(Target.ProjectFile), new string[] { "-NoIntellisense", "-iosdeployonly", ((Target.ProjectFile != null) ? "-game" : "") });
 					}
 
 					// now that 
@@ -1335,11 +1334,6 @@ namespace UnrealBuildTool
 			}
 		}
 
-		public override UnrealTargetPlatform GetPlatform()
-		{
-			return UnrealTargetPlatform.IOS;
-		}
-
 		public static int RunExecutableAndWait(string ExeName, string ArgumentList, out string StdOutResults)
 		{
 			// Create the process
diff --git a/Engine/Source/Programs/UnrealBuildTool/IOS/UEBuildIOS.cs b/Engine/Source/Programs/UnrealBuildTool/IOS/UEBuildIOS.cs
--- a/Engine/Source/Programs/UnrealBuildTool/IOS/UEBuildIOS.cs
+++ b/Engine/Source/Programs/UnrealBuildTool/IOS/UEBuildIOS.cs
@@ -463,6 +463,16 @@ namespace UnrealBuildTool
 			}
 		}
 
+		/// <summary>
+		/// Creates a toolchain instance for the given platform.
+		/// </summary>
+		/// <param name="Platform">The platform to create a toolchain for</param>
+		/// <returns>New toolchain instance.</returns>
+		public override UEToolChain CreateToolChain(CPPTargetPlatform Platform, FileReference ProjectFile)
+		{
+			return new IOSToolChain(ProjectFile);
+		}
+
 		/// <summary>
 		/// Create a build deployment handler
 		/// </summary>
diff --git a/Engine/Source/Programs/UnrealBuildTool/IOS/UEDeployIOS.cs b/Engine/Source/Programs/UnrealBuildTool/IOS/UEDeployIOS.cs
--- a/Engine/Source/Programs/UnrealBuildTool/IOS/UEDeployIOS.cs
+++ b/Engine/Source/Programs/UnrealBuildTool/IOS/UEDeployIOS.cs
@@ -172,7 +172,7 @@ namespace UnrealBuildTool
 
 			// get the settings from the ini file
 			// plist replacements
-			ConfigCacheIni Ini = ConfigCacheIni.CreateConfigCacheIni(UnrealTargetPlatform.IOS, "Engine", UnrealBuildTool.GetUProjectPath());
+			ConfigCacheIni Ini = ConfigCacheIni.CreateConfigCacheIni(UnrealTargetPlatform.IOS, "Engine", bIsUE4Game? null : new DirectoryReference(ProjectDirectory));
 
 			// orientations
 			string SupportedOrientations = "";
@@ -463,7 +463,7 @@ namespace UnrealBuildTool
 			return bSkipDefaultPNGs;
 		}
 
-		public override bool PrepForUATPackageOrDeploy(string InProjectName, string InProjectDirectory, string InExecutablePath, string InEngineDir, bool bForDistribution, string CookFlavor, bool bIsDataDeploy)
+		public override bool PrepForUATPackageOrDeploy(FileReference ProjectFile, string InProjectName, string InProjectDirectory, string InExecutablePath, string InEngineDir, bool bForDistribution, string CookFlavor, bool bIsDataDeploy)
 		{
 			if (BuildHostPlatform.Current.Platform != UnrealTargetPlatform.Mac)
 			{
@@ -486,7 +486,7 @@ namespace UnrealBuildTool
 			Directory.CreateDirectory(BuildDirectory);
 
 			// create the entitlements file
-			WriteEntitlementsFile(Path.Combine(IntermediateDirectory, GameName + ".entitlements"));
+			WriteEntitlementsFile(Path.Combine(IntermediateDirectory, GameName + ".entitlements"), ProjectFile);
 
 			// delete some old files if they exist
 			if (Directory.Exists(AppDirectory + "/_CodeSignature"))
@@ -688,20 +688,20 @@ namespace UnrealBuildTool
 
 			if (BuildHostPlatform.Current.Platform == UnrealTargetPlatform.Mac && Environment.GetEnvironmentVariable("UBT_NO_POST_DEPLOY") != "true")
 			{
-				return PrepForUATPackageOrDeploy(GameName, ProjectDirectory, BuildPath + "/" + DecoratedGameName, "../../Engine", false, "", false);
+				return PrepForUATPackageOrDeploy(InTarget.ProjectFile, GameName, ProjectDirectory, BuildPath + "/" + DecoratedGameName, "../../Engine", false, "", false);
 			}
 			else
 			{
-				GeneratePList(ProjectDirectory, bIsUE4Game, GameName, UnrealBuildTool.GetUProjectFile() == null ? "" : Path.GetFileNameWithoutExtension(UnrealBuildTool.GetUProjectFile().FullName), "../../Engine", "");
+				GeneratePList(ProjectDirectory, bIsUE4Game, GameName, (InTarget.ProjectFile == null) ? "" : Path.GetFileNameWithoutExtension(InTarget.ProjectFile.FullName), "../../Engine", "");
 			}
 			return true;
 		}
 
-		private void WriteEntitlementsFile(string OutputFilename)
+		private void WriteEntitlementsFile(string OutputFilename, FileReference ProjectFile)
 		{
 			// get the settings from the ini file
 			// plist replacements
-			ConfigCacheIni Ini = ConfigCacheIni.CreateConfigCacheIni(UnrealTargetPlatform.IOS, "Engine", UnrealBuildTool.GetUProjectPath());
+			ConfigCacheIni Ini = ConfigCacheIni.CreateConfigCacheIni(UnrealTargetPlatform.IOS, "Engine", DirectoryReference.FromFile(ProjectFile));
 			bool bSupported = false;
 			Ini.GetBool("/Script/IOSRuntimeSettings.IOSRuntimeSettings", "bEnableCloudKitSupport", out bSupported);
 
diff --git a/Engine/Source/Programs/UnrealBuildTool/Linux/LinuxToolChain.cs b/Engine/Source/Programs/UnrealBuildTool/Linux/LinuxToolChain.cs
--- a/Engine/Source/Programs/UnrealBuildTool/Linux/LinuxToolChain.cs
+++ b/Engine/Source/Programs/UnrealBuildTool/Linux/LinuxToolChain.cs
@@ -11,6 +11,69 @@
 {
 	class LinuxToolChain : UEToolChain
 	{
+		public LinuxToolChain() 
+			: base(CPPTargetPlatform.Linux)
+		{
+			if (!CrossCompiling())
+			{
+				// use native linux toolchain
+				string[] ClangNames = { "clang++", "clang++-3.5", "clang++-3.3" };
+				foreach (var ClangName in ClangNames)
+				{
+					ClangPath = Which(ClangName);
+					if (!String.IsNullOrEmpty(ClangPath))
+					{
+						break;
+					}
+				}
+				GCCPath = Which("g++");
+				ArPath = Which("ar");
+				RanlibPath = Which("ranlib");
+				StripPath = Which("strip");
+
+				// if clang is available, zero out gcc (@todo: support runtime switching?)
+				if (!String.IsNullOrEmpty(ClangPath))
+				{
+					GCCPath = null;
+				}
+			}
+			else
+			{
+				// use cross linux toolchain if LINUX_ROOT is specified
+				BaseLinuxPath = Environment.GetEnvironmentVariable("LINUX_ROOT");
+
+				// don't register if we don't have an LINUX_ROOT specified
+				if (String.IsNullOrEmpty(BaseLinuxPath))
+					throw new BuildException("LINUX_ROOT environment variable is not set; cannot instantiate Linux toolchain");
+
+				BaseLinuxPath = BaseLinuxPath.Replace("\"", "");
+
+				// set up the path to our toolchains
+				GCCPath = "";
+				ClangPath = Path.Combine(BaseLinuxPath, @"bin\clang++.exe");
+				// ar and ranlib will be switched later to match the architecture
+				ArPath = "ar.exe";
+				RanlibPath = "ranlib.exe";
+				StripPath = "strip.exe";
+			}
+
+			if (!DetermineCompilerVersion())
+			{
+				throw new BuildException("Could not determine version of the compiler, not registering Linux toolchain.");
+			}
+
+			// refuse to use compilers that we know won't work
+			// disable that only if you are a dev and you know what you are doing
+			if (!UsingClang())
+			{
+				throw new BuildException("This version of the engine can only be compiled by clang - refusing to register the Linux toolchain.");
+			}
+			else if (CompilerVersionMajor == 3 && CompilerVersionMinor == 4)
+			{
+				throw new BuildException("clang 3.4.x is known to miscompile the engine - refusing to register the Linux toolchain.");
+			}
+		}
+
 		protected static bool CrossCompiling()
 		{
 			return BuildHostPlatform.Current.Platform != UnrealTargetPlatform.Linux;
@@ -134,78 +197,6 @@
 			return !String.IsNullOrEmpty(CompilerVersionString);
 		}
 
-		public override void RegisterToolChain()
-		{
-			if (!CrossCompiling())
-			{
-				// use native linux toolchain
-				string[] ClangNames = { "clang++", "clang++-3.5", "clang++-3.3" };
-				foreach (var ClangName in ClangNames)
-				{
-					ClangPath = Which(ClangName);
-					if (!String.IsNullOrEmpty(ClangPath))
-					{
-						break;
-					}
-				}
-				GCCPath = Which("g++");
-				ArPath = Which("ar");
-				RanlibPath = Which("ranlib");
-				StripPath = Which("strip");
-
-				// if clang is available, zero out gcc (@todo: support runtime switching?)
-				if (!String.IsNullOrEmpty(ClangPath))
-				{
-					GCCPath = null;
-				}
-			}
-			else
-			{
-				// use cross linux toolchain if LINUX_ROOT is specified
-				BaseLinuxPath = Environment.GetEnvironmentVariable("LINUX_ROOT");
-
-				// don't register if we don't have an LINUX_ROOT specified
-				if (String.IsNullOrEmpty(BaseLinuxPath))
-					return;
-
-				BaseLinuxPath = BaseLinuxPath.Replace("\"", "");
-
-				// set up the path to our toolchains
-				GCCPath = "";
-				ClangPath = Path.Combine(BaseLinuxPath, @"bin\clang++.exe");
-				// ar and ranlib will be switched later to match the architecture
-				ArPath = "ar.exe";
-				RanlibPath = "ranlib.exe";
-				StripPath = "strip.exe";
-			}
-
-			if (!DetermineCompilerVersion())
-			{
-				Console.WriteLine("\n*** Could not determine version of the compiler, not registering Linux toolchain.\n");
-				return;
-			}
-
-			// refuse to use compilers that we know won't work
-			// disable that only if you are a dev and you know what you are doing
-			if (!UsingClang())
-			{
-				Console.WriteLine("\n*** This version of the engine can only be compiled by clang - refusing to register the Linux toolchain.\n");
-				return;
-			}
-			else if (CompilerVersionMajor == 3 && CompilerVersionMinor == 4)
-			{
-				Console.WriteLine("\n*** clang 3.4.x is known to miscompile the engine - refusing to register the Linux toolchain.\n");
-				return;
-			}
-
-			// Register this tool chain for both Linux
-			if (BuildConfiguration.bPrintDebugInfo)
-			{
-				Console.WriteLine("        Registered for {0}", CPPTargetPlatform.Linux.ToString());
-			}
-			UEToolChain.RegisterPlatformToolChain(CPPTargetPlatform.Linux, this);
-		}
-
 		/// <summary>
 		/// Checks if compiler version matches the requirements
 		/// </summary>
@@ -1159,11 +1150,6 @@
 			return OutputFiles;
 		}
 
-		public override UnrealTargetPlatform GetPlatform()
-		{
-			return UnrealTargetPlatform.Linux;
-		}
-
 		public override void StripSymbols(string SourceFileName, string TargetFileName)
 		{
 			File.Copy(SourceFileName, TargetFileName, true);
diff --git a/Engine/Source/Programs/UnrealBuildTool/Linux/UEBuildLinux.cs b/Engine/Source/Programs/UnrealBuildTool/Linux/UEBuildLinux.cs
--- a/Engine/Source/Programs/UnrealBuildTool/Linux/UEBuildLinux.cs
+++ b/Engine/Source/Programs/UnrealBuildTool/Linux/UEBuildLinux.cs
@@ -394,6 +394,16 @@ namespace UnrealBuildTool
 		{
 		}
 
+		/// <summary>
+		/// Creates a toolchain instance for the given platform.
+		/// </summary>
+		/// <param name="Platform">The platform to create a toolchain for</param>
+		/// <returns>New toolchain instance.</returns>
+		public override UEToolChain CreateToolChain(CPPTargetPlatform Platform, FileReference ProjectFile)
+		{
+			return new LinuxToolChain();
+		}
+
 		/// <summary>
 		/// Create a build deployment handler
 		/// </summary>
diff --git a/Engine/Source/Programs/UnrealBuildTool/Mac/MacToolChain.cs b/Engine/Source/Programs/UnrealBuildTool/Mac/MacToolChain.cs
--- a/Engine/Source/Programs/UnrealBuildTool/Mac/MacToolChain.cs
+++ b/Engine/Source/Programs/UnrealBuildTool/Mac/MacToolChain.cs
@@ -14,9 +14,9 @@ namespace UnrealBuildTool
 {
 	class MacToolChain : AppleToolChain
 	{
-		public override void RegisterToolChain()
+		public MacToolChain(FileReference InProjectFile) 
+			: base(CPPTargetPlatform.Mac, UnrealTargetPlatform.Mac, InProjectFile)
 		{
-			RegisterRemoteToolChain(UnrealTargetPlatform.Mac, CPPTargetPlatform.Mac);
 		}
 
 		/***********************************************************************
@@ -1640,11 +1640,6 @@ namespace UnrealBuildTool
 			return OutputFiles;
 		}
 
-		public override UnrealTargetPlatform GetPlatform()
-		{
-			return UnrealTargetPlatform.Mac;
-		}
-
 		public override void StripSymbols(string SourceFileName, string TargetFileName)
 		{
 			SetupXcodePaths(false);
diff --git a/Engine/Source/Programs/UnrealBuildTool/Mac/UEBuildMac.cs b/Engine/Source/Programs/UnrealBuildTool/Mac/UEBuildMac.cs
--- a/Engine/Source/Programs/UnrealBuildTool/Mac/UEBuildMac.cs
+++ b/Engine/Source/Programs/UnrealBuildTool/Mac/UEBuildMac.cs
@@ -219,6 +219,16 @@ namespace UnrealBuildTool
 			return false;
 		}
 
+		/// <summary>
+		/// Creates a toolchain instance for the given platform.
+		/// </summary>
+		/// <param name="Platform">The platform to create a toolchain for</param>
+		/// <returns>New toolchain instance.</returns>
+		public override UEToolChain CreateToolChain(CPPTargetPlatform Platform, FileReference ProjectFile)
+		{
+			return new MacToolChain(ProjectFile);
+		}
+
 		/// <summary>
 		/// Create a build deployment handler
 		/// </summary>
diff --git a/Engine/Source/Programs/UnrealBuildTool/System/ExternalExecution.cs b/Engine/Source/Programs/UnrealBuildTool/System/ExternalExecution.cs
--- a/Engine/Source/Programs/UnrealBuildTool/System/ExternalExecution.cs
+++ b/Engine/Source/Programs/UnrealBuildTool/System/ExternalExecution.cs
@@ -654,7 +654,7 @@ namespace UnrealBuildTool
 		/// Builds and runs the header tool and touches the header directories.
 		/// Performs any early outs if headers need no changes, given the UObject modules, tool path, game name, and configuration
 		/// </summary>
-		public static bool ExecuteHeaderToolIfNecessary(IUEToolChain ToolChain, UEBuildTarget Target, CPPEnvironment GlobalCompileEnvironment, List<UHTModuleInfo> UObjectModules, FileReference ModuleInfoFileName, ref ECompilationResult UHTResult)
+		public static bool ExecuteHeaderToolIfNecessary(UEToolChain ToolChain, UEBuildTarget Target, CPPEnvironment GlobalCompileEnvironment, List<UHTModuleInfo> UObjectModules, FileReference ModuleInfoFileName, ref ECompilationResult UHTResult)
 		{
 			if (ProgressWriter.bWriteMarkup)
 			{
diff --git a/Engine/Source/Programs/UnrealBuildTool/System/Unity.cs b/Engine/Source/Programs/UnrealBuildTool/System/Unity.cs
--- a/Engine/Source/Programs/UnrealBuildTool/System/Unity.cs
+++ b/Engine/Source/Programs/UnrealBuildTool/System/Unity.cs
@@ -150,7 +150,7 @@ namespace UnrealBuildTool
 		/// <param name="BaseName">Base name to use for the Unity files</param>
 		/// <returns>The "unity" C++ files.</returns>
 		public static List<FileItem> GenerateUnityCPPs(
-			IUEToolChain ToolChain,
+			UEToolChain ToolChain,
 			UEBuildTarget Target,
 			List<FileItem> CPPFiles,
 			CPPEnvironment CompileEnvironment,
diff --git a/Engine/Source/Programs/UnrealBuildTool/System/UnrealBuildTool.cs b/Engine/Source/Programs/UnrealBuildTool/System/UnrealBuildTool.cs
--- a/Engine/Source/Programs/UnrealBuildTool/System/UnrealBuildTool.cs
+++ b/Engine/Source/Programs/UnrealBuildTool/System/UnrealBuildTool.cs
@@ -500,13 +500,7 @@ namespace UnrealBuildTool
 				{
 					if (CheckType.IsClass && !CheckType.IsAbstract)
 					{
-						if (Utils.ImplementsInterface<IUEToolChain>(CheckType))
-						{
-							Log.TraceVerbose("    Registering tool chain    : {0}", CheckType.ToString());
-							var TempInst = (UEToolChain)(UBTAssembly.CreateInstance(CheckType.FullName, true));
-							TempInst.RegisterToolChain();
-						}
-						else if (CheckType.IsSubclassOf(typeof(UEPlatformProjectGenerator)))
+						if (CheckType.IsSubclassOf(typeof(UEPlatformProjectGenerator)))
 						{
 							ProjectGeneratorList.Add(CheckType);
 						}
@@ -1396,7 +1390,7 @@ namespace UnrealBuildTool
 									BuildConfiguration.bFlushBuildDirOnRemoteMac = false;
 									var TargetDescs = UEBuildTarget.ParseTargetCommandLine(Arguments, GetUProjectFile());
 									UEBuildTarget CheckTarget = UEBuildTarget.CreateTarget(TargetDescs[0]);	// @todo ubtmake: This may not work in assembler only mode.  We don't want to be loading target rules assemblies here either.
-									IUEToolChain ToolChain = UEToolChain.GetPlatformToolChain(BuildPlatform.GetCPPTargetPlatform(TargetDescs[0].Platform));
+									UEToolChain ToolChain = BuildPlatform.CreateToolChain(BuildPlatform.GetCPPTargetPlatform(TargetDescs[0].Platform), UnrealBuildTool.GetUProjectFile());
 									CheckTarget.SetupGlobalEnvironment(ToolChain);
 									if ((CheckTarget.TargetType == TargetRules.TargetType.Game) ||
 										(CheckTarget.TargetType == TargetRules.TargetType.Server) ||
@@ -1624,7 +1618,7 @@ namespace UnrealBuildTool
 			string ExecutorName = "Unknown";
 			ECompilationResult BuildResult = ECompilationResult.Succeeded;
 
-			var ToolChain = UEToolChain.GetPlatformToolChain(BuildPlatform.GetCPPTargetPlatform(ResetPlatform));
+			var ToolChain = BuildPlatform.CreateToolChain(BuildPlatform.GetCPPTargetPlatform(ResetPlatform), UnrealBuildTool.GetUProjectFile());
 
 			string EULAViolationWarning = null;
 			Thread CPPIncludesThread = null;
@@ -2212,7 +2206,7 @@ namespace UnrealBuildTool
 
 			if (!ProjectFileGenerator.bGenerateProjectFiles && !UEBuildConfiguration.bGenerateManifest && bIsEditorTarget)
 			{
-				var EditorProcessFilenames = UEBuildTarget.MakeExecutablePaths(UnrealBuildTool.EngineDirectory, "UE4Editor", TargetDesc.Platform, TargetDesc.Configuration, UnrealTargetConfiguration.Development, false, null);
+				var EditorProcessFilenames = UEBuildTarget.MakeExecutablePaths(UnrealBuildTool.EngineDirectory, "UE4Editor", TargetDesc.Platform, TargetDesc.Configuration, UnrealTargetConfiguration.Development, false, null, null);
 				if (EditorProcessFilenames.Count != 1)
 				{
 					throw new BuildException("ShouldDoHotReload cannot handle multiple binaries returning from UEBuildTarget.MakeExecutablePaths");
diff --git a/Engine/Source/Programs/UnrealBuildTool/ToolChain/AppleToolChain.cs b/Engine/Source/Programs/UnrealBuildTool/ToolChain/AppleToolChain.cs
--- a/Engine/Source/Programs/UnrealBuildTool/ToolChain/AppleToolChain.cs
+++ b/Engine/Source/Programs/UnrealBuildTool/ToolChain/AppleToolChain.cs
@@ -11,6 +11,11 @@ namespace UnrealBuildTool
 {
 	public abstract class AppleToolChain : RemoteToolChain
 	{
+		public AppleToolChain(CPPTargetPlatform InCppPlatform, UnrealTargetPlatform InRemoteToolChainPlatform, FileReference InProjectFile)
+			: base(InCppPlatform, InRemoteToolChainPlatform, InProjectFile)
+		{
+		}
+
 		protected static void SelectXcode(ref string DeveloperDir, bool bVerbose)
 		{
 			string Reason = "hardcoded";
diff --git a/Engine/Source/Programs/UnrealBuildTool/ToolChain/RemoteToolChain.cs b/Engine/Source/Programs/UnrealBuildTool/ToolChain/RemoteToolChain.cs
--- a/Engine/Source/Programs/UnrealBuildTool/ToolChain/RemoteToolChain.cs
+++ b/Engine/Source/Programs/UnrealBuildTool/ToolChain/RemoteToolChain.cs
@@ -29,13 +29,13 @@ namespace UnrealBuildTool
 			SSHCommandFailed = 6,
 		};
 
-		protected void RegisterRemoteToolChain(UnrealTargetPlatform InPlatform, CPPTargetPlatform CPPPlatform)
-		{
-			RemoteToolChainPlatform = InPlatform;
+		protected readonly FileReference ProjectFile;
 
-			// Register this tool chain for the given platform
-			Log.TraceVerbose("        Registered for {0}", CPPPlatform.ToString());
-			UEToolChain.RegisterPlatformToolChain(CPPPlatform, this);
+		public RemoteToolChain(CPPTargetPlatform InCppPlatform, UnrealTargetPlatform InRemoteToolChainPlatform, FileReference InProjectFile)
+			: base(InCppPlatform)
+		{
+			RemoteToolChainPlatform = InRemoteToolChainPlatform;
+			ProjectFile = InProjectFile;
 		}
 
 		/// <summary>
@@ -151,7 +151,7 @@ namespace UnrealBuildTool
 			BranchDirectory = BranchDirectory.Replace("Engine\\Source\\", "");
 		}
 
-		private static string ResolveString(string Input, bool bIsPath)
+		private string ResolveString(string Input, bool bIsPath)
 		{
 			string Result = Input;
 
@@ -187,12 +187,12 @@ namespace UnrealBuildTool
 
 				if (Result.Contains("${PROJECT_ROOT}"))
 				{
-					if (!UnrealBuildTool.HasUProjectFile())
+					if (ProjectFile == null)
 					{
 						throw new BuildException("Configuration setting was using ${PROJECT_ROOT}, but there was no project specified");
 					}
 
-					string Temp = Result.Replace("${PROJECT_ROOT}", Path.GetFullPath(UnrealBuildTool.GetUProjectPath().FullName));
+					string Temp = Result.Replace("${PROJECT_ROOT}", ProjectFile.Directory.FullName);
 
 					// get the best version
 					Result = LookForSpecialFile(Temp);
@@ -243,7 +243,7 @@ namespace UnrealBuildTool
 		{
 			base.ParseProjectSettings();
 
-			ConfigCacheIni Ini = ConfigCacheIni.CreateConfigCacheIni(UnrealTargetPlatform.IOS, "Engine", UnrealBuildTool.GetUProjectPath());
+			ConfigCacheIni Ini = ConfigCacheIni.CreateConfigCacheIni(UnrealTargetPlatform.IOS, "Engine", DirectoryReference.FromFile(ProjectFile));
 			string ServerName = RemoteServerName;
 			if (Ini.GetString("/Script/IOSRuntimeSettings.IOSRuntimeSettings", "RemoteServerName", out ServerName) && !String.IsNullOrEmpty(ServerName))
 			{
@@ -309,7 +309,7 @@ namespace UnrealBuildTool
 
 		// Do any one-time, global initialization for the tool chain
 		static RemoteToolChainErrorCode InitializationErrorCode = RemoteToolChainErrorCode.NoError;
-		private static RemoteToolChainErrorCode InitializeRemoteExecution()
+		private RemoteToolChainErrorCode InitializeRemoteExecution()
 		{
 			if (bHasBeenInitialized)
 			{
@@ -422,11 +422,11 @@ namespace UnrealBuildTool
 						List<string> Locations = new List<string>();
 						Locations.Add(Path.Combine(Environment.GetFolderPath(Environment.SpecialFolder.ApplicationData), "Unreal Engine", "UnrealBuildTool"));
 						Locations.Add(Path.Combine(Environment.GetFolderPath(Environment.SpecialFolder.Personal), "Unreal Engine", "UnrealBuildTool"));
-						if (UnrealBuildTool.HasUProjectFile())
+						if (ProjectFile != null)
 						{
-							Locations.Add(Path.Combine(UnrealBuildTool.GetUProjectPath().FullName, "Build", "NotForLicensees"));
-							Locations.Add(Path.Combine(UnrealBuildTool.GetUProjectPath().FullName, "Build", "NoRedist"));
-							Locations.Add(Path.Combine(UnrealBuildTool.GetUProjectPath().FullName, "Build"));
+							Locations.Add(Path.Combine(ProjectFile.Directory.FullName, "Build", "NotForLicensees"));
+							Locations.Add(Path.Combine(ProjectFile.Directory.FullName, "Build", "NoRedist"));
+							Locations.Add(Path.Combine(ProjectFile.Directory.FullName, "Build"));
 						}
 						Locations.Add(Path.Combine(BuildConfiguration.RelativeEnginePath, "Build", "NotForLicensees"));
 						Locations.Add(Path.Combine(BuildConfiguration.RelativeEnginePath, "Build", "NoRedist"));
diff --git a/Engine/Source/Programs/UnrealBuildTool/ToolChain/UEToolChain.cs b/Engine/Source/Programs/UnrealBuildTool/ToolChain/UEToolChain.cs
--- a/Engine/Source/Programs/UnrealBuildTool/ToolChain/UEToolChain.cs
+++ b/Engine/Source/Programs/UnrealBuildTool/ToolChain/UEToolChain.cs
@@ -9,84 +9,15 @@ using Microsoft.Win32;
 
 namespace UnrealBuildTool
 {
-	public interface IUEToolChain
+	public abstract class UEToolChain
 	{
-		void RegisterToolChain();
+		public readonly CPPTargetPlatform CppPlatform;
 
-		CPPOutput CompileCPPFiles(UEBuildTarget Target, CPPEnvironment CompileEnvironment, List<FileItem> SourceFiles, string ModuleName);
-
-		CPPOutput CompileRCFiles(UEBuildTarget Target, CPPEnvironment Environment, List<FileItem> RCFiles);
-
-		FileItem[] LinkAllFiles(LinkEnvironment LinkEnvironment, bool bBuildImportLibraryOnly);
-
-		void CompileCSharpProject(CSharpEnvironment CompileEnvironment, FileReference ProjectFileName, FileReference DestinationFile);
-
-		/// <summary>
-		/// Converts the passed in path from UBT host to compiler native format.
-		/// </summary>
-		String ConvertPath(String OriginalPath);
-
-		/// <summary>
-		/// Called immediately after UnrealHeaderTool is executed to generated code for all UObjects modules.  Only is called if UnrealHeaderTool was actually run in this session.
-		/// </summary>
-		/// <param name="Manifest">List of UObject modules we generated code for.</param>
-		void PostCodeGeneration(UHTManifest Manifest);
-
-		void PreBuildSync();
-
-		void PostBuildSync(UEBuildTarget Target);
-
-		ICollection<FileItem> PostBuild(FileItem Executable, LinkEnvironment ExecutableLinkEnvironment);
-
-		void SetUpGlobalEnvironment();
-
-		void ModifyBuildProducts(UEBuildBinary Binary, Dictionary<FileReference, BuildProductType> BuildProducts);
-
-		bool ShouldAddDebugFileToReceipt(FileReference OutputFile, BuildProductType OutputType);
-
-		void SetupBundleDependencies(List<UEBuildBinary> Binaries, string GameName);
-
-		void FixBundleBinariesPaths(UEBuildTarget Target, List<UEBuildBinary> Binaries);
-
-		UnrealTargetPlatform GetPlatform();
-
-		void StripSymbols(string SourceFileName, string TargetFileName);
-	}
-
-	public abstract class UEToolChain : IUEToolChain
-	{
-		static Dictionary<CPPTargetPlatform, IUEToolChain> CPPToolChainDictionary = new Dictionary<CPPTargetPlatform, IUEToolChain>();
-
-		public static void RegisterPlatformToolChain(CPPTargetPlatform InPlatform, IUEToolChain InToolChain)
-		{
-			if (CPPToolChainDictionary.ContainsKey(InPlatform) == true)
-			{
-				Log.TraceInformation("RegisterPlatformToolChain Warning: Registering tool chain {0} for {1} when it is already set to {2}",
-					InToolChain.ToString(), InPlatform.ToString(), CPPToolChainDictionary[InPlatform].ToString());
-				CPPToolChainDictionary[InPlatform] = InToolChain;
-			}
-			else
-			{
-				CPPToolChainDictionary.Add(InPlatform, InToolChain);
-			}
-		}
-
-		public static void UnregisterPlatformToolChain(CPPTargetPlatform InPlatform)
+		public UEToolChain(CPPTargetPlatform InCppPlatform)
 		{
-			CPPToolChainDictionary.Remove(InPlatform);
-		}
-
-		public static IUEToolChain GetPlatformToolChain(CPPTargetPlatform InPlatform)
-		{
-			if (CPPToolChainDictionary.ContainsKey(InPlatform) == true)
-			{
-				return CPPToolChainDictionary[InPlatform];
-			}
-			throw new BuildException("GetPlatformToolChain: No tool chain found for {0}", InPlatform.ToString());
+			CppPlatform = InCppPlatform;
 		}
 
-		public abstract void RegisterToolChain();
-
 		public abstract CPPOutput CompileCPPFiles(UEBuildTarget Target, CPPEnvironment CompileEnvironment, List<FileItem> SourceFiles, string ModuleName);
 
 		public virtual CPPOutput CompileRCFiles(UEBuildTarget Target, CPPEnvironment Environment, List<FileItem> RCFiles)
@@ -219,14 +150,9 @@ namespace UnrealBuildTool
 
 		}
 
-		public virtual UnrealTargetPlatform GetPlatform()
-		{
-			return UnrealTargetPlatform.Unknown;
-		}
-
 		public virtual void StripSymbols(string SourceFileName, string TargetFileName)
 		{
-			Log.TraceWarning("StripSymbols() has not been implemented for {0}; copying files", GetPlatform().ToString());
+			Log.TraceWarning("StripSymbols() has not been implemented for {0}; copying files", CppPlatform.ToString());
 			File.Copy(SourceFileName, TargetFileName, true);
 		}
 	};
diff --git a/Engine/Source/Programs/UnrealBuildTool/UWP/UEBuildUWP.cs b/Engine/Source/Programs/UnrealBuildTool/UWP/UEBuildUWP.cs
--- a/Engine/Source/Programs/UnrealBuildTool/UWP/UEBuildUWP.cs
+++ b/Engine/Source/Programs/UnrealBuildTool/UWP/UEBuildUWP.cs
@@ -476,6 +476,16 @@ namespace UnrealBuildTool
 			return false;
 		}
 
+		/// <summary>
+		/// Creates a toolchain instance for the given platform.
+		/// </summary>
+		/// <param name="Platform">The platform to create a toolchain for</param>
+		/// <returns>New toolchain instance.</returns>
+		public override UEToolChain CreateToolChain(CPPTargetPlatform Platform, FileReference ProjectFile)
+		{
+			return new UWPToolChain();
+		}
+
 		/// <summary>
 		/// Create a build deployment handler
 		/// </summary>
diff --git a/Engine/Source/Programs/UnrealBuildTool/UWP/UWPToolChain.cs b/Engine/Source/Programs/UnrealBuildTool/UWP/UWPToolChain.cs
--- a/Engine/Source/Programs/UnrealBuildTool/UWP/UWPToolChain.cs
+++ b/Engine/Source/Programs/UnrealBuildTool/UWP/UWPToolChain.cs
@@ -13,14 +13,9 @@ namespace UnrealBuildTool
 {
 	public class UWPToolChain : UEToolChain
 	{
-		public override void RegisterToolChain()
+		public UWPToolChain() 
+			: base(CPPTargetPlatform.UWP)
 		{
-			if (UWPPlatform.bEnableUWPSupport)
-			{
-				// Register this tool chain for UWP
-				Log.TraceVerbose("        Registered for {0}", CPPTargetPlatform.UWP.ToString());
-				UEToolChain.RegisterPlatformToolChain(CPPTargetPlatform.UWP, this);
-			}
 		}
 
 		static void AppendCLArguments_Global(CPPEnvironment CompileEnvironment, VCEnvironment EnvVars, StringBuilder Arguments)
diff --git a/Engine/Source/Programs/UnrealBuildTool/WinRT/UEBuildWinRT.cs b/Engine/Source/Programs/UnrealBuildTool/WinRT/UEBuildWinRT.cs
--- a/Engine/Source/Programs/UnrealBuildTool/WinRT/UEBuildWinRT.cs
+++ b/Engine/Source/Programs/UnrealBuildTool/WinRT/UEBuildWinRT.cs
@@ -635,6 +635,16 @@ namespace UnrealBuildTool
 			return WinRTPlatform.bCompileWinRT;
 		}
 
+		/// <summary>
+		/// Creates a toolchain instance for the given platform.
+		/// </summary>
+		/// <param name="Platform">The platform to create a toolchain for</param>
+		/// <returns>New toolchain instance.</returns>
+		public override UEToolChain CreateToolChain(CPPTargetPlatform Platform, FileReference ProjectFile)
+		{
+			return new WinRTToolChain(Platform);
+		}
+
 		/// <summary>
 		/// Create a build deployment handler
 		/// </summary>
diff --git a/Engine/Source/Programs/UnrealBuildTool/WinRT/WinRTToolChain.cs b/Engine/Source/Programs/UnrealBuildTool/WinRT/WinRTToolChain.cs
--- a/Engine/Source/Programs/UnrealBuildTool/WinRT/WinRTToolChain.cs
+++ b/Engine/Source/Programs/UnrealBuildTool/WinRT/WinRTToolChain.cs
@@ -11,13 +11,9 @@ namespace UnrealBuildTool
 {
 	class WinRTToolChain : UEToolChain
 	{
-		public override void RegisterToolChain()
+		public WinRTToolChain(CPPTargetPlatform InCppPlatform)
+			: base(InCppPlatform)
 		{
-			// Register this tool chain for WinRT
-			Log.TraceVerbose("        Registered for {0}", CPPTargetPlatform.WinRT.ToString());
-			Log.TraceVerbose("        Registered for {0}", CPPTargetPlatform.WinRT_ARM.ToString());
-			UEToolChain.RegisterPlatformToolChain(CPPTargetPlatform.WinRT, this);
-			UEToolChain.RegisterPlatformToolChain(CPPTargetPlatform.WinRT_ARM, this);
 		}
 
 		static string GetCLArguments_Global(CPPEnvironment CompileEnvironment)
@@ -976,10 +972,5 @@ namespace UnrealBuildTool
 				bAreEnvironmentVariablesAlreadyInitialized = true;
 			}
 		}
-
-		public override UnrealTargetPlatform GetPlatform()
-		{
-			return UnrealTargetPlatform.WinRT;
-		}
 	};
 }
diff --git a/Engine/Source/Programs/UnrealBuildTool/Windows/UEBuildWindows.cs b/Engine/Source/Programs/UnrealBuildTool/Windows/UEBuildWindows.cs
--- a/Engine/Source/Programs/UnrealBuildTool/Windows/UEBuildWindows.cs
+++ b/Engine/Source/Programs/UnrealBuildTool/Windows/UEBuildWindows.cs
@@ -789,6 +789,16 @@ namespace UnrealBuildTool
 			return false;
 		}
 
+		/// <summary>
+		/// Creates a toolchain instance for the given platform.
+		/// </summary>
+		/// <param name="Platform">The platform to create a toolchain for</param>
+		/// <returns>New toolchain instance.</returns>
+		public override UEToolChain CreateToolChain(CPPTargetPlatform Platform, FileReference ProjectFile)
+		{
+			return new VCToolChain(Platform);
+		}
+
 		/// <summary>
 		/// Create a build deployment handler
 		/// </summary>
diff --git a/Engine/Source/Programs/UnrealBuildTool/Windows/VCToolChain.cs b/Engine/Source/Programs/UnrealBuildTool/Windows/VCToolChain.cs
--- a/Engine/Source/Programs/UnrealBuildTool/Windows/VCToolChain.cs
+++ b/Engine/Source/Programs/UnrealBuildTool/Windows/VCToolChain.cs
@@ -12,16 +12,11 @@ namespace UnrealBuildTool
 {
 	public class VCToolChain : UEToolChain
 	{
-		public override void RegisterToolChain()
+		public VCToolChain(CPPTargetPlatform CppPlatform)
+			: base(CppPlatform)
 		{
-			// Register this tool chain for both Win64 and Win32
-			Log.TraceVerbose("        Registered for {0}", CPPTargetPlatform.Win64.ToString());
-			UEToolChain.RegisterPlatformToolChain(CPPTargetPlatform.Win64, this);
-			Log.TraceVerbose("        Registered for {0}", CPPTargetPlatform.Win32.ToString());
-			UEToolChain.RegisterPlatformToolChain(CPPTargetPlatform.Win32, this);
 		}
 
-
 		static void AddDefinition(StringBuilder String, string Definition)
 		{
 			// Split the definition into name and value
