Remove unnecessary UEBuildTarget parameter being passed to build platforms.

diff --git a/Engine/Source/Programs/UnrealBuildTool/Android/UEBuildAndroid.cs b/Engine/Source/Programs/UnrealBuildTool/Android/UEBuildAndroid.cs
--- a/Engine/Source/Programs/UnrealBuildTool/Android/UEBuildAndroid.cs
+++ b/Engine/Source/Programs/UnrealBuildTool/Android/UEBuildAndroid.cs
@@ -292,7 +292,7 @@ namespace UnrealBuildTool
 			return true;
 		}
 
-		public override void GetExtraModules(TargetInfo Target, UEBuildTarget BuildTarget, ref List<string> PlatformExtraModules)
+		public override void GetExtraModules(TargetInfo Target, List<string> PlatformExtraModules)
 		{
 		}
 
diff --git a/Engine/Source/Programs/UnrealBuildTool/Configuration/UEBuildPlatform.cs b/Engine/Source/Programs/UnrealBuildTool/Configuration/UEBuildPlatform.cs
--- a/Engine/Source/Programs/UnrealBuildTool/Configuration/UEBuildPlatform.cs
+++ b/Engine/Source/Programs/UnrealBuildTool/Configuration/UEBuildPlatform.cs
@@ -125,9 +125,8 @@ namespace UnrealBuildTool
 		/// This is to allow undisclosed platforms to add modules they need without exposing information about the platform.
 		/// </summary>
 		/// <param name="Target">The target being build</param>
-		/// <param name="BuildTarget">The UEBuildTarget getting build</param>
 		/// <param name="PlatformExtraModules">OUTPUT the list of extra modules the platform needs to add to the target</param>
-		void GetExtraModules(TargetInfo Target, UEBuildTarget BuildTarget, ref List<string> PlatformExtraModules);
+		void GetExtraModules(TargetInfo Target, List<string> PlatformExtraModules);
 
 		/// <summary>
 		/// Modify the newly created module passed in for this platform.
@@ -591,7 +590,7 @@ namespace UnrealBuildTool
 		/// <param name="Target">     The target being build</param>
 		/// <param name="BuildTarget">    The UEBuildTarget getting build</param>
 		/// <param name="PlatformExtraModules"> OUTPUT the list of extra modules the platform needs to add to the target</param>
-		public virtual void GetExtraModules(TargetInfo Target, UEBuildTarget BuildTarget, ref List<string> PlatformExtraModules)
+		public virtual void GetExtraModules(TargetInfo Target, List<string> PlatformExtraModules)
 		{
 		}
 
diff --git a/Engine/Source/Programs/UnrealBuildTool/Configuration/UEBuildTarget.cs b/Engine/Source/Programs/UnrealBuildTool/Configuration/UEBuildTarget.cs
--- a/Engine/Source/Programs/UnrealBuildTool/Configuration/UEBuildTarget.cs
+++ b/Engine/Source/Programs/UnrealBuildTool/Configuration/UEBuildTarget.cs
@@ -3050,7 +3050,7 @@ namespace UnrealBuildTool
 		{
 			var BuildPlatform = UEBuildPlatform.GetBuildPlatform(Platform);
 			List<string> PlatformExtraModules = new List<string>();
-			BuildPlatform.GetExtraModules(TargetInfo, this, ref PlatformExtraModules);
+			BuildPlatform.GetExtraModules(TargetInfo, PlatformExtraModules);
 			ExtraModuleNames.AddRange(PlatformExtraModules);
 		}
 
diff --git a/Engine/Source/Programs/UnrealBuildTool/Linux/UEBuildLinux.cs b/Engine/Source/Programs/UnrealBuildTool/Linux/UEBuildLinux.cs
--- a/Engine/Source/Programs/UnrealBuildTool/Linux/UEBuildLinux.cs
+++ b/Engine/Source/Programs/UnrealBuildTool/Linux/UEBuildLinux.cs
@@ -264,9 +264,8 @@ namespace UnrealBuildTool
 		/// This is to allow undisclosed platforms to add modules they need without exposing information about the platfomr.
 		/// </summary>
 		/// <param name="Target">     The target being build</param>
-		/// <param name="BuildTarget">    The UEBuildTarget getting build</param>
 		/// <param name="PlatformExtraModules"> OUTPUT the list of extra modules the platform needs to add to the target</param>
-		public override void GetExtraModules(TargetInfo Target, UEBuildTarget BuildTarget, ref List<string> PlatformExtraModules)
+		public override void GetExtraModules(TargetInfo Target, List<string> PlatformExtraModules)
 		{
 		}
 
diff --git a/Engine/Source/Programs/UnrealBuildTool/WinRT/UEBuildWinRT.cs b/Engine/Source/Programs/UnrealBuildTool/WinRT/UEBuildWinRT.cs
--- a/Engine/Source/Programs/UnrealBuildTool/WinRT/UEBuildWinRT.cs
+++ b/Engine/Source/Programs/UnrealBuildTool/WinRT/UEBuildWinRT.cs
@@ -218,9 +218,8 @@ namespace UnrealBuildTool
 		/// This is to allow undisclosed platforms to add modules they need without exposing information about the platfomr.
 		/// </summary>
 		/// <param name="Target">     The target being build</param>
-		/// <param name="BuildTarget">    The UEBuildTarget getting build</param>
 		/// <param name="PlatformExtraModules"> OUTPUT the list of extra modules the platform needs to add to the target</param>
-		public override void GetExtraModules(TargetInfo Target, UEBuildTarget BuildTarget, ref List<string> PlatformExtraModules)
+		public override void GetExtraModules(TargetInfo Target, List<string> PlatformExtraModules)
 		{
 			if (Target.Platform == UnrealTargetPlatform.WinRT)
 			{
