Remove a reference to the global uproject file.

diff --git a/Engine/Source/Programs/UnrealBuildTool/System/UnrealBuildTool.cs b/Engine/Source/Programs/UnrealBuildTool/System/UnrealBuildTool.cs
--- a/Engine/Source/Programs/UnrealBuildTool/System/UnrealBuildTool.cs
+++ b/Engine/Source/Programs/UnrealBuildTool/System/UnrealBuildTool.cs
@@ -2710,9 +2710,9 @@ namespace UnrealBuildTool
 			// If there's only one target, just save the UBTMakefile in the target's build intermediate directory
 			// under a folder for that target (and platform/config combo.)
 			DirectoryReference PlatformIntermediatePath;
-			if (UnrealBuildTool.HasUProjectFile())
+			if (Target.ProjectFile != null)
 			{
-				PlatformIntermediatePath = DirectoryReference.Combine(UnrealBuildTool.GetUProjectPath(), BuildConfiguration.PlatformIntermediateFolder);
+				PlatformIntermediatePath = DirectoryReference.Combine(Target.ProjectFile.Directory, BuildConfiguration.PlatformIntermediateFolder);
 			}
 			else
 			{
