Remove IUEBuildPlatform; nothing is using it.

diff --git a/Engine/Source/Programs/AutomationTool/AutomationUtils/ProjectUtils.cs b/Engine/Source/Programs/AutomationTool/AutomationUtils/ProjectUtils.cs
--- a/Engine/Source/Programs/AutomationTool/AutomationUtils/ProjectUtils.cs
+++ b/Engine/Source/Programs/AutomationTool/AutomationUtils/ProjectUtils.cs
@@ -211,7 +211,7 @@ namespace AutomationTool
 			bool RetVal = false;
 			foreach (UnrealTargetPlatform TargetPlatformType in TargetPlatforms)
 			{
-				IUEBuildPlatform BuildPlat = UEBuildPlatform.GetBuildPlatform(TargetPlatformType, true);
+				UEBuildPlatform BuildPlat = UEBuildPlatform.GetBuildPlatform(TargetPlatformType, true);
 				if (!GlobalCommandLine.Rocket && BuildPlat != null && !(BuildPlat as UEBuildPlatform).HasDefaultBuildConfig(TargetPlatformType, RawProjectPath.Directory))
 				{
 					RetVal = true;
diff --git a/Engine/Source/Programs/AutomationTool/Scripts/RocketBuild.Automation.cs b/Engine/Source/Programs/AutomationTool/Scripts/RocketBuild.Automation.cs
--- a/Engine/Source/Programs/AutomationTool/Scripts/RocketBuild.Automation.cs
+++ b/Engine/Source/Programs/AutomationTool/Scripts/RocketBuild.Automation.cs
@@ -426,7 +426,7 @@ namespace Rocket
 
 		public static void StripSymbols(UnrealTargetPlatform TargetPlatform, string[] SourceFileNames, string[] TargetFileNames)
 		{
-			IUEBuildPlatform Platform = UEBuildPlatform.GetBuildPlatform(TargetPlatform);
+			UEBuildPlatform Platform = UEBuildPlatform.GetBuildPlatform(TargetPlatform);
 			UEToolChain ToolChain = Platform.CreateToolChain(Platform.GetCPPTargetPlatform(TargetPlatform), null);
 			for (int Idx = 0; Idx < SourceFileNames.Length; Idx++)
 			{
diff --git a/Engine/Source/Programs/UnrealBuildTool/Configuration/UEBuildModule.cs b/Engine/Source/Programs/UnrealBuildTool/Configuration/UEBuildModule.cs
--- a/Engine/Source/Programs/UnrealBuildTool/Configuration/UEBuildModule.cs
+++ b/Engine/Source/Programs/UnrealBuildTool/Configuration/UEBuildModule.cs
@@ -1586,7 +1586,7 @@ namespace UnrealBuildTool
 
 
 
-		private static FileItem CachePCHUsageForCPPFile(UEBuildTarget Target, FileItem CPPFile, IUEBuildPlatform BuildPlatform, List<string> IncludePathsToSearch, Dictionary<string, FileItem> IncludeFileSearchDictionary)
+		private static FileItem CachePCHUsageForCPPFile(UEBuildTarget Target, FileItem CPPFile, UEBuildPlatform BuildPlatform, List<string> IncludePathsToSearch, Dictionary<string, FileItem> IncludeFileSearchDictionary)
 		{
 			// @todo ubtmake: We don't really need to scan every file looking for PCH headers, just need one.  The rest is just for error checking.
 			// @todo ubtmake: We don't need all of the direct includes either.  We just need the first, unless we want to check for errors.
diff --git a/Engine/Source/Programs/UnrealBuildTool/Configuration/UEBuildPlatform.cs b/Engine/Source/Programs/UnrealBuildTool/Configuration/UEBuildPlatform.cs
--- a/Engine/Source/Programs/UnrealBuildTool/Configuration/UEBuildPlatform.cs
+++ b/Engine/Source/Programs/UnrealBuildTool/Configuration/UEBuildPlatform.cs
@@ -14,258 +14,9 @@ namespace UnrealBuildTool
 		Invalid,		// Could not find the desired SDK, SDK setup failed, etc.		
 	};
 
-	public interface IUEBuildPlatform
+	public abstract partial class UEBuildPlatform
 	{
-		/// <summary>
-		/// Whether the required external SDKs are installed for this platform
-		/// </summary>
-		SDKStatus HasRequiredSDKsInstalled();
-
-		/// <summary>
-		/// If this platform can be compiled with XGE
-		/// </summary>
-		bool CanUseXGE();
-
-		/// <summary>
-		/// If this platform can be compiled with DMUCS/Distcc
-		/// </summary>
-		bool CanUseDistcc();
-
-		/// <summary>
-		/// If this platform can be compiled with SN-DBS
-		/// </summary>
-		bool CanUseSNDBS();
-
-		/// <summary>
-		/// Register the platform with the UEBuildPlatform class
-		/// </summary>
-		void RegisterBuildPlatform();
-
-		/// <summary>
-		/// Attempt to set up AutoSDK for this platform
-		/// </summary>
-		void ManageAndValidateSDK();
-
-		/// <summary>
-		/// Retrieve the CPPTargetPlatform for the given UnrealTargetPlatform
-		/// </summary>
-		/// <param name="InUnrealTargetPlatform">The UnrealTargetPlatform being build</param>
-		/// <returns>CPPTargetPlatform The CPPTargetPlatform to compile for</returns>
-		CPPTargetPlatform GetCPPTargetPlatform(UnrealTargetPlatform InUnrealTargetPlatform);
-
-		/// <summary>
-		/// Get the extension to use for the given binary type
-		/// </summary>
-		/// <param name="InBinaryType">The binary type being built</param>
-		/// <returns>string The binary extension (i.e. 'exe' or 'dll')</returns>
-		string GetBinaryExtension(UEBuildBinaryType InBinaryType);
-
-		/// <summary>
-		/// Get the extension to use for debug info for the given binary type
-		/// </summary>
-		/// <param name="InBinaryType">The binary type being built</param>
-		/// <returns>string The debug info extension (i.e. 'pdb')</returns>
-		string GetDebugInfoExtension(UEBuildBinaryType InBinaryType);
-
-		/// <summary>
-		/// Whether incremental linking should be used
-		/// </summary>
-		/// <param name="InPlatform">The CPPTargetPlatform being built</param>
-		/// <param name="InConfiguration">The CPPTargetConfiguration being built</param>
-		/// <returns>bool true if incremental linking should be used, false if not</returns>
-		bool ShouldUseIncrementalLinking(CPPTargetPlatform Platform, CPPTargetConfiguration Configuration);
-
-		/// <summary>
-		/// Whether PDB files should be used
-		/// </summary>
-		/// <param name="InPlatform">The CPPTargetPlatform being built</param>
-		/// <param name="InConfiguration">The CPPTargetConfiguration being built</param>
-		/// <param name="bInCreateDebugInfo">true if debug info is getting create, false if not</param>
-		/// <returns>bool true if PDB files should be used, false if not</returns>
-		bool ShouldUsePDBFiles(CPPTargetPlatform Platform, CPPTargetConfiguration Configuration, bool bCreateDebugInfo);
-
-		/// <summary>
-		/// Whether PCH files should be used
-		/// </summary>
-		/// <param name="InPlatform">The CPPTargetPlatform being built</param>
-		/// <param name="InConfiguration">The CPPTargetConfiguration being built</param>
-		/// <returns>bool true if PCH files should be used, false if not</returns>
-		bool ShouldUsePCHFiles(CPPTargetPlatform Platform, CPPTargetConfiguration Configuration);
-
-		/// <summary>
-		/// Whether the editor should be built for this platform or not
-		/// </summary>
-		/// <param name="InPlatform">The UnrealTargetPlatform being built</param>
-		/// <param name="InConfiguration">The UnrealTargetConfiguration being built</param>
-		/// <returns>bool true if the editor should be built, false if not</returns>
-		bool ShouldNotBuildEditor(UnrealTargetPlatform InPlatform, UnrealTargetConfiguration InConfiguration);
-
-		/// <summary>
-		/// Whether this build should support ONLY cooked data or not
-		/// </summary>
-		/// <param name="InPlatform">The UnrealTargetPlatform being built</param>
-		/// <param name="InConfiguration">The UnrealTargetConfiguration being built</param>
-		/// <returns>bool true if the editor should be built, false if not</returns>
-		bool BuildRequiresCookedData(UnrealTargetPlatform InPlatform, UnrealTargetConfiguration InConfiguration);
-
-		/// <summary>
-		/// Whether the platform requires the extra UnityCPPWriter
-		/// This is used to add an extra file for UBT to get the #include dependencies from
-		/// </summary>
-		/// <returns>bool true if it is required, false if not</returns>
-		bool RequiresExtraUnityCPPWriter();
-
-		/// <summary>
-		/// Whether this platform should build a monolithic binary
-		/// </summary>
-		bool ShouldCompileMonolithicBinary(UnrealTargetPlatform InPlatform);
-
-		/// <summary>
-		/// Get a list of extra modules the platform requires.
-		/// This is to allow undisclosed platforms to add modules they need without exposing information about the platform.
-		/// </summary>
-		/// <param name="Target">The target being build</param>
-		/// <param name="PlatformExtraModules">OUTPUT the list of extra modules the platform needs to add to the target</param>
-		void GetExtraModules(TargetInfo Target, List<string> PlatformExtraModules);
-
-		/// <summary>
-		/// Modify the newly created module passed in for this platform.
-		/// This is not required - but allows for hiding details of a
-		/// particular platform.
-		/// </summary>
-		/// <param name="">ModuleName  The name of the module</param>
-		/// <param name="Rules">  The module rules</param>
-		/// <param name="Target">  The target being build</param>
-		void ModifyModuleRules(string ModuleName, ModuleRules Rules, TargetInfo Target);
-
-		/// <summary>
-		/// Setup the target environment for building
-		/// </summary>
-		/// <param name="InBuildTarget">The target being built</param>
-		void SetUpEnvironment(UEBuildTarget InBuildTarget);
-
-		/// <summary>
-		/// Allow the platform to set an optional architecture
-		/// </summary>
-		string GetActiveArchitecture();
-
-		/// <summary>
-		/// Allow the platform to apply architecture-specific name according to its rules
-		/// </summary>
-		/// <param name="BinaryName">name of the binary, not specific to any architecture</param>
-		/// <returns>Possibly architecture-specific name</returns>
-		string ApplyArchitectureName(string BinaryName);
-
-		/// <summary>
-		/// For platforms that need to output multiple files per binary (ie Android "fat" binaries)
-		/// this will emit multiple paths. By default, it simply makes an array from the input
-		/// </summary>
-		List<FileReference> FinalizeBinaryPaths(FileReference BinaryName, FileReference ProjectFile);
-
-		/// <summary>
-		/// Setup the configuration environment for building
-		/// </summary>
-		/// <param name="InBuildTarget">The target being built</param>
-		void SetUpConfigurationEnvironment(UEBuildTarget InBuildTarget);
-
-		/// <summary>
-		/// Setup the project environment for building
-		/// </summary>
-		/// <param name="InBuildTarget">The target being built</param>
-		void SetUpProjectEnvironment(UnrealTargetPlatform InPlatform);
-
-		/// <summary>
-		/// Whether this platform should create debug information or not
-		/// </summary>
-		/// <param name="InPlatform">The UnrealTargetPlatform being built</param>
-		/// <param name="InConfiguration">The UnrealTargetConfiguration being built</param>
-		/// <returns>bool true if debug info should be generated, false if not</returns>
-		bool ShouldCreateDebugInfo(UnrealTargetPlatform Platform, UnrealTargetConfiguration Configuration);
-
-		/// <summary>
-		/// Gives the platform a chance to 'override' the configuration settings
-		/// that are overridden on calls to RunUBT.
-		/// </summary>
-		/// <param name="InPlatform">The UnrealTargetPlatform being built</param>
-		/// <param name="InConfiguration">The UnrealTargetConfiguration being built</param>
-		/// <returns>bool true if debug info should be generated, false if not</returns>
-		void ResetBuildConfiguration(UnrealTargetPlatform InPlatform, UnrealTargetConfiguration InConfiguration);
-
-		/// <summary>
-		/// Validate the UEBuildConfiguration for this platform
-		/// This is called BEFORE calling UEBuildConfiguration to allow setting
-		/// various fields used in that function such as CompileLeanAndMean...
-		/// </summary>
-		void ValidateUEBuildConfiguration();
-
-		/// <summary>
-		/// Validate configuration for this platform
-		/// NOTE: This function can/will modify BuildConfiguration!
-		/// </summary>
-		/// <param name="InPlatform">The CPPTargetPlatform being built</param>
-		/// <param name="InConfiguration">The CPPTargetConfiguration being built</param>
-		/// <param name="bInCreateDebugInfo">true if debug info is getting create, false if not</param>
-		void ValidateBuildConfiguration(CPPTargetConfiguration Configuration, CPPTargetPlatform Platform, bool bCreateDebugInfo);
-
-		/// <summary>
-		/// Return whether this platform has uniquely named binaries across multiple games
-		/// </summary>
-		bool HasUniqueBinaries();
-
-		/// <summary>
-		/// Return whether we wish to have this platform's binaries in our builds
-		/// </summary>
-		bool IsBuildRequired();
-
-		/// <summary>
-		/// Return whether we wish to have this platform's binaries in our CIS tests
-		/// </summary>
-		bool IsCISRequired();
-
-		/// <summary>
-		/// Whether the build platform requires deployment prep
-		/// </summary>
-		/// <returns></returns>
-		bool RequiresDeployPrepAfterCompile();
-
-		/// <summary>
-		/// Return all valid configurations for this platform
-		/// Typically, this is always Debug, Development, and Shipping - but Test is a likely future addition for some platforms
-		/// </summary>
-		List<UnrealTargetConfiguration> GetConfigurations(UnrealTargetPlatform InUnrealTargetPlatform, bool bIncludeDebug);
-
-		/// <summary>
-		/// Setup the binaries for this specific platform.
-		/// </summary>
-		/// <param name="InBuildTarget">The target being built</param>
-		void SetupBinaries(UEBuildTarget InBuildTarget);
-
-		/// <summary>
-		/// Check for the default configuration
-		/// return true if the project uses the default build config
-		/// </summary>
-		bool HasDefaultBuildConfig(UnrealTargetPlatform InPlatform, DirectoryReference InProjectPath);
-
-		/// <summary>
-		/// Creates a toolchain instance for the given platform. There should be a single toolchain instance per-target, as their may be
-		/// state data and configuration cached between calls.
-		/// </summary>
-		/// <param name="Platform">The platform to create a toolchain for</param>
-		/// <returns>New toolchain instance.</returns>
-		UEToolChain CreateToolChain(CPPTargetPlatform Platform, FileReference ProjectFile);
-
-		/// <summary>
-		/// Create a build deployment handler
-		/// </summary>
-		/// <param name="ProjectFile">The project file of the target being deployed. Used to find any deployment specific settings.</param>
-		/// <param name="DeploymentHandler">The output deployment handler</param>
-		/// <returns>True if the platform requires a deployment handler, false otherwise</returns>
-		bool TryCreateDeploymentHandler(FileReference ProjectFile, out UEBuildDeploy DeploymentHandler);
-	}
-
-	public abstract partial class UEBuildPlatform : IUEBuildPlatform
-	{
-		public static Dictionary<UnrealTargetPlatform, IUEBuildPlatform> BuildPlatformDictionary = new Dictionary<UnrealTargetPlatform, IUEBuildPlatform>();
+		public static Dictionary<UnrealTargetPlatform, UEBuildPlatform> BuildPlatformDictionary = new Dictionary<UnrealTargetPlatform, UEBuildPlatform>();
 
 		// a mapping of a group to the platforms in the group (ie, Microsoft contains Win32 and Win64)
 		static Dictionary<UnrealPlatformGroup, List<UnrealTargetPlatform>> PlatformGroupDictionary = new Dictionary<UnrealPlatformGroup, List<UnrealTargetPlatform>>();
@@ -353,7 +104,7 @@ namespace UnrealBuildTool
 		/// <param name="InPlatform">  The UnrealTargetPlatform being built</param>
 		/// <param name="bInAllowFailure"> If true, do not throw an exception and return null</param>
 		/// <returns>UEBuildPlatform  The instance of the build platform</returns>
-		public static IUEBuildPlatform GetBuildPlatform(UnrealTargetPlatform InPlatform, bool bInAllowFailure = false)
+		public static UEBuildPlatform GetBuildPlatform(UnrealTargetPlatform InPlatform, bool bInAllowFailure = false)
 		{
 			if (BuildPlatformDictionary.ContainsKey(InPlatform) == true)
 			{
@@ -372,7 +123,7 @@ namespace UnrealBuildTool
 		/// <param name="InPlatform">  The CPPTargetPlatform being built</param>
 		/// <param name="bInAllowFailure"> If true, do not throw an exception and return null</param>
 		/// <returns>UEBuildPlatform  The instance of the build platform</returns>
-		public static IUEBuildPlatform GetBuildPlatformForCPPTargetPlatform(CPPTargetPlatform InPlatform, bool bInAllowFailure = false)
+		public static UEBuildPlatform GetBuildPlatformForCPPTargetPlatform(CPPTargetPlatform InPlatform, bool bInAllowFailure = false)
 		{
 			UnrealTargetPlatform UTPlatform = UEBuildTarget.CPPTargetPlatformToUnrealTargetPlatform(InPlatform);
 			if (BuildPlatformDictionary.ContainsKey(UTPlatform) == true)
@@ -976,11 +727,8 @@ namespace UnrealBuildTool
 		/// <param name="DeploymentHandler">The output deployment handler</param>
 		/// <returns>True if the platform requires a deployment handler, false otherwise</returns>
 		public abstract bool TryCreateDeploymentHandler(FileReference ProjectFile, out UEBuildDeploy DeploymentHandler);
-	}
 
-	// AutoSDKs handling portion
-	public abstract partial class UEBuildPlatform : IUEBuildPlatform
-	{
+		// AutoSDKs handling portion
 
 		#region protected AutoSDKs Utility
 
diff --git a/Engine/Source/Programs/UnrealBuildTool/Configuration/UEBuildTarget.cs b/Engine/Source/Programs/UnrealBuildTool/Configuration/UEBuildTarget.cs
--- a/Engine/Source/Programs/UnrealBuildTool/Configuration/UEBuildTarget.cs
+++ b/Engine/Source/Programs/UnrealBuildTool/Configuration/UEBuildTarget.cs
@@ -1306,7 +1306,7 @@ namespace UnrealBuildTool
 		protected void CleanTarget(BuildManifest Manifest)
 		{
 			// Expand all the paths in the receipt; they'll currently use variables for the engine and project directories
-			IUEBuildPlatform BuildPlatform = UEBuildPlatform.GetBuildPlatform(Platform);
+			UEBuildPlatform BuildPlatform = UEBuildPlatform.GetBuildPlatform(Platform);
 			// Expand all the paths in the receipt; they'll currently use variables for the engine and project directories
 			TargetReceipt ReceiptWithFullPaths;
 			if (!TargetReceipt.TryRead(ReceiptFileName, out ReceiptWithFullPaths))
@@ -1569,7 +1569,7 @@ namespace UnrealBuildTool
 			// Expand all the paths in the receipt; they'll currently use variables for the engine and project directories
 			TargetReceipt ReceiptWithFullPaths = new TargetReceipt(Receipt);
 			ReceiptWithFullPaths.ExpandPathVariables(UnrealBuildTool.EngineDirectory, ProjectDirectory);
-			IUEBuildPlatform BuildPlatform = UEBuildPlatform.GetBuildPlatform(Platform);
+			UEBuildPlatform BuildPlatform = UEBuildPlatform.GetBuildPlatform(Platform);
 			// Expand all the paths in the receipt; they'll currently use variables for the engine and project directories
 			TargetReceipt BuiltReceiptWithFullPaths;
 			if (!TargetReceipt.TryRead(ReceiptFileName, out BuiltReceiptWithFullPaths))
@@ -1658,7 +1658,7 @@ namespace UnrealBuildTool
 			HashSet<string> FileNames = new HashSet<string>(StringComparer.CurrentCultureIgnoreCase);
 
 			// Get the platform we're building for
-			IUEBuildPlatform BuildPlatform = UEBuildPlatform.GetBuildPlatform(Platform);
+			UEBuildPlatform BuildPlatform = UEBuildPlatform.GetBuildPlatform(Platform);
 
 			// Add all their include paths
 			foreach (string ModuleName in ModuleNames)
@@ -1802,7 +1802,7 @@ namespace UnrealBuildTool
 				Manifest.AddBuildProduct(BuildProduct.Path);
 			}
 
-			IUEBuildPlatform BuildPlatform = UEBuildPlatform.GetBuildPlatform(Platform);
+			UEBuildPlatform BuildPlatform = UEBuildPlatform.GetBuildPlatform(Platform);
 			if (OnlyModules.Count == 0)
 			{
 				Manifest.AddBuildProduct(ReceiptFileName);
@@ -1916,7 +1916,7 @@ namespace UnrealBuildTool
 		{
 			if (Receipt != null)
 			{
-				IUEBuildPlatform BuildPlatform = UEBuildPlatform.GetBuildPlatform(Platform);
+				UEBuildPlatform BuildPlatform = UEBuildPlatform.GetBuildPlatform(Platform);
 				if (OnlyModules == null || OnlyModules.Count == 0)
 				{
 					Directory.CreateDirectory(Path.GetDirectoryName(ReceiptFileName));
@@ -2996,7 +2996,7 @@ namespace UnrealBuildTool
 				Result.AppendFormat("-{0}-{1}", Platform.ToString(), Configuration.ToString());
 			}
 
-			IUEBuildPlatform BuildPlatform = UEBuildPlatform.GetBuildPlatform(Platform);
+			UEBuildPlatform BuildPlatform = UEBuildPlatform.GetBuildPlatform(Platform);
 			Result.Append(BuildPlatform.ApplyArchitectureName(""));
 
 			if (BuildConfiguration.bRunUnrealCodeAnalyzer)
diff --git a/Engine/Source/Programs/UnrealBuildTool/System/CPPHeaders.cs b/Engine/Source/Programs/UnrealBuildTool/System/CPPHeaders.cs
--- a/Engine/Source/Programs/UnrealBuildTool/System/CPPHeaders.cs
+++ b/Engine/Source/Programs/UnrealBuildTool/System/CPPHeaders.cs
@@ -217,7 +217,7 @@ namespace UnrealBuildTool
 		static Dictionary<FileItem, IncludedFilesSet> OnlyCachedIncludedFilesMap = new Dictionary<FileItem, IncludedFilesSet>();
 
 
-		public static List<FileItem> FindAndCacheAllIncludedFiles(UEBuildTarget Target, FileItem SourceFile, IUEBuildPlatform BuildPlatform, CPPIncludeInfo CPPIncludeInfo, bool bOnlyCachedDependencies)
+		public static List<FileItem> FindAndCacheAllIncludedFiles(UEBuildTarget Target, FileItem SourceFile, UEBuildPlatform BuildPlatform, CPPIncludeInfo CPPIncludeInfo, bool bOnlyCachedDependencies)
 		{
 			List<FileItem> Result = null;
 
@@ -280,7 +280,7 @@ namespace UnrealBuildTool
 		/// <param name="CPPFile">C++ file to get the dependencies for.</param>
 		/// <param name="Result">List of CPPFile dependencies.</param>
 		/// <returns>false if CPPFile is still being processed further down the callstack, true otherwise.</returns>
-		public static bool FindAndCacheAllIncludedFiles(UEBuildTarget Target, FileItem CPPFile, IUEBuildPlatform BuildPlatform, CPPIncludeInfo CPPIncludeInfo, ref IncludedFilesSet Result, bool bOnlyCachedDependencies)
+		public static bool FindAndCacheAllIncludedFiles(UEBuildTarget Target, FileItem CPPFile, UEBuildPlatform BuildPlatform, CPPIncludeInfo CPPIncludeInfo, ref IncludedFilesSet Result, bool bOnlyCachedDependencies)
 		{
 			IncludedFilesSet IncludedFileList;
 			var IncludedFilesMap = bOnlyCachedDependencies ? OnlyCachedIncludedFilesMap : ExhaustiveIncludedFilesMap;
@@ -426,7 +426,7 @@ namespace UnrealBuildTool
 		/// <summary>
 		/// Finds the names of files directly included by the given C++ file, and also whether the file contains any UObjects
 		/// </summary>
-		public static List<DependencyInclude> GetDirectIncludeDependencies(UEBuildTarget Target, FileItem CPPFile, IUEBuildPlatform BuildPlatform, bool bOnlyCachedDependencies)
+		public static List<DependencyInclude> GetDirectIncludeDependencies(UEBuildTarget Target, FileItem CPPFile, UEBuildPlatform BuildPlatform, bool bOnlyCachedDependencies)
 		{
 			// Try to fulfill request from cache first.
 			List<DependencyInclude> Info = IncludeDependencyCache[Target].GetCachedDependencyInfo(CPPFile);
diff --git a/Engine/Source/Programs/UnrealBuildTool/System/UnrealBuildTool.cs b/Engine/Source/Programs/UnrealBuildTool/System/UnrealBuildTool.cs
--- a/Engine/Source/Programs/UnrealBuildTool/System/UnrealBuildTool.cs
+++ b/Engine/Source/Programs/UnrealBuildTool/System/UnrealBuildTool.cs
@@ -486,7 +486,7 @@ namespace UnrealBuildTool
 					{
 						if (CheckType.IsClass && !CheckType.IsAbstract)
 						{
-							if (Utils.ImplementsInterface<IUEBuildPlatform>(CheckType))
+							if (CheckType.IsSubclassOf(typeof(UEBuildPlatform)))
 							{
 								Log.TraceVerbose("    Registering build platform: {0}", CheckType.ToString());
 								var TempInst = (UEBuildPlatform)(UBTAssembly.CreateInstance(CheckType.FullName, true));
@@ -1381,7 +1381,7 @@ namespace UnrealBuildTool
 							if ((Result == ECompilationResult.Succeeded) && (BuildConfiguration.bDeployAfterCompile == true) && (BuildConfiguration.bXGEExport == false) &&
 								(UEBuildConfiguration.bGenerateManifest == false) && (UEBuildConfiguration.bGenerateExternalFileList == false) && (UEBuildConfiguration.bCleanProject == false))
 							{
-								IUEBuildPlatform BuildPlatform = UEBuildPlatform.GetBuildPlatform(CheckPlatform);
+								UEBuildPlatform BuildPlatform = UEBuildPlatform.GetBuildPlatform(CheckPlatform);
 
 								UEBuildDeploy DeploymentHandler;
 								if (BuildPlatform.TryCreateDeploymentHandler(UnrealBuildTool.GetUProjectFile(), out DeploymentHandler))
@@ -2683,10 +2683,12 @@ namespace UnrealBuildTool
 				// configurations and platforms, and save it into the base intermediate folder
 				var TargetCollectionName = MakeTargetCollectionName(TargetDescs);
 
+				TargetDescriptor DescriptorWithProject = TargetDescs.FirstOrDefault(x => x.ProjectFile != null);
+
 				DirectoryReference ProjectIntermediatePath;
-				if (UnrealBuildTool.HasUProjectFile())
+				if (DescriptorWithProject != null)
 				{
-					ProjectIntermediatePath = DirectoryReference.Combine(UnrealBuildTool.GetUProjectPath(), BuildConfiguration.BaseIntermediateFolder);
+					ProjectIntermediatePath = DirectoryReference.Combine(DescriptorWithProject.ProjectFile.Directory, BuildConfiguration.BaseIntermediateFolder);
 				}
 				else
 				{
diff --git a/Engine/Source/Programs/UnrealBuildTool/ToolChain/UEToolChain.cs b/Engine/Source/Programs/UnrealBuildTool/ToolChain/UEToolChain.cs
--- a/Engine/Source/Programs/UnrealBuildTool/ToolChain/UEToolChain.cs
+++ b/Engine/Source/Programs/UnrealBuildTool/ToolChain/UEToolChain.cs
@@ -106,7 +106,7 @@ namespace UnrealBuildTool
 			return true;
 		}
 
-		protected void AddPrerequisiteSourceFile(UEBuildTarget Target, IUEBuildPlatform BuildPlatform, CPPEnvironment CompileEnvironment, FileItem SourceFile, List<FileItem> PrerequisiteItems)
+		protected void AddPrerequisiteSourceFile(UEBuildTarget Target, UEBuildPlatform BuildPlatform, CPPEnvironment CompileEnvironment, FileItem SourceFile, List<FileItem> PrerequisiteItems)
 		{
 			PrerequisiteItems.Add(SourceFile);
 
