Pass the compile and link environment directly to UEBuildPlatform.SetupConfigurationEnvironment(), it breaks encapsulation to pass the target directly.

diff --git a/Engine/Source/Programs/UnrealBuildTool/Configuration/UEBuildPlatform.cs b/Engine/Source/Programs/UnrealBuildTool/Configuration/UEBuildPlatform.cs
--- a/Engine/Source/Programs/UnrealBuildTool/Configuration/UEBuildPlatform.cs
+++ b/Engine/Source/Programs/UnrealBuildTool/Configuration/UEBuildPlatform.cs
@@ -406,11 +406,11 @@ namespace UnrealBuildTool
 		/// Setup the configuration environment for building
 		/// </summary>
 		/// <param name="InBuildTarget"> The target being built</param>
-		public virtual void SetUpConfigurationEnvironment(UEBuildTarget InBuildTarget)
+		public virtual void SetUpConfigurationEnvironment(TargetInfo Target, CPPEnvironment GlobalCompileEnvironment, LinkEnvironment GlobalLinkEnvironment)
 		{
 			// Determine the C++ compile/link configuration based on the Unreal configuration.
 			CPPTargetConfiguration CompileConfiguration;
-			UnrealTargetConfiguration CheckConfig = InBuildTarget.Configuration;
+			UnrealTargetConfiguration CheckConfig = Target.Configuration;
 
 			switch (CheckConfig)
 			{
@@ -419,41 +419,41 @@ namespace UnrealBuildTool
 					CompileConfiguration = CPPTargetConfiguration.Debug;
 					if (BuildConfiguration.bDebugBuildsActuallyUseDebugCRT)
 					{
-						InBuildTarget.GlobalCompileEnvironment.Config.Definitions.Add("_DEBUG=1"); // the engine doesn't use this, but lots of 3rd party stuff does
+						GlobalCompileEnvironment.Config.Definitions.Add("_DEBUG=1"); // the engine doesn't use this, but lots of 3rd party stuff does
 					}
 					else
 					{
-						InBuildTarget.GlobalCompileEnvironment.Config.Definitions.Add("NDEBUG=1"); // the engine doesn't use this, but lots of 3rd party stuff does
+						GlobalCompileEnvironment.Config.Definitions.Add("NDEBUG=1"); // the engine doesn't use this, but lots of 3rd party stuff does
 					}
-					InBuildTarget.GlobalCompileEnvironment.Config.Definitions.Add("UE_BUILD_DEBUG=1");
+					GlobalCompileEnvironment.Config.Definitions.Add("UE_BUILD_DEBUG=1");
 					break;
 				case UnrealTargetConfiguration.DebugGame:
 				// Individual game modules can be switched to be compiled in debug as necessary. By default, everything is compiled in development.
 				case UnrealTargetConfiguration.Development:
 					CompileConfiguration = CPPTargetConfiguration.Development;
-					InBuildTarget.GlobalCompileEnvironment.Config.Definitions.Add("NDEBUG=1"); // the engine doesn't use this, but lots of 3rd party stuff does
-					InBuildTarget.GlobalCompileEnvironment.Config.Definitions.Add("UE_BUILD_DEVELOPMENT=1");
+					GlobalCompileEnvironment.Config.Definitions.Add("NDEBUG=1"); // the engine doesn't use this, but lots of 3rd party stuff does
+					GlobalCompileEnvironment.Config.Definitions.Add("UE_BUILD_DEVELOPMENT=1");
 					break;
 				case UnrealTargetConfiguration.Shipping:
 					CompileConfiguration = CPPTargetConfiguration.Shipping;
-					InBuildTarget.GlobalCompileEnvironment.Config.Definitions.Add("NDEBUG=1"); // the engine doesn't use this, but lots of 3rd party stuff does
-					InBuildTarget.GlobalCompileEnvironment.Config.Definitions.Add("UE_BUILD_SHIPPING=1");
+					GlobalCompileEnvironment.Config.Definitions.Add("NDEBUG=1"); // the engine doesn't use this, but lots of 3rd party stuff does
+					GlobalCompileEnvironment.Config.Definitions.Add("UE_BUILD_SHIPPING=1");
 					break;
 				case UnrealTargetConfiguration.Test:
 					CompileConfiguration = CPPTargetConfiguration.Shipping;
-					InBuildTarget.GlobalCompileEnvironment.Config.Definitions.Add("NDEBUG=1"); // the engine doesn't use this, but lots of 3rd party stuff does
-					InBuildTarget.GlobalCompileEnvironment.Config.Definitions.Add("UE_BUILD_TEST=1");
+					GlobalCompileEnvironment.Config.Definitions.Add("NDEBUG=1"); // the engine doesn't use this, but lots of 3rd party stuff does
+					GlobalCompileEnvironment.Config.Definitions.Add("UE_BUILD_TEST=1");
 					break;
 			}
 
 			// Set up the global C++ compilation and link environment.
-			InBuildTarget.GlobalCompileEnvironment.Config.Target.Configuration = CompileConfiguration;
-			InBuildTarget.GlobalLinkEnvironment.Config.Target.Configuration = CompileConfiguration;
+			GlobalCompileEnvironment.Config.Target.Configuration = CompileConfiguration;
+			GlobalLinkEnvironment.Config.Target.Configuration = CompileConfiguration;
 
 			// Create debug info based on the heuristics specified by the user.
-			InBuildTarget.GlobalCompileEnvironment.Config.bCreateDebugInfo =
-				!BuildConfiguration.bDisableDebugInfo && ShouldCreateDebugInfo(InBuildTarget.Platform, CheckConfig);
-			InBuildTarget.GlobalLinkEnvironment.Config.bCreateDebugInfo = InBuildTarget.GlobalCompileEnvironment.Config.bCreateDebugInfo;
+			GlobalCompileEnvironment.Config.bCreateDebugInfo =
+				!BuildConfiguration.bDisableDebugInfo && ShouldCreateDebugInfo(Target.Platform, CheckConfig);
+			GlobalLinkEnvironment.Config.bCreateDebugInfo = GlobalCompileEnvironment.Config.bCreateDebugInfo;
 		}
 
 		/// <summary>
diff --git a/Engine/Source/Programs/UnrealBuildTool/Configuration/UEBuildTarget.cs b/Engine/Source/Programs/UnrealBuildTool/Configuration/UEBuildTarget.cs
--- a/Engine/Source/Programs/UnrealBuildTool/Configuration/UEBuildTarget.cs
+++ b/Engine/Source/Programs/UnrealBuildTool/Configuration/UEBuildTarget.cs
@@ -3432,7 +3432,7 @@ namespace UnrealBuildTool
 
 		void SetUpConfigurationEnvironment()
 		{
-			UEBuildPlatform.GetBuildPlatform(Platform).SetUpConfigurationEnvironment(this);
+			UEBuildPlatform.GetBuildPlatform(Platform).SetUpConfigurationEnvironment(TargetInfo, GlobalCompileEnvironment, GlobalLinkEnvironment);
 
 			// Check to see if we're compiling a library or not
 			bool bIsBuildingDLL = OutputPaths[0].HasExtension(".dll");
diff --git a/Engine/Source/Programs/UnrealBuildTool/UWP/UEBuildUWP.cs b/Engine/Source/Programs/UnrealBuildTool/UWP/UEBuildUWP.cs
--- a/Engine/Source/Programs/UnrealBuildTool/UWP/UEBuildUWP.cs
+++ b/Engine/Source/Programs/UnrealBuildTool/UWP/UEBuildUWP.cs
@@ -390,11 +390,11 @@ namespace UnrealBuildTool
 		/// Setup the configuration environment for building
 		/// </summary>
 		/// <param name="InBuildTarget"> The target being built</param>
-		public override void SetUpConfigurationEnvironment(UEBuildTarget InBuildTarget)
+		public override void SetUpConfigurationEnvironment(TargetInfo Target, CPPEnvironment GlobalCompileEnvironment, LinkEnvironment GlobalLinkEnvironment)
 		{
 			// Determine the C++ compile/link configuration based on the Unreal configuration.
 			CPPTargetConfiguration CompileConfiguration;
-			UnrealTargetConfiguration CheckConfig = InBuildTarget.Configuration;
+			UnrealTargetConfiguration CheckConfig = Target.Configuration;
 			switch (CheckConfig)
 			{
 				default:
@@ -402,40 +402,40 @@ namespace UnrealBuildTool
 					CompileConfiguration = CPPTargetConfiguration.Debug;
 					if (BuildConfiguration.bDebugBuildsActuallyUseDebugCRT)
 					{
-						InBuildTarget.GlobalCompileEnvironment.Config.Definitions.Add("_DEBUG=1"); // the engine doesn't use this, but lots of 3rd party stuff does
+						GlobalCompileEnvironment.Config.Definitions.Add("_DEBUG=1"); // the engine doesn't use this, but lots of 3rd party stuff does
 					}
 					else
 					{
-						InBuildTarget.GlobalCompileEnvironment.Config.Definitions.Add("NDEBUG=1"); // the engine doesn't use this, but lots of 3rd party stuff does
+						GlobalCompileEnvironment.Config.Definitions.Add("NDEBUG=1"); // the engine doesn't use this, but lots of 3rd party stuff does
 					}
-					InBuildTarget.GlobalCompileEnvironment.Config.Definitions.Add("UE_BUILD_DEBUG=1");
+					GlobalCompileEnvironment.Config.Definitions.Add("UE_BUILD_DEBUG=1");
 					break;
 				case UnrealTargetConfiguration.DebugGame:
 				// Default to Development; can be overriden by individual modules.
 				case UnrealTargetConfiguration.Development:
 					CompileConfiguration = CPPTargetConfiguration.Development;
-					InBuildTarget.GlobalCompileEnvironment.Config.Definitions.Add("NDEBUG=1"); // the engine doesn't use this, but lots of 3rd party stuff does
-					InBuildTarget.GlobalCompileEnvironment.Config.Definitions.Add("UE_BUILD_DEVELOPMENT=1");
+					GlobalCompileEnvironment.Config.Definitions.Add("NDEBUG=1"); // the engine doesn't use this, but lots of 3rd party stuff does
+					GlobalCompileEnvironment.Config.Definitions.Add("UE_BUILD_DEVELOPMENT=1");
 					break;
 				case UnrealTargetConfiguration.Shipping:
 					CompileConfiguration = CPPTargetConfiguration.Shipping;
-					InBuildTarget.GlobalCompileEnvironment.Config.Definitions.Add("NDEBUG=1"); // the engine doesn't use this, but lots of 3rd party stuff does
-					InBuildTarget.GlobalCompileEnvironment.Config.Definitions.Add("UE_BUILD_SHIPPING=1");
+					GlobalCompileEnvironment.Config.Definitions.Add("NDEBUG=1"); // the engine doesn't use this, but lots of 3rd party stuff does
+					GlobalCompileEnvironment.Config.Definitions.Add("UE_BUILD_SHIPPING=1");
 					break;
 				case UnrealTargetConfiguration.Test:
 					CompileConfiguration = CPPTargetConfiguration.Shipping;
-					InBuildTarget.GlobalCompileEnvironment.Config.Definitions.Add("NDEBUG=1"); // the engine doesn't use this, but lots of 3rd party stuff does
-					InBuildTarget.GlobalCompileEnvironment.Config.Definitions.Add("UE_BUILD_TEST=1");
+					GlobalCompileEnvironment.Config.Definitions.Add("NDEBUG=1"); // the engine doesn't use this, but lots of 3rd party stuff does
+					GlobalCompileEnvironment.Config.Definitions.Add("UE_BUILD_TEST=1");
 					break;
 			}
 
 			// Set up the global C++ compilation and link environment.
-			InBuildTarget.GlobalCompileEnvironment.Config.Target.Configuration = CompileConfiguration;
-			InBuildTarget.GlobalLinkEnvironment.Config.Target.Configuration = CompileConfiguration;
+			GlobalCompileEnvironment.Config.Target.Configuration = CompileConfiguration;
+			GlobalLinkEnvironment.Config.Target.Configuration = CompileConfiguration;
 
 			// Create debug info based on the heuristics specified by the user.
-			InBuildTarget.GlobalCompileEnvironment.Config.bCreateDebugInfo =
-				!BuildConfiguration.bDisableDebugInfo && ShouldCreateDebugInfo(InBuildTarget.Platform, CheckConfig);
+			GlobalCompileEnvironment.Config.bCreateDebugInfo =
+				!BuildConfiguration.bDisableDebugInfo && ShouldCreateDebugInfo(Target.Platform, CheckConfig);
 
 			// NOTE: Even when debug info is turned off, we currently force the linker to generate debug info
 			//       anyway on Visual C++ platforms.  This will cause a PDB file to be generated with symbols
@@ -443,7 +443,7 @@ namespace UnrealBuildTool
 			//       useful call stacks, even though compiler-generate debug info may be disabled.  This gives
 			//       us much of the build-time savings of fully-disabled debug info, without giving up call
 			//       data completely.
-			InBuildTarget.GlobalLinkEnvironment.Config.bCreateDebugInfo = true;
+			GlobalLinkEnvironment.Config.bCreateDebugInfo = true;
 		}
 
 		/// <summary>
diff --git a/Engine/Source/Programs/UnrealBuildTool/Windows/UEBuildWindows.cs b/Engine/Source/Programs/UnrealBuildTool/Windows/UEBuildWindows.cs
--- a/Engine/Source/Programs/UnrealBuildTool/Windows/UEBuildWindows.cs
+++ b/Engine/Source/Programs/UnrealBuildTool/Windows/UEBuildWindows.cs
@@ -702,12 +702,12 @@ namespace UnrealBuildTool
 		/// Setup the configuration environment for building
 		/// </summary>
 		/// <param name="InBuildTarget"> The target being built</param>
-		public override void SetUpConfigurationEnvironment(UEBuildTarget InBuildTarget)
+		public override void SetUpConfigurationEnvironment(TargetInfo Target, CPPEnvironment GlobalCompileEnvironment, LinkEnvironment GlobalLinkEnvironment)
 		{
 			// Determine the C++ compile/link configuration based on the Unreal configuration.
 			CPPTargetConfiguration CompileConfiguration;
 			//@todo SAS: Add a true Debug mode!
-			UnrealTargetConfiguration CheckConfig = InBuildTarget.Configuration;
+			UnrealTargetConfiguration CheckConfig = Target.Configuration;
 			switch (CheckConfig)
 			{
 				default:
@@ -715,40 +715,40 @@ namespace UnrealBuildTool
 					CompileConfiguration = CPPTargetConfiguration.Debug;
 					if (BuildConfiguration.bDebugBuildsActuallyUseDebugCRT)
 					{
-						InBuildTarget.GlobalCompileEnvironment.Config.Definitions.Add("_DEBUG=1"); // the engine doesn't use this, but lots of 3rd party stuff does
+						GlobalCompileEnvironment.Config.Definitions.Add("_DEBUG=1"); // the engine doesn't use this, but lots of 3rd party stuff does
 					}
 					else
 					{
-						InBuildTarget.GlobalCompileEnvironment.Config.Definitions.Add("NDEBUG=1"); // the engine doesn't use this, but lots of 3rd party stuff does
+						GlobalCompileEnvironment.Config.Definitions.Add("NDEBUG=1"); // the engine doesn't use this, but lots of 3rd party stuff does
 					}
-					InBuildTarget.GlobalCompileEnvironment.Config.Definitions.Add("UE_BUILD_DEBUG=1");
+					GlobalCompileEnvironment.Config.Definitions.Add("UE_BUILD_DEBUG=1");
 					break;
 				case UnrealTargetConfiguration.DebugGame:
 				// Default to Development; can be overridden by individual modules.
 				case UnrealTargetConfiguration.Development:
 					CompileConfiguration = CPPTargetConfiguration.Development;
-					InBuildTarget.GlobalCompileEnvironment.Config.Definitions.Add("NDEBUG=1"); // the engine doesn't use this, but lots of 3rd party stuff does
-					InBuildTarget.GlobalCompileEnvironment.Config.Definitions.Add("UE_BUILD_DEVELOPMENT=1");
+					GlobalCompileEnvironment.Config.Definitions.Add("NDEBUG=1"); // the engine doesn't use this, but lots of 3rd party stuff does
+					GlobalCompileEnvironment.Config.Definitions.Add("UE_BUILD_DEVELOPMENT=1");
 					break;
 				case UnrealTargetConfiguration.Shipping:
 					CompileConfiguration = CPPTargetConfiguration.Shipping;
-					InBuildTarget.GlobalCompileEnvironment.Config.Definitions.Add("NDEBUG=1"); // the engine doesn't use this, but lots of 3rd party stuff does
-					InBuildTarget.GlobalCompileEnvironment.Config.Definitions.Add("UE_BUILD_SHIPPING=1");
+					GlobalCompileEnvironment.Config.Definitions.Add("NDEBUG=1"); // the engine doesn't use this, but lots of 3rd party stuff does
+					GlobalCompileEnvironment.Config.Definitions.Add("UE_BUILD_SHIPPING=1");
 					break;
 				case UnrealTargetConfiguration.Test:
 					CompileConfiguration = CPPTargetConfiguration.Shipping;
-					InBuildTarget.GlobalCompileEnvironment.Config.Definitions.Add("NDEBUG=1"); // the engine doesn't use this, but lots of 3rd party stuff does
-					InBuildTarget.GlobalCompileEnvironment.Config.Definitions.Add("UE_BUILD_TEST=1");
+					GlobalCompileEnvironment.Config.Definitions.Add("NDEBUG=1"); // the engine doesn't use this, but lots of 3rd party stuff does
+					GlobalCompileEnvironment.Config.Definitions.Add("UE_BUILD_TEST=1");
 					break;
 			}
 
 			// Set up the global C++ compilation and link environment.
-			InBuildTarget.GlobalCompileEnvironment.Config.Target.Configuration = CompileConfiguration;
-			InBuildTarget.GlobalLinkEnvironment.Config.Target.Configuration = CompileConfiguration;
+			GlobalCompileEnvironment.Config.Target.Configuration = CompileConfiguration;
+			GlobalLinkEnvironment.Config.Target.Configuration = CompileConfiguration;
 
 			// Create debug info based on the heuristics specified by the user.
-			InBuildTarget.GlobalCompileEnvironment.Config.bCreateDebugInfo =
-				!BuildConfiguration.bDisableDebugInfo && ShouldCreateDebugInfo(InBuildTarget.Platform, CheckConfig);
+			GlobalCompileEnvironment.Config.bCreateDebugInfo =
+				!BuildConfiguration.bDisableDebugInfo && ShouldCreateDebugInfo(Target.Platform, CheckConfig);
 
 			// NOTE: Even when debug info is turned off, we currently force the linker to generate debug info
 			//       anyway on Visual C++ platforms.  This will cause a PDB file to be generated with symbols
@@ -756,7 +756,7 @@ namespace UnrealBuildTool
 			//       useful call stacks, even though compiler-generate debug info may be disabled.  This gives
 			//       us much of the build-time savings of fully-disabled debug info, without giving up call
 			//       data completely.
-			InBuildTarget.GlobalLinkEnvironment.Config.bCreateDebugInfo = true;
+			GlobalLinkEnvironment.Config.bCreateDebugInfo = true;
 		}
 
 		/// <summary>
