Move setting for a target's preferred sub-platform into UEBuildConfiguration, rather than the target rules. Saves having to do hacky stuff like command line parsing from inside the build platform.

diff --git a/Engine/Source/Programs/UnrealBuildTool/Configuration/UEBuildConfiguration.cs b/Engine/Source/Programs/UnrealBuildTool/Configuration/UEBuildConfiguration.cs
--- a/Engine/Source/Programs/UnrealBuildTool/Configuration/UEBuildConfiguration.cs
+++ b/Engine/Source/Programs/UnrealBuildTool/Configuration/UEBuildConfiguration.cs
@@ -284,6 +284,11 @@ namespace UnrealBuildTool
 		[XmlConfig]
 		public static bool bCompileCEF3;
 
+		/// <summary>
+		/// Allow a target to specify a preferred sub-platform. Can be used to target a build using sub platform specifics.
+		/// </summary>
+		public static string PreferredSubPlatform = "";
+
 		/// <summary>
 		/// Sets the configuration back to defaults.
 		/// </summary>
@@ -324,6 +329,7 @@ namespace UnrealBuildTool
 			bAllowHotReloadFromIDE = true;
 			bSkipLinkingWhenNothingToCompile = false;
 			bCompileCEF3 = true;
+			PreferredSubPlatform = "";
 		}
 
 		/// <summary>
diff --git a/Engine/Source/Programs/UnrealBuildTool/System/RulesCompiler.cs b/Engine/Source/Programs/UnrealBuildTool/System/RulesCompiler.cs
--- a/Engine/Source/Programs/UnrealBuildTool/System/RulesCompiler.cs
+++ b/Engine/Source/Programs/UnrealBuildTool/System/RulesCompiler.cs
@@ -508,13 +508,6 @@ namespace UnrealBuildTool
 		/// </summary>
 		public bool bUseStaticCRT = false;
 
-
-		/// <summary>
-		/// Allow a target to specify a preferred sub-platform.
-		/// Can be used to target a build using sub platform specifics.
-		/// </summary>
-		public string PreferredSubPlatform = String.Empty;
-
 		/// <summary>
 		/// By default we use the Release C++ Runtime (CRT), even when compiling Debug builds.  This is because the Debug C++
 		/// Runtime isn't very useful when debugging Unreal Engine projects, and linking against the Debug CRT libraries forces
diff --git a/Engine/Source/Programs/UnrealBuildTool/System/UnrealBuildTool.cs b/Engine/Source/Programs/UnrealBuildTool/System/UnrealBuildTool.cs
--- a/Engine/Source/Programs/UnrealBuildTool/System/UnrealBuildTool.cs
+++ b/Engine/Source/Programs/UnrealBuildTool/System/UnrealBuildTool.cs
@@ -857,6 +857,10 @@ namespace UnrealBuildTool
 				{
 					UEBuildConfiguration.bCompileCEF3 = false;
 				}
+				else if (LowercaseArg == "-winxp")
+				{
+					UEBuildConfiguration.PreferredSubPlatform = "WindowsXP";
+				}
 				else if (LowercaseArg == "-rununrealcodeanalyzer")
 				{
 					BuildConfiguration.bRunUnrealCodeAnalyzer = true;
diff --git a/Engine/Source/Programs/UnrealBuildTool/Windows/UEBuildWindows.cs b/Engine/Source/Programs/UnrealBuildTool/Windows/UEBuildWindows.cs
--- a/Engine/Source/Programs/UnrealBuildTool/Windows/UEBuildWindows.cs
+++ b/Engine/Source/Programs/UnrealBuildTool/Windows/UEBuildWindows.cs
@@ -415,22 +415,18 @@ namespace UnrealBuildTool
 			string[] CmdLine = Environment.GetCommandLineArgs();
 
 			bool SupportWindowsXPIfAvailable = false;
-			SupportWindowsXPIfAvailable = CmdLine.Contains("-winxp", StringComparer.InvariantCultureIgnoreCase);
 
 			// ...check if it was supported from a config.
-			if (!SupportWindowsXPIfAvailable)
+			ConfigCacheIni Ini = ConfigCacheIni.CreateConfigCacheIni(UnrealTargetPlatform.Win64, "Engine", UnrealBuildTool.GetUProjectPath());
+			string MinimumOS;
+			if (Ini.GetString("/Script/WindowsTargetPlatform.WindowsTargetSettings", "MinimumOSVersion", out MinimumOS))
 			{
-				ConfigCacheIni Ini = ConfigCacheIni.CreateConfigCacheIni(UnrealTargetPlatform.Win64, "Engine", UnrealBuildTool.GetUProjectPath());
-				string MinimumOS;
-				if (Ini.GetString("/Script/WindowsTargetPlatform.WindowsTargetSettings", "MinimumOSVersion", out MinimumOS))
+				if (string.IsNullOrEmpty(MinimumOS) == false)
 				{
-					if (string.IsNullOrEmpty(MinimumOS) == false)
-					{
-						SupportWindowsXPIfAvailable = MinimumOS == "MSOS_XP";
-					}
+					SupportWindowsXPIfAvailable = MinimumOS == "MSOS_XP";
 				}
 			}
-
+	
 			SupportWindowsXP = SupportWindowsXPIfAvailable;
 		}
 
@@ -535,13 +531,9 @@ namespace UnrealBuildTool
 			InBuildTarget.GlobalCompileEnvironment.Config.Definitions.Add("WIN32=1");
 
 			// Should we enable Windows XP support
-			if ((InBuildTarget.TargetType == TargetRules.TargetType.Program) && (GetCPPTargetPlatform(InBuildTarget.Platform) == CPPTargetPlatform.Win32))
+			if(UEBuildConfiguration.PreferredSubPlatform.Equals("WindowsXP", StringComparison.InvariantCultureIgnoreCase))
 			{
-				// Check if the target has requested XP support.
-				if (String.Equals(InBuildTarget.Rules.PreferredSubPlatform, "WindowsXP", StringComparison.InvariantCultureIgnoreCase))
-				{
-					SupportWindowsXP = true;
-				}
+				SupportWindowsXP = true;
 			}
 
 			if (WindowsPlatform.bUseWindowsSDK10 && WindowsPlatform.Compiler == WindowsCompiler.VisualStudio2015)
diff --git a/Engine/Source/Programs/Windows/BootstrapPackagedGame/BootstrapPackagedGame.Target.cs b/Engine/Source/Programs/Windows/BootstrapPackagedGame/BootstrapPackagedGame.Target.cs
--- a/Engine/Source/Programs/Windows/BootstrapPackagedGame/BootstrapPackagedGame.Target.cs
+++ b/Engine/Source/Programs/Windows/BootstrapPackagedGame/BootstrapPackagedGame.Target.cs
@@ -10,11 +10,6 @@ public class BootstrapPackagedGameTarget : TargetRules
 		Type = TargetType.Program;
 
 		bUseStaticCRT = true;
-
-		if (Target.Platform == UnrealTargetPlatform.Win32)
-		{
-			PreferredSubPlatform = "WindowsXP";
-		}
 	}
 
 	//
@@ -54,6 +49,11 @@ public class BootstrapPackagedGameTarget : TargetRules
 		UEBuildConfiguration.bBuildWithEditorOnlyData = false;
 		UEBuildConfiguration.bCompileAgainstEngine = false;
 		UEBuildConfiguration.bCompileAgainstCoreUObject = false;
+
+		if (Target.Platform == UnrealTargetPlatform.Win32)
+		{
+			UEBuildConfiguration.PreferredSubPlatform = "WindowsXP";
+		}
 	}
 	
     public override bool GUBP_AlwaysBuildWithTools(UnrealTargetPlatform InHostPlatform, out bool bInternalToolOnly, out bool SeparateNode, out bool CrossCompile)
