Merge UEBuildPlatform.GetExtraModules() and UEBuildPlatform.SetupBinaries(), because they were being used for the same purpose.

diff --git a/Engine/Source/Programs/UnrealBuildTool/Android/UEBuildAndroid.cs b/Engine/Source/Programs/UnrealBuildTool/Android/UEBuildAndroid.cs
--- a/Engine/Source/Programs/UnrealBuildTool/Android/UEBuildAndroid.cs
+++ b/Engine/Source/Programs/UnrealBuildTool/Android/UEBuildAndroid.cs
@@ -292,7 +292,7 @@ namespace UnrealBuildTool
 			return true;
 		}
 
-		public override void GetExtraModules(TargetInfo Target, List<string> PlatformExtraModules)
+		public override void AddExtraModules(TargetInfo Target, List<string> PlatformExtraModules)
 		{
 		}
 
diff --git a/Engine/Source/Programs/UnrealBuildTool/Configuration/UEBuildClient.cs b/Engine/Source/Programs/UnrealBuildTool/Configuration/UEBuildClient.cs
--- a/Engine/Source/Programs/UnrealBuildTool/Configuration/UEBuildClient.cs
+++ b/Engine/Source/Programs/UnrealBuildTool/Configuration/UEBuildClient.cs
@@ -56,9 +56,6 @@ namespace UnrealBuildTool
 				// Modules should properly identify the 'extra modules' they need now.
 				// There should be nothing here!
 			}
-
-			// Allow the platform to setup binaries
-			UEBuildPlatform.GetBuildPlatform(Platform).SetupBinaries(this);
 		}
 
 		public override void SetupDefaultGlobalEnvironment(
diff --git a/Engine/Source/Programs/UnrealBuildTool/Configuration/UEBuildEditor.cs b/Engine/Source/Programs/UnrealBuildTool/Configuration/UEBuildEditor.cs
--- a/Engine/Source/Programs/UnrealBuildTool/Configuration/UEBuildEditor.cs
+++ b/Engine/Source/Programs/UnrealBuildTool/Configuration/UEBuildEditor.cs
@@ -51,9 +51,6 @@ namespace UnrealBuildTool
 				// Modules should properly identify the 'extra modules' they need now.
 				// There should be nothing here!
 			}
-
-			// Allow the platform to setup binaries
-			UEBuildPlatform.GetBuildPlatform(Platform).SetupBinaries(this);
 		}
 
 		public override void SetupDefaultGlobalEnvironment(
diff --git a/Engine/Source/Programs/UnrealBuildTool/Configuration/UEBuildPlatform.cs b/Engine/Source/Programs/UnrealBuildTool/Configuration/UEBuildPlatform.cs
--- a/Engine/Source/Programs/UnrealBuildTool/Configuration/UEBuildPlatform.cs
+++ b/Engine/Source/Programs/UnrealBuildTool/Configuration/UEBuildPlatform.cs
@@ -338,10 +338,9 @@ namespace UnrealBuildTool
 		/// Get a list of extra modules the platform requires.
 		/// This is to allow undisclosed platforms to add modules they need without exposing information about the platform.
 		/// </summary>
-		/// <param name="Target">     The target being build</param>
-		/// <param name="BuildTarget">    The UEBuildTarget getting build</param>
-		/// <param name="PlatformExtraModules"> OUTPUT the list of extra modules the platform needs to add to the target</param>
-		public virtual void GetExtraModules(TargetInfo Target, List<string> PlatformExtraModules)
+		/// <param name="Target">The target being build</param>
+		/// <param name="ExtraModuleNames">List of extra modules the platform needs to add to the target</param>
+		public virtual void AddExtraModules(TargetInfo Target, List<string> ExtraModuleNames)
 		{
 		}
 
@@ -638,14 +637,6 @@ namespace UnrealBuildTool
 			return Configurations;
 		}
 
-		/// <summary>
-		/// Setup the binaries for this specific platform.
-		/// </summary>
-		/// <param name="InBuildTarget"> The target being built</param>
-		public virtual void SetupBinaries(UEBuildTarget InBuildTarget)
-		{
-		}
-
 		protected static bool DoProjectSettingsMatchDefault(UnrealTargetPlatform Platform, DirectoryReference ProjectDirectoryName, string Section, string[] BoolKeys, string[] IntKeys, string[] StringKeys)
 		{
 			ConfigCacheIni ProjIni = ConfigCacheIni.CreateConfigCacheIni(Platform, "Engine", ProjectDirectoryName);
diff --git a/Engine/Source/Programs/UnrealBuildTool/Configuration/UEBuildServer.cs b/Engine/Source/Programs/UnrealBuildTool/Configuration/UEBuildServer.cs
--- a/Engine/Source/Programs/UnrealBuildTool/Configuration/UEBuildServer.cs
+++ b/Engine/Source/Programs/UnrealBuildTool/Configuration/UEBuildServer.cs
@@ -50,9 +50,6 @@ namespace UnrealBuildTool
 				// Modules should properly identify the 'extra modules' they need now.
 				// There should be nothing here!
 			}
-
-			// Allow the platform to setup binaries
-			UEBuildPlatform.GetBuildPlatform(Platform).SetupBinaries(this);
 		}
 
 		public override void SetupDefaultGlobalEnvironment(
diff --git a/Engine/Source/Programs/UnrealBuildTool/Configuration/UEBuildTarget.cs b/Engine/Source/Programs/UnrealBuildTool/Configuration/UEBuildTarget.cs
--- a/Engine/Source/Programs/UnrealBuildTool/Configuration/UEBuildTarget.cs
+++ b/Engine/Source/Programs/UnrealBuildTool/Configuration/UEBuildTarget.cs
@@ -2242,9 +2242,6 @@ namespace UnrealBuildTool
 				AddPlugin(BuildPlugin);
 			}
 
-			// Allow the platform to setup binaries/plugins/modules
-			UEBuildPlatform.GetBuildPlatform(Platform).SetupBinaries(this);
-
 			// Describe what's being built.
 			Log.TraceVerbose("Building {0} - {1} - {2} - {3}", AppName, TargetName, Platform, Configuration);
 
@@ -3050,7 +3047,7 @@ namespace UnrealBuildTool
 		{
 			var BuildPlatform = UEBuildPlatform.GetBuildPlatform(Platform);
 			List<string> PlatformExtraModules = new List<string>();
-			BuildPlatform.GetExtraModules(TargetInfo, PlatformExtraModules);
+			BuildPlatform.AddExtraModules(TargetInfo, PlatformExtraModules);
 			ExtraModuleNames.AddRange(PlatformExtraModules);
 		}
 
diff --git a/Engine/Source/Programs/UnrealBuildTool/HTML5/UEBuildHTML5.cs b/Engine/Source/Programs/UnrealBuildTool/HTML5/UEBuildHTML5.cs
--- a/Engine/Source/Programs/UnrealBuildTool/HTML5/UEBuildHTML5.cs
+++ b/Engine/Source/Programs/UnrealBuildTool/HTML5/UEBuildHTML5.cs
@@ -360,15 +360,6 @@ namespace UnrealBuildTool
 			};
 		}
 
-		/// <summary>
-		/// Setup the binaries for this specific platform.
-		/// </summary>
-		/// <param name="InBuildTarget"> The target being built</param>
-		public override void SetupBinaries(UEBuildTarget InBuildTarget)
-		{
-			// 			InBuildTarget.ExtraModuleNames.Add("ES2RHI");
-		}
-
 		/// <summary>
 		/// Creates a toolchain instance for the given platform.
 		/// </summary>
diff --git a/Engine/Source/Programs/UnrealBuildTool/Linux/UEBuildLinux.cs b/Engine/Source/Programs/UnrealBuildTool/Linux/UEBuildLinux.cs
--- a/Engine/Source/Programs/UnrealBuildTool/Linux/UEBuildLinux.cs
+++ b/Engine/Source/Programs/UnrealBuildTool/Linux/UEBuildLinux.cs
@@ -259,16 +259,6 @@ namespace UnrealBuildTool
 			return true;
 		}
 
-		/// <summary>
-		/// Get a list of extra modules the platform requires.
-		/// This is to allow undisclosed platforms to add modules they need without exposing information about the platfomr.
-		/// </summary>
-		/// <param name="Target">     The target being build</param>
-		/// <param name="PlatformExtraModules"> OUTPUT the list of extra modules the platform needs to add to the target</param>
-		public override void GetExtraModules(TargetInfo Target, List<string> PlatformExtraModules)
-		{
-		}
-
 		/// <summary>
 		/// Modify the newly created module passed in for this platform.
 		/// This is not required - but allows for hiding details of a
@@ -385,14 +375,6 @@ namespace UnrealBuildTool
 			};
 		}
 
-		/// <summary>
-		/// Setup the binaries for this specific platform.
-		/// </summary>
-		/// <param name="InBuildTarget"> The target being built</param>
-		public override void SetupBinaries(UEBuildTarget InBuildTarget)
-		{
-		}
-
 		/// <summary>
 		/// Creates a toolchain instance for the given platform.
 		/// </summary>
diff --git a/Engine/Source/Programs/UnrealBuildTool/WinRT/UEBuildWinRT.cs b/Engine/Source/Programs/UnrealBuildTool/WinRT/UEBuildWinRT.cs
--- a/Engine/Source/Programs/UnrealBuildTool/WinRT/UEBuildWinRT.cs
+++ b/Engine/Source/Programs/UnrealBuildTool/WinRT/UEBuildWinRT.cs
@@ -217,13 +217,14 @@ namespace UnrealBuildTool
 		/// Get a list of extra modules the platform requires.
 		/// This is to allow undisclosed platforms to add modules they need without exposing information about the platfomr.
 		/// </summary>
-		/// <param name="Target">     The target being build</param>
-		/// <param name="PlatformExtraModules"> OUTPUT the list of extra modules the platform needs to add to the target</param>
-		public override void GetExtraModules(TargetInfo Target, List<string> PlatformExtraModules)
+		/// <param name="Target">The target being build</param>
+		/// <param name="PlatformExtraModules">List of extra modules the platform needs to add to the target</param>
+		public override void AddExtraModules(TargetInfo Target, List<string> ExtraModuleNames)
 		{
+			ExtraModuleNames.Add("XInput");
 			if (Target.Platform == UnrealTargetPlatform.WinRT)
 			{
-				PlatformExtraModules.Add("XAudio2");
+				ExtraModuleNames.Add("XAudio2");
 			}
 		}
 
@@ -614,15 +615,6 @@ namespace UnrealBuildTool
 			};
 		}
 
-		/// <summary>
-		/// Setup the binaries for this specific platform.
-		/// </summary>
-		/// <param name="InBuildTarget"> The target being built</param>
-		public override void SetupBinaries(UEBuildTarget InBuildTarget)
-		{
-			InBuildTarget.ExtraModuleNames.Add("XInput");
-		}
-
 		//
 		// WinRT specific functions
 		//
