Delete the dependency cache file if it's corrupt or versioned incorrectly, and can't be read.

diff --git a/Engine/Source/Programs/UnrealBuildTool/System/DependencyCache.cs b/Engine/Source/Programs/UnrealBuildTool/System/DependencyCache.cs
--- a/Engine/Source/Programs/UnrealBuildTool/System/DependencyCache.cs
+++ b/Engine/Source/Programs/UnrealBuildTool/System/DependencyCache.cs
@@ -150,6 +150,7 @@ namespace UnrealBuildTool
 			catch (Exception Ex)
 			{
 				Console.Error.WriteLine("Failed to read dependency cache: {0}", Ex.Message);
+				CacheFile.Delete();
 			}
 			return Result;
 		}
diff --git a/Engine/Source/Programs/UnrealBuildTool/System/FlatCPPIncludeDepencencyCache.cs b/Engine/Source/Programs/UnrealBuildTool/System/FlatCPPIncludeDepencencyCache.cs
--- a/Engine/Source/Programs/UnrealBuildTool/System/FlatCPPIncludeDepencencyCache.cs
+++ b/Engine/Source/Programs/UnrealBuildTool/System/FlatCPPIncludeDepencencyCache.cs
@@ -128,6 +128,7 @@ namespace UnrealBuildTool
 			catch (Exception Ex)
 			{
 				Console.Error.WriteLine("Failed to read FlatCPPIncludeDependencyCache: {0}", Ex.Message);
+				BackingFile.Delete();
 			}
 			return Result;
 		}
