Send a telemetry event with the current machine's XGE settings.

diff --git a/Engine/Source/Programs/UnrealBuildTool/System/XGE.cs b/Engine/Source/Programs/UnrealBuildTool/System/XGE.cs
--- a/Engine/Source/Programs/UnrealBuildTool/System/XGE.cs
+++ b/Engine/Source/Programs/UnrealBuildTool/System/XGE.cs
@@ -8,6 +8,7 @@ using System.Xml;
 using System.Text.RegularExpressions;
 using System.Linq;
 using System.Reflection;
+using Microsoft.Win32;
 
 namespace UnrealBuildTool
 {
@@ -54,6 +55,29 @@ namespace UnrealBuildTool
 					// Try to execute the XGE tasks, and if XGE is available, skip the local execution fallback.
 					if (Telemetry.IsAvailable())
 					{
+			            try
+			            {
+				            const string BuilderKey = "HKEY_LOCAL_MACHINE\\SOFTWARE\\Wow6432Node\\Xoreax\\IncrediBuild\\Builder";
+            
+				            string CPUUtilization = Registry.GetValue(BuilderKey, "ForceCPUCount", "").ToString();
+				            string AvoidTaskExecutionOnLocalMachine = Registry.GetValue(BuilderKey, "AvoidLocalExec", "").ToString();
+				            string RestartRemoteProcessesOnLocalMachine = Registry.GetValue(BuilderKey, "AllowDoubleTargets", "").ToString();
+				            string LimitMaxNumberOfCores = Registry.GetValue(BuilderKey, "MaxHelpers", "").ToString();
+				            string WriteOutputToDiskInBackground = Registry.GetValue(BuilderKey, "LazyOutputWriter_Beta", "").ToString();
+				            string MaxConcurrentPDBs = Registry.GetValue(BuilderKey, "MaxConcurrentPDBs", "").ToString();
+            
+				            Telemetry.SendEvent("XGESettings.1",
+					            "CPUUtilization", CPUUtilization,
+					            "AvoidTaskExecutionOnLocalMachine", AvoidTaskExecutionOnLocalMachine,
+					            "RestartRemoteProcessesOnLocalMachine", RestartRemoteProcessesOnLocalMachine,
+					            "LimitMaxNumberOfCores", LimitMaxNumberOfCores,
+					            "WriteOutputToDiskInBackground", WriteOutputToDiskInBackground,
+					            "MaxConcurrentPDBs", MaxConcurrentPDBs);
+			            }
+			            catch
+			            {
+			            }
+
 						// Add a custom output handler to determine the build duration for each task and map it back to the action that generated it.
 						XGEResult = XGE.ExecuteTaskFileWithProgressMarkup(XGETaskFilePath, Actions.Count, (sender, args) =>
 							{
