Add 'Enabled as helper' setting to XGE telemetry.

diff --git a/Engine/Source/Programs/UnrealBuildTool/System/XGE.cs b/Engine/Source/Programs/UnrealBuildTool/System/XGE.cs
--- a/Engine/Source/Programs/UnrealBuildTool/System/XGE.cs
+++ b/Engine/Source/Programs/UnrealBuildTool/System/XGE.cs
@@ -65,14 +65,16 @@ namespace UnrealBuildTool
 				            string LimitMaxNumberOfCores = Registry.GetValue(BuilderKey, "MaxHelpers", "").ToString();
 				            string WriteOutputToDiskInBackground = Registry.GetValue(BuilderKey, "LazyOutputWriter_Beta", "").ToString();
 				            string MaxConcurrentPDBs = Registry.GetValue(BuilderKey, "MaxConcurrentPDBs", "").ToString();
+							string EnabledAsHelper = Registry.GetValue(BuilderKey, "LastEnabled", "").ToString();
             
-				            Telemetry.SendEvent("XGESettings.1",
+				            Telemetry.SendEvent("XGESettings.2",
 					            "CPUUtilization", CPUUtilization,
 					            "AvoidTaskExecutionOnLocalMachine", AvoidTaskExecutionOnLocalMachine,
 					            "RestartRemoteProcessesOnLocalMachine", RestartRemoteProcessesOnLocalMachine,
 					            "LimitMaxNumberOfCores", LimitMaxNumberOfCores,
 					            "WriteOutputToDiskInBackground", WriteOutputToDiskInBackground,
-					            "MaxConcurrentPDBs", MaxConcurrentPDBs);
+					            "MaxConcurrentPDBs", MaxConcurrentPDBs,
+								"EnabledAsHelper", EnabledAsHelper);
 			            }
 			            catch
 			            {
