Instance UEBuildPlatform for each supported platform so we don't have to keep passing an UnrealTargetPlatform back into it, and separate the SDK logic for each group of platforms into a separate class (UEBuildPlatformSDK). Add a new factory class (UEBuildPlatformFactory) for creating them and binding them together.

diff --git a/Engine/Source/Programs/AutomationTool/AutomationUtils/Automation.cs b/Engine/Source/Programs/AutomationTool/AutomationUtils/Automation.cs
--- a/Engine/Source/Programs/AutomationTool/AutomationUtils/Automation.cs
+++ b/Engine/Source/Programs/AutomationTool/AutomationUtils/Automation.cs
@@ -332,7 +332,7 @@ AutomationTool.exe [-verbose] [-compileonly] [-p4] Command0 [-Arg0 -Arg1 -Arg2 
 			// Disable AutoSDKs if specified on the command line
 			if (GlobalCommandLine.NoAutoSDK)
 			{
-				UEBuildPlatform.bAllowAutoSDKSwitching = false;
+				UEBuildPlatformSDK.bAllowAutoSDKSwitching = false;
 			}
 
 			// Setup environment
diff --git a/Engine/Source/Programs/AutomationTool/AutomationUtils/CommandUtils.cs b/Engine/Source/Programs/AutomationTool/AutomationUtils/CommandUtils.cs
--- a/Engine/Source/Programs/AutomationTool/AutomationUtils/CommandUtils.cs
+++ b/Engine/Source/Programs/AutomationTool/AutomationUtils/CommandUtils.cs
@@ -1841,12 +1841,7 @@ namespace AutomationTool
 			{
 				if (UBTTargetPlatforms == null || UBTTargetPlatforms.Length == 0)
 				{
-					UBTTargetPlatforms = new UnrealBuildTool.UnrealTargetPlatform[UnrealBuildTool.UEBuildPlatform.BuildPlatformDictionary.Count];
-					int Index = 0;
-					foreach (var Platform in UnrealBuildTool.UEBuildPlatform.BuildPlatformDictionary)
-					{
-						UBTTargetPlatforms[Index++] = Platform.Key;
-					}
+					UBTTargetPlatforms = UnrealBuildTool.UEBuildPlatform.GetRegisteredPlatforms().ToArray();
 				}
 				return UBTTargetPlatforms;
 			}
diff --git a/Engine/Source/Programs/AutomationTool/AutomationUtils/UE4Build.cs b/Engine/Source/Programs/AutomationTool/AutomationUtils/UE4Build.cs
--- a/Engine/Source/Programs/AutomationTool/AutomationUtils/UE4Build.cs
+++ b/Engine/Source/Programs/AutomationTool/AutomationUtils/UE4Build.cs
@@ -1175,7 +1175,7 @@ namespace AutomationTool
 
 		public bool CanUseXGE(UnrealBuildTool.UnrealTargetPlatform Platform)
 		{
-			if (!UnrealBuildTool.UEBuildPlatform.BuildPlatformDictionary.ContainsKey(Platform))
+			if (!UnrealBuildTool.UEBuildPlatform.IsPlatformAvailable(Platform))
 			{
 				return false;
 			}
diff --git a/Engine/Source/Programs/UnrealBuildTool/Android/AndroidToolChain.cs b/Engine/Source/Programs/UnrealBuildTool/Android/AndroidToolChain.cs
--- a/Engine/Source/Programs/UnrealBuildTool/Android/AndroidToolChain.cs
+++ b/Engine/Source/Programs/UnrealBuildTool/Android/AndroidToolChain.cs
@@ -908,7 +908,7 @@ namespace UnrealBuildTool
 			var NDKRoot = Environment.GetEnvironmentVariable("NDKROOT").Replace("\\", "/");
 
 			string BasePCHName = "";
-			var PCHExtension = UEBuildPlatform.BuildPlatformDictionary[UnrealTargetPlatform.Android].GetBinaryExtension(UEBuildBinaryType.PrecompiledHeader);
+			var PCHExtension = UEBuildPlatform.GetBuildPlatform(UnrealTargetPlatform.Android).GetBinaryExtension(UEBuildBinaryType.PrecompiledHeader);
 			if (CompileEnvironment.Config.PrecompiledHeaderAction == PrecompiledHeaderAction.Include)
 			{
 				BasePCHName = RemoveArchName(CompileEnvironment.PrecompiledHeaderFile.AbsolutePath).Replace(PCHExtension, "");
@@ -1039,7 +1039,7 @@ namespace UnrealBuildTool
 								CompileAction.PrerequisiteItems.Add(ArchPrecompiledHeaderFile);
 							}
 
-							var ObjectFileExtension = UEBuildPlatform.BuildPlatformDictionary[UnrealTargetPlatform.Android].GetBinaryExtension(UEBuildBinaryType.Object);
+							var ObjectFileExtension = UEBuildPlatform.GetBuildPlatform(UnrealTargetPlatform.Android).GetBinaryExtension(UEBuildBinaryType.Object);
 
 							// Add the object file to the produced item list.
 							FileItem ObjectFile = FileItem.GetItemByFileReference(
diff --git a/Engine/Source/Programs/UnrealBuildTool/Android/UEBuildAndroid.cs b/Engine/Source/Programs/UnrealBuildTool/Android/UEBuildAndroid.cs
--- a/Engine/Source/Programs/UnrealBuildTool/Android/UEBuildAndroid.cs
+++ b/Engine/Source/Programs/UnrealBuildTool/Android/UEBuildAndroid.cs
@@ -12,6 +12,18 @@
 {
 	class AndroidPlatform : UEBuildPlatform
 	{
+		AndroidPlatformSDK SDK;
+
+		public AndroidPlatform(AndroidPlatformSDK InSDK) : base(UnrealTargetPlatform.Android)
+		{
+			SDK = InSDK;
+		}
+
+		public override SDKStatus HasRequiredSDKsInstalled()
+		{
+			return SDK.HasRequiredSDKsInstalled();
+		}
+
 		// The current architecture - affects everything about how UBT operates on Android
 		public override string GetActiveArchitecture()
 		{
@@ -24,161 +36,4 @@
 			return false;
 		}
 
-		protected override bool PlatformSupportsAutoSDKs()
-		{
-			return true;
-		}
-
-		public override string GetSDKTargetPlatformName()
-		{
-			return "Android";
-		}
-
-		protected override string GetRequiredSDKString()
-		{
-			return "-19";
-		}
-
-		protected override String GetRequiredScriptVersionString()
-		{
-			return "3.0";
-		}
-
-		// prefer auto sdk on android as correct 'manual' sdk detection isn't great at the moment.
-		protected override bool PreferAutoSDK()
-		{
-			return true;
-		}
-
-		/// <summary>
-		/// checks if the sdk is installed or has been synced
-		/// </summary>
-		/// <returns></returns>
-		private bool HasAnySDK()
-		{
-			string NDKPath = Environment.GetEnvironmentVariable("NDKROOT");
-			{
-				var configCacheIni = ConfigCacheIni.CreateConfigCacheIni(UnrealTargetPlatform.Unknown, "Engine", (DirectoryReference)null);
-				var AndroidEnv = new Dictionary<string, string>();
-
-				Dictionary<string, string> EnvVarNames = new Dictionary<string, string> { 
-                                                         {"ANDROID_HOME", "SDKPath"}, 
-                                                         {"NDKROOT", "NDKPath"}, 
-                                                         {"ANT_HOME", "ANTPath"},
-                                                         {"JAVA_HOME", "JavaPath"}
-                                                         };
-
-				string path;
-				foreach (var kvp in EnvVarNames)
-				{
-					if (configCacheIni.GetPath("/Script/AndroidPlatformEditor.AndroidSDKSettings", kvp.Value, out path) && !string.IsNullOrEmpty(path))
-					{
-						AndroidEnv.Add(kvp.Key, path);
-					}
-					else
-					{
-						var envValue = Environment.GetEnvironmentVariable(kvp.Key);
-						if (!String.IsNullOrEmpty(envValue))
-						{
-							AndroidEnv.Add(kvp.Key, envValue);
-						}
-					}
-				}
-
-				// If we are on Mono and we are still missing a key then go and find it from the .bash_profile
-				if (Utils.IsRunningOnMono && !EnvVarNames.All(s => AndroidEnv.ContainsKey(s.Key)))
-				{
-					string BashProfilePath = Path.Combine(Environment.GetFolderPath(Environment.SpecialFolder.Personal), ".bash_profile");
-					if (File.Exists(BashProfilePath))
-					{
-						string[] BashProfileContents = File.ReadAllLines(BashProfilePath);
-
-						// Walk backwards so we keep the last export setting instead of the first
-						for (int LineIndex = BashProfileContents.Length - 1; LineIndex >= 0; --LineIndex)
-						{
-							foreach (var kvp in EnvVarNames)
-							{
-								if (AndroidEnv.ContainsKey(kvp.Key))
-								{
-									continue;
-								}
-
-								if (BashProfileContents[LineIndex].StartsWith("export " + kvp.Key + "="))
-								{
-									string PathVar = BashProfileContents[LineIndex].Split('=')[1].Replace("\"", "");
-									AndroidEnv.Add(kvp.Key, PathVar);
-								}
-							}
-						}
-					}
-				}
-
-				// Set for the process
-				foreach (var kvp in AndroidEnv)
-				{
-					Environment.SetEnvironmentVariable(kvp.Key, kvp.Value);
-				}
-
-				// See if we have an NDK path now...
-				AndroidEnv.TryGetValue("NDKROOT", out NDKPath);
-			}
-
-			// we don't have an NDKROOT specified
-			if (String.IsNullOrEmpty(NDKPath))
-			{
-				return false;
-			}
-
-			NDKPath = NDKPath.Replace("\"", "");
-
-			// need a supported llvm
-			if (!Directory.Exists(Path.Combine(NDKPath, @"toolchains/llvm-3.6")) &&
-				!Directory.Exists(Path.Combine(NDKPath, @"toolchains/llvm-3.5")) &&
-				!Directory.Exists(Path.Combine(NDKPath, @"toolchains/llvm-3.3")) &&
-				!Directory.Exists(Path.Combine(NDKPath, @"toolchains/llvm-3.1")))
-			{
-				return false;
-			}
-			return true;
-		}
-
-
-		protected override SDKStatus HasRequiredManualSDKInternal()
-		{
-			// if any autosdk setup has been done then the local process environment is suspect
-			if (HasSetupAutoSDK())
-			{
-				return SDKStatus.Invalid;
-			}
-
-			if (HasAnySDK())
-			{
-				return SDKStatus.Valid;
-			}
-
-			return SDKStatus.Invalid;
-		}
-
-		protected override void RegisterBuildPlatformInternal()
-		{
-			if ((ProjectFileGenerator.bGenerateProjectFiles == true) || (HasRequiredSDKsInstalled() == SDKStatus.Valid) || Environment.GetEnvironmentVariable("IsBuildMachine") == "1")
-			{
-				bool bRegisterBuildPlatform = true;
-
-				FileReference AndroidTargetPlatformFile = FileReference.Combine(UnrealBuildTool.EngineSourceDirectory, "Developer", "Android", "AndroidTargetPlatform", "AndroidTargetPlatform.Build.cs");
-				if (AndroidTargetPlatformFile.Exists() == false)
-				{
-					bRegisterBuildPlatform = false;
-				}
-
-				if (bRegisterBuildPlatform == true)
-				{
-					// Register this build platform
-					Log.TraceVerbose("        Registering for {0}", UnrealTargetPlatform.Android.ToString());
-					UEBuildPlatform.RegisterBuildPlatform(UnrealTargetPlatform.Android, this);
-					UEBuildPlatform.RegisterPlatformWithGroup(UnrealTargetPlatform.Android, UnrealPlatformGroup.Android);
-				}
-			}
-		}
-
 		public override CPPTargetPlatform GetCPPTargetPlatform(UnrealTargetPlatform InUnrealTargetPlatform)
@@ -492,4 +347,169 @@
 			return true;
 		}
 	}
+
+	class AndroidPlatformSDK : UEBuildPlatformSDK
+	{
+		protected override bool PlatformSupportsAutoSDKs()
+		{
+			return true;
+		}
+
+		public override string GetSDKTargetPlatformName()
+		{
+			return "Android";
+		}
+
+		protected override string GetRequiredSDKString()
+		{
+			return "-19";
+		}
+
+		protected override String GetRequiredScriptVersionString()
+		{
+			return "3.0";
+		}
+
+		// prefer auto sdk on android as correct 'manual' sdk detection isn't great at the moment.
+		protected override bool PreferAutoSDK()
+		{
+			return true;
+		}
+
+		/// <summary>
+		/// checks if the sdk is installed or has been synced
+		/// </summary>
+		/// <returns></returns>
+		private bool HasAnySDK()
+		{
+			string NDKPath = Environment.GetEnvironmentVariable("NDKROOT");
+			{
+				var configCacheIni = ConfigCacheIni.CreateConfigCacheIni(UnrealTargetPlatform.Unknown, "Engine", (DirectoryReference)null);
+				var AndroidEnv = new Dictionary<string, string>();
+
+				Dictionary<string, string> EnvVarNames = new Dictionary<string, string> { 
+                                                         {"ANDROID_HOME", "SDKPath"}, 
+                                                         {"NDKROOT", "NDKPath"}, 
+                                                         {"ANT_HOME", "ANTPath"},
+                                                         {"JAVA_HOME", "JavaPath"}
+                                                         };
+
+				string path;
+				foreach (var kvp in EnvVarNames)
+				{
+					if (configCacheIni.GetPath("/Script/AndroidPlatformEditor.AndroidSDKSettings", kvp.Value, out path) && !string.IsNullOrEmpty(path))
+					{
+						AndroidEnv.Add(kvp.Key, path);
+					}
+					else
+					{
+						var envValue = Environment.GetEnvironmentVariable(kvp.Key);
+						if (!String.IsNullOrEmpty(envValue))
+						{
+							AndroidEnv.Add(kvp.Key, envValue);
+						}
+					}
+				}
+
+				// If we are on Mono and we are still missing a key then go and find it from the .bash_profile
+				if (Utils.IsRunningOnMono && !EnvVarNames.All(s => AndroidEnv.ContainsKey(s.Key)))
+				{
+					string BashProfilePath = Path.Combine(Environment.GetFolderPath(Environment.SpecialFolder.Personal), ".bash_profile");
+					if (File.Exists(BashProfilePath))
+					{
+						string[] BashProfileContents = File.ReadAllLines(BashProfilePath);
+
+						// Walk backwards so we keep the last export setting instead of the first
+						for (int LineIndex = BashProfileContents.Length - 1; LineIndex >= 0; --LineIndex)
+						{
+							foreach (var kvp in EnvVarNames)
+							{
+								if (AndroidEnv.ContainsKey(kvp.Key))
+								{
+									continue;
+								}
+
+								if (BashProfileContents[LineIndex].StartsWith("export " + kvp.Key + "="))
+								{
+									string PathVar = BashProfileContents[LineIndex].Split('=')[1].Replace("\"", "");
+									AndroidEnv.Add(kvp.Key, PathVar);
+								}
+							}
+						}
+					}
+				}
+
+				// Set for the process
+				foreach (var kvp in AndroidEnv)
+				{
+					Environment.SetEnvironmentVariable(kvp.Key, kvp.Value);
+				}
+
+				// See if we have an NDK path now...
+				AndroidEnv.TryGetValue("NDKROOT", out NDKPath);
+			}
+
+			// we don't have an NDKROOT specified
+			if (String.IsNullOrEmpty(NDKPath))
+			{
+				return false;
+			}
+
+			NDKPath = NDKPath.Replace("\"", "");
+
+			// need a supported llvm
+			if (!Directory.Exists(Path.Combine(NDKPath, @"toolchains/llvm-3.6")) &&
+				!Directory.Exists(Path.Combine(NDKPath, @"toolchains/llvm-3.5")) &&
+				!Directory.Exists(Path.Combine(NDKPath, @"toolchains/llvm-3.3")) &&
+				!Directory.Exists(Path.Combine(NDKPath, @"toolchains/llvm-3.1")))
+			{
+				return false;
+			}
+			return true;
+		}
+
+		protected override SDKStatus HasRequiredManualSDKInternal()
+		{
+			// if any autosdk setup has been done then the local process environment is suspect
+			if (HasSetupAutoSDK())
+			{
+				return SDKStatus.Invalid;
+			}
+
+			if (HasAnySDK())
+			{
+				return SDKStatus.Valid;
+			}
+
+			return SDKStatus.Invalid;
+		}
+	}
+
+	public class AndroidPlatformFactory : UEBuildPlatformFactory
+	{
+		public override void RegisterBuildPlatforms()
+		{
+			AndroidPlatformSDK SDK = new AndroidPlatformSDK();
+			SDK.ManageAndValidateSDK();
+
+			if ((ProjectFileGenerator.bGenerateProjectFiles == true) || (SDK.HasRequiredSDKsInstalled() == SDKStatus.Valid) || Environment.GetEnvironmentVariable("IsBuildMachine") == "1")
+			{
+				bool bRegisterBuildPlatform = true;
+
+				FileReference AndroidTargetPlatformFile = FileReference.Combine(UnrealBuildTool.EngineSourceDirectory, "Developer", "Android", "AndroidTargetPlatform", "AndroidTargetPlatform.Build.cs");
+				if (AndroidTargetPlatformFile.Exists() == false)
+				{
+					bRegisterBuildPlatform = false;
+				}
+
+				if (bRegisterBuildPlatform == true)
+				{
+					// Register this build platform
+					Log.TraceVerbose("        Registering for {0}", UnrealTargetPlatform.Android.ToString());
+					UEBuildPlatform.RegisterBuildPlatform(new AndroidPlatform(SDK));
+					UEBuildPlatform.RegisterPlatformWithGroup(UnrealTargetPlatform.Android, UnrealPlatformGroup.Android);
+				}
+			}
+		}
+	}
 }
diff --git a/Engine/Source/Programs/UnrealBuildTool/Configuration/UEBuildPlatform.cs b/Engine/Source/Programs/UnrealBuildTool/Configuration/UEBuildPlatform.cs
--- a/Engine/Source/Programs/UnrealBuildTool/Configuration/UEBuildPlatform.cs
+++ b/Engine/Source/Programs/UnrealBuildTool/Configuration/UEBuildPlatform.cs
@@ -16,13 +16,38 @@ namespace UnrealBuildTool
 
 	public abstract partial class UEBuildPlatform
 	{
-		public static Dictionary<UnrealTargetPlatform, UEBuildPlatform> BuildPlatformDictionary = new Dictionary<UnrealTargetPlatform, UEBuildPlatform>();
+		private static Dictionary<UnrealTargetPlatform, UEBuildPlatform> BuildPlatformDictionary = new Dictionary<UnrealTargetPlatform, UEBuildPlatform>();
 
 		// a mapping of a group to the platforms in the group (ie, Microsoft contains Win32 and Win64)
 		static Dictionary<UnrealPlatformGroup, List<UnrealTargetPlatform>> PlatformGroupDictionary = new Dictionary<UnrealPlatformGroup, List<UnrealTargetPlatform>>();
 
 		private bool bInitializedProject = false;
 
+		public readonly UnrealTargetPlatform Platform;
+
+		/// <summary>
+		/// Constructor.
+		/// </summary>
+		/// <param name="InPlatform">The enum value for this platform</param>
+		public UEBuildPlatform(UnrealTargetPlatform InPlatform)
+		{
+			Platform = InPlatform;
+		}
+
+		/// <summary>
+		/// Whether the required external SDKs are installed for this platform. Could be either a manual install or an AutoSDK.
+		/// </summary>
+		public abstract SDKStatus HasRequiredSDKsInstalled();
+
+		/// <summary>
+		/// Gets all the registered platforms
+		/// </summary>
+		/// <returns>Sequence of registered platforms</returns>
+		public static IEnumerable<UnrealTargetPlatform> GetRegisteredPlatforms()
+		{
+			return BuildPlatformDictionary.Keys;
+		}
+
 		/// <summary>
 		/// Attempt to convert a string to an UnrealTargetPlatform enum entry
 		/// </summary>
@@ -49,33 +74,32 @@ namespace UnrealBuildTool
 			return UnrealTargetPlatform.Unknown;
 		}
 
+		/// <summary>
+		/// Determines whether a given platform is available
+		/// </summary>
+		/// <param name="Platform">The platform to check for</param>
+		/// <returns>True if it's available, false otherwise</returns>
+		public static bool IsPlatformAvailable(UnrealTargetPlatform Platform)
+		{
+			return BuildPlatformDictionary.ContainsKey(Platform);
+		}
+
 		/// <summary>
 		/// Register the given platforms UEBuildPlatform instance
 		/// </summary>
 		/// <param name="InPlatform">  The UnrealTargetPlatform to register with</param>
 		/// <param name="InBuildPlatform"> The UEBuildPlatform instance to use for the InPlatform</param>
-		public static void RegisterBuildPlatform(UnrealTargetPlatform InPlatform, UEBuildPlatform InBuildPlatform)
+		public static void RegisterBuildPlatform(UEBuildPlatform InBuildPlatform)
 		{
-			if (BuildPlatformDictionary.ContainsKey(InPlatform) == true)
+			if (BuildPlatformDictionary.ContainsKey(InBuildPlatform.Platform) == true)
 			{
 				Log.TraceWarning("RegisterBuildPlatform Warning: Registering build platform {0} for {1} when it is already set to {2}",
-					InBuildPlatform.ToString(), InPlatform.ToString(), BuildPlatformDictionary[InPlatform].ToString());
-				BuildPlatformDictionary[InPlatform] = InBuildPlatform;
+					InBuildPlatform.ToString(), InBuildPlatform.Platform.ToString(), BuildPlatformDictionary[InBuildPlatform.Platform].ToString());
+				BuildPlatformDictionary[InBuildPlatform.Platform] = InBuildPlatform;
 			}
 			else
 			{
-				BuildPlatformDictionary.Add(InPlatform, InBuildPlatform);
-			}
-		}
-
-		/// <summary>
-		/// Unregister the given platform
-		/// </summary>
-		public static void UnregisterBuildPlatform(UnrealTargetPlatform InPlatform)
-		{
-			if (BuildPlatformDictionary.ContainsKey(InPlatform) == true)
-			{
-				BuildPlatformDictionary.Remove(InPlatform);
+				BuildPlatformDictionary.Add(InBuildPlatform.Platform, InBuildPlatform);
 			}
 		}
 
@@ -160,7 +184,7 @@ namespace UnrealBuildTool
 		/// <summary>
 		/// Returns the delimiter used to separate paths in the PATH environment variable for the platform we are executing on.
 		/// </summary>
-		public String GetPathVarDelimiter()
+		public static String GetPathVarDelimiter()
 		{
 			switch (BuildHostPlatform.Current.Platform)
 			{
@@ -203,17 +227,6 @@ namespace UnrealBuildTool
 			return false;
 		}
 
-		/// <summary>
-		/// Register the platform with the UEBuildPlatform class
-		/// </summary>
-		public void RegisterBuildPlatform()
-		{
-			ManageAndValidateSDK();
-			RegisterBuildPlatformInternal();
-		}
-
-		protected abstract void RegisterBuildPlatformInternal();
-
 		/// <summary>
 		/// Retrieve the CPPTargetPlatform for the given UnrealTargetPlatform
 		/// </summary>
@@ -718,7 +731,10 @@ namespace UnrealBuildTool
 		/// <param name="DeploymentHandler">The output deployment handler</param>
 		/// <returns>True if the platform requires a deployment handler, false otherwise</returns>
 		public abstract bool TryCreateDeploymentHandler(FileReference ProjectFile, out UEBuildDeploy DeploymentHandler);
+	}
 
+	public abstract class UEBuildPlatformSDK
+	{
 		// AutoSDKs handling portion
 
 		#region protected AutoSDKs Utility
@@ -1124,7 +1140,7 @@ namespace UnrealBuildTool
 
 					// actually perform the PATH stripping / adding.
 					String OrigPathVar = Environment.GetEnvironmentVariable("PATH");
-					String PathDelimiter = GetPathVarDelimiter();
+					String PathDelimiter = UEBuildPlatform.GetPathVarDelimiter();
 					String[] PathVars = OrigPathVar.Split(PathDelimiter.ToCharArray());
 
 					List<String> ModifiedPathVars = new List<string>();
@@ -1510,4 +1526,12 @@ namespace UnrealBuildTool
 
 		#endregion
 	}
+
+	public abstract class UEBuildPlatformFactory
+	{
+		/// <summary>
+		/// Register the platform with the UEBuildPlatform class
+		/// </summary>
+		public abstract void RegisterBuildPlatforms();
+	}
 }
diff --git a/Engine/Source/Programs/UnrealBuildTool/HTML5/HTML5ToolChain.cs b/Engine/Source/Programs/UnrealBuildTool/HTML5/HTML5ToolChain.cs
--- a/Engine/Source/Programs/UnrealBuildTool/HTML5/HTML5ToolChain.cs
+++ b/Engine/Source/Programs/UnrealBuildTool/HTML5/HTML5ToolChain.cs
@@ -337,7 +337,7 @@ namespace UnrealBuildTool
 
 				// Add the source file path to the command-line.
 				string FileArguments = string.Format(" \"{0}\"", SourceFile.AbsolutePath);
-				var ObjectFileExtension = UEBuildPlatform.BuildPlatformDictionary[UnrealTargetPlatform.HTML5].GetBinaryExtension(UEBuildBinaryType.Object);
+				var ObjectFileExtension = UEBuildPlatform.GetBuildPlatform(UnrealTargetPlatform.HTML5).GetBinaryExtension(UEBuildBinaryType.Object);
 				// Add the object file to the produced item list.
 				FileItem ObjectFile = FileItem.GetItemByFileReference(
 					FileReference.Combine(
diff --git a/Engine/Source/Programs/UnrealBuildTool/HTML5/UEBuildHTML5.cs b/Engine/Source/Programs/UnrealBuildTool/HTML5/UEBuildHTML5.cs
--- a/Engine/Source/Programs/UnrealBuildTool/HTML5/UEBuildHTML5.cs
+++ b/Engine/Source/Programs/UnrealBuildTool/HTML5/UEBuildHTML5.cs
@@ -14,32 +14,19 @@
 		[XmlConfig]
 		public static string HTML5Architecture = "";
 
-		/// <summary>
-		/// Whether platform supports switching SDKs during runtime
-		/// </summary>
-		/// <returns>true if supports</returns>
-		protected override bool PlatformSupportsAutoSDKs()
+		HTML5PlatformSDK SDK;
+
+		public HTML5Platform(HTML5PlatformSDK InSDK) : base(UnrealTargetPlatform.HTML5)
 		{
-			return false;
-		}
-
-		public override string GetSDKTargetPlatformName()
-		{
-			return "HTML5";
+			SDK = InSDK;
 		}
 
 		/// <summary>
-		/// Returns SDK string as required by the platform
+		/// Whether the required external SDKs are installed for this platform. Could be either a manual install or an AutoSDK.
 		/// </summary>
-		/// <returns>Valid SDK string</returns>
-		protected override string GetRequiredSDKString()
+		public override SDKStatus HasRequiredSDKsInstalled()
 		{
-			return HTML5SDKInfo.EmscriptenVersion();
+			return SDK.HasRequiredSDKsInstalled();
 		}
 
-		protected override String GetRequiredScriptVersionString()
-		{
-			return "3.0";
-		}
-
 		// The current architecture - affects everything about how UBT operates on HTML5
@@ -56,59 +43,11 @@
 			return (ExeName + "-win32").ChangeExtension(".exe");
 		}
 
-		/// <summary>
-		/// Whether the required external SDKs are installed for this platform
-		/// </summary>
-		protected override SDKStatus HasRequiredManualSDKInternal()
-		{
-			if (!HTML5SDKInfo.IsSDKInstalled())
-			{
-				return SDKStatus.Invalid;
-			}
-			return SDKStatus.Valid;
-		}
-
 		public override bool CanUseXGE()
 		{
 			return (GetActiveArchitecture() == "-win32");
 		}
 
-		/// <summary>
-		/// Register the platform with the UEBuildPlatform class
-		/// </summary>
-		protected override void RegisterBuildPlatformInternal()
-		{
-			// Make sure the SDK is installed
-			if ((ProjectFileGenerator.bGenerateProjectFiles == true) || (HasRequiredSDKsInstalled() == SDKStatus.Valid))
-			{
-				bool bRegisterBuildPlatform = true;
-
-				// make sure we have the HTML5 files; if not, then this user doesn't really have HTML5 access/files, no need to compile HTML5!
-				string EngineSourcePath = Path.Combine(ProjectFileGenerator.EngineRelativePath, "Source");
-				string HTML5TargetPlatformFile = Path.Combine(EngineSourcePath, "Developer", "HTML5", "HTML5TargetPlatform", "HTML5TargetPlatform.Build.cs");
-				if ((File.Exists(HTML5TargetPlatformFile) == false))
-				{
-					bRegisterBuildPlatform = false;
-					Log.TraceWarning("Missing required components (.... HTML5TargetPlatformFile, others here...). Check source control filtering, or try resyncing.");
-				}
-
-				if (bRegisterBuildPlatform == true)
-				{
-					// Register this build platform for HTML5
-					Log.TraceVerbose("        Registering for {0}", UnrealTargetPlatform.HTML5.ToString());
-					UEBuildPlatform.RegisterBuildPlatform(UnrealTargetPlatform.HTML5, this);
-					if (GetActiveArchitecture() == "-win32")
-					{
-						UEBuildPlatform.RegisterPlatformWithGroup(UnrealTargetPlatform.HTML5, UnrealPlatformGroup.Simulator);
-					}
-					else
-					{
-						UEBuildPlatform.RegisterPlatformWithGroup(UnrealTargetPlatform.HTML5, UnrealPlatformGroup.Device);
-					}
-				}
-			}
-
-		}
 
 		/// <summary>
 		/// Retrieve the CPPTargetPlatform for the given UnrealTargetPlatform
@@ -382,4 +321,90 @@
 			return false;
 		}
 	}
+	
+	class HTML5PlatformSDK : UEBuildPlatformSDK
+	{
+		/// <summary>
+		/// Whether platform supports switching SDKs during runtime
+		/// </summary>
+		/// <returns>true if supports</returns>
+		protected override bool PlatformSupportsAutoSDKs()
+		{
+			return false;
+		}
+
+		public override string GetSDKTargetPlatformName()
+		{
+			return "HTML5";
+		}
+
+		/// <summary>
+		/// Returns SDK string as required by the platform
+		/// </summary>
+		/// <returns>Valid SDK string</returns>
+		protected override string GetRequiredSDKString()
+		{
+			return HTML5SDKInfo.EmscriptenVersion();
+		}
+
+		protected override String GetRequiredScriptVersionString()
+		{
+			return "3.0";
+		}
+
+		/// <summary>
+		/// Whether the required external SDKs are installed for this platform
+		/// </summary>
+		protected override SDKStatus HasRequiredManualSDKInternal()
+		{
+			if (!HTML5SDKInfo.IsSDKInstalled())
+			{
+				return SDKStatus.Invalid;
+			}
+			return SDKStatus.Valid;
+		}
+	}
+
+	class HTML5PlatformFactory : UEBuildPlatformFactory
+	{
+		/// <summary>
+		/// Register the platform with the UEBuildPlatform class
+		/// </summary>
+		public override void RegisterBuildPlatforms()
+		{
+			HTML5PlatformSDK SDK = new HTML5PlatformSDK();
+			SDK.ManageAndValidateSDK();
+
+			// Make sure the SDK is installed
+			if ((ProjectFileGenerator.bGenerateProjectFiles == true) || (SDK.HasRequiredSDKsInstalled() == SDKStatus.Valid))
+			{
+				bool bRegisterBuildPlatform = true;
+
+				// make sure we have the HTML5 files; if not, then this user doesn't really have HTML5 access/files, no need to compile HTML5!
+				string EngineSourcePath = Path.Combine(ProjectFileGenerator.EngineRelativePath, "Source");
+				string HTML5TargetPlatformFile = Path.Combine(EngineSourcePath, "Developer", "HTML5", "HTML5TargetPlatform", "HTML5TargetPlatform.Build.cs");
+				if ((File.Exists(HTML5TargetPlatformFile) == false))
+				{
+					bRegisterBuildPlatform = false;
+					Log.TraceWarning("Missing required components (.... HTML5TargetPlatformFile, others here...). Check source control filtering, or try resyncing.");
+				}
+
+				if (bRegisterBuildPlatform == true)
+				{
+					// Register this build platform for HTML5
+					Log.TraceVerbose("        Registering for {0}", UnrealTargetPlatform.HTML5.ToString());
+					UEBuildPlatform.RegisterBuildPlatform(new HTML5Platform(SDK));
+					if (HTML5Platform.HTML5Architecture == "-win32")
+					{
+						UEBuildPlatform.RegisterPlatformWithGroup(UnrealTargetPlatform.HTML5, UnrealPlatformGroup.Simulator);
+					}
+					else
+					{
+						UEBuildPlatform.RegisterPlatformWithGroup(UnrealTargetPlatform.HTML5, UnrealPlatformGroup.Device);
+					}
+				}
+			}
+
+		}
+	}
 }
diff --git a/Engine/Source/Programs/UnrealBuildTool/IOS/IOSToolChain.cs b/Engine/Source/Programs/UnrealBuildTool/IOS/IOSToolChain.cs
--- a/Engine/Source/Programs/UnrealBuildTool/IOS/IOSToolChain.cs
+++ b/Engine/Source/Programs/UnrealBuildTool/IOS/IOSToolChain.cs
@@ -437,7 +437,7 @@ namespace UnrealBuildTool
 			{
 				// Add the precompiled header file's path to the include path so GCC can find it.
 				// This needs to be before the other include paths to ensure GCC uses it instead of the source header file.
-				var PrecompiledFileExtension = UEBuildPlatform.BuildPlatformDictionary[UnrealTargetPlatform.IOS].GetBinaryExtension(UEBuildBinaryType.PrecompiledHeader);
+				var PrecompiledFileExtension = UEBuildPlatform.GetBuildPlatform(UnrealTargetPlatform.IOS).GetBinaryExtension(UEBuildBinaryType.PrecompiledHeader);
 				PCHArguments += string.Format(" -include \"{0}\"", ConvertPath(CompileEnvironment.PrecompiledHeaderFile.AbsolutePath.Replace(PrecompiledFileExtension, "")));
 			}
 
@@ -524,7 +524,7 @@ namespace UnrealBuildTool
 
 				if (CompileEnvironment.Config.PrecompiledHeaderAction == PrecompiledHeaderAction.Create)
 				{
-					var PrecompiledFileExtension = UEBuildPlatform.BuildPlatformDictionary[UnrealTargetPlatform.IOS].GetBinaryExtension(UEBuildBinaryType.PrecompiledHeader);
+					var PrecompiledFileExtension = UEBuildPlatform.GetBuildPlatform(UnrealTargetPlatform.IOS).GetBinaryExtension(UEBuildBinaryType.PrecompiledHeader);
 					// Add the precompiled header file to the produced item list.
 					FileItem PrecompiledHeaderFile = FileItem.GetItemByFileReference(
 						FileReference.Combine(
@@ -547,7 +547,7 @@ namespace UnrealBuildTool
 						CompileAction.bIsUsingPCH = true;
 						CompileAction.PrerequisiteItems.Add(CompileEnvironment.PrecompiledHeaderFile);
 					}
-					var ObjectFileExtension = UEBuildPlatform.BuildPlatformDictionary[UnrealTargetPlatform.IOS].GetBinaryExtension(UEBuildBinaryType.Object);
+					var ObjectFileExtension = UEBuildPlatform.GetBuildPlatform(UnrealTargetPlatform.IOS).GetBinaryExtension(UEBuildBinaryType.Object);
 					// Add the object file to the produced item list.
 					FileItem ObjectFile = FileItem.GetItemByFileReference(
 						FileReference.Combine(
diff --git a/Engine/Source/Programs/UnrealBuildTool/IOS/UEBuildIOS.cs b/Engine/Source/Programs/UnrealBuildTool/IOS/UEBuildIOS.cs
--- a/Engine/Source/Programs/UnrealBuildTool/IOS/UEBuildIOS.cs
+++ b/Engine/Source/Programs/UnrealBuildTool/IOS/UEBuildIOS.cs
@@ -50,6 +50,18 @@ namespace UnrealBuildTool
 
 		private bool bInitializedProject = false;
 
+		IOSPlatformSDK SDK;
+
+		public IOSPlatform(IOSPlatformSDK InSDK) : base(UnrealTargetPlatform.IOS)
+		{
+			SDK = InSDK;
+		}
+
+		public override SDKStatus HasRequiredSDKsInstalled()
+		{
+			return SDK.HasRequiredSDKsInstalled();
+		}
+
 		public string GetRunTimeVersion()
 		{
 			return RunTimeIOSVersion;
@@ -78,45 +90,6 @@ namespace UnrealBuildTool
 			return IOSArchitecture;
 		}
 
-		protected override SDKStatus HasRequiredManualSDKInternal()
-		{
-			if (!Utils.IsRunningOnMono)
-			{
-				// check to see if iTunes is installed
-				string dllPath = Microsoft.Win32.Registry.GetValue("HKEY_LOCAL_MACHINE\\SOFTWARE\\Wow6432Node\\Apple Inc.\\Apple Mobile Device Support\\Shared", "iTunesMobileDeviceDLL", null) as string;
-				if (String.IsNullOrEmpty(dllPath) || !File.Exists(dllPath))
-				{
-					dllPath = Microsoft.Win32.Registry.GetValue("HKEY_LOCAL_MACHINE\\SOFTWARE\\Wow6432Node\\Apple Inc.\\Apple Mobile Device Support\\Shared", "MobileDeviceDLL", null) as string;
-					if (String.IsNullOrEmpty(dllPath) || !File.Exists(dllPath))
-					{
-						return SDKStatus.Invalid;
-					}
-				}
-			}
-			return SDKStatus.Valid;
-		}
-
-		/// <summary>
-		/// Register the platform with the UEBuildPlatform class
-		/// </summary>
-		protected override void RegisterBuildPlatformInternal()
-		{
-			// Register this build platform for IOS
-			Log.TraceVerbose("        Registering for {0}", UnrealTargetPlatform.IOS.ToString());
-			UEBuildPlatform.RegisterBuildPlatform(UnrealTargetPlatform.IOS, this);
-			UEBuildPlatform.RegisterPlatformWithGroup(UnrealTargetPlatform.IOS, UnrealPlatformGroup.Unix);
-			UEBuildPlatform.RegisterPlatformWithGroup(UnrealTargetPlatform.IOS, UnrealPlatformGroup.Apple);
-
-			if (GetActiveArchitecture() == "-simulator")
-			{
-				UEBuildPlatform.RegisterPlatformWithGroup(UnrealTargetPlatform.IOS, UnrealPlatformGroup.Simulator);
-			}
-			else
-			{
-				UEBuildPlatform.RegisterPlatformWithGroup(UnrealTargetPlatform.IOS, UnrealPlatformGroup.Device);
-			}
-		}
-
 		/// <summary>
 		/// Retrieve the CPPTargetPlatform for the given UnrealTargetPlatform
 		/// </summary>
@@ -485,5 +458,53 @@ namespace UnrealBuildTool
 			return true;
 		}
 	}
+
+	public class IOSPlatformSDK : UEBuildPlatformSDK
+	{
+		protected override SDKStatus HasRequiredManualSDKInternal()
+		{
+			if (!Utils.IsRunningOnMono)
+			{
+				// check to see if iTunes is installed
+				string dllPath = Microsoft.Win32.Registry.GetValue("HKEY_LOCAL_MACHINE\\SOFTWARE\\Wow6432Node\\Apple Inc.\\Apple Mobile Device Support\\Shared", "iTunesMobileDeviceDLL", null) as string;
+				if (String.IsNullOrEmpty(dllPath) || !File.Exists(dllPath))
+				{
+					dllPath = Microsoft.Win32.Registry.GetValue("HKEY_LOCAL_MACHINE\\SOFTWARE\\Wow6432Node\\Apple Inc.\\Apple Mobile Device Support\\Shared", "MobileDeviceDLL", null) as string;
+					if (String.IsNullOrEmpty(dllPath) || !File.Exists(dllPath))
+					{
+						return SDKStatus.Invalid;
+					}
+				}
+			}
+			return SDKStatus.Valid;
+		}
+	}
+
+	public class IOSPlatformFactory : UEBuildPlatformFactory
+	{
+		/// <summary>
+		/// Register the platform with the UEBuildPlatform class
+		/// </summary>
+		public override void RegisterBuildPlatforms()
+		{
+			IOSPlatformSDK SDK = new IOSPlatformSDK();
+			SDK.ManageAndValidateSDK();
+
+			// Register this build platform for IOS
+			Log.TraceVerbose("        Registering for {0}", UnrealTargetPlatform.IOS.ToString());
+			UEBuildPlatform.RegisterBuildPlatform(new IOSPlatform(SDK));
+			UEBuildPlatform.RegisterPlatformWithGroup(UnrealTargetPlatform.IOS, UnrealPlatformGroup.Unix);
+			UEBuildPlatform.RegisterPlatformWithGroup(UnrealTargetPlatform.IOS, UnrealPlatformGroup.Apple);
+
+			if (IOSPlatform.IOSArchitecture == "-simulator")
+			{
+				UEBuildPlatform.RegisterPlatformWithGroup(UnrealTargetPlatform.IOS, UnrealPlatformGroup.Simulator);
+			}
+			else
+			{
+				UEBuildPlatform.RegisterPlatformWithGroup(UnrealTargetPlatform.IOS, UnrealPlatformGroup.Device);
+			}
+		}
+	}
 }
 
diff --git a/Engine/Source/Programs/UnrealBuildTool/Linux/LinuxToolChain.cs b/Engine/Source/Programs/UnrealBuildTool/Linux/LinuxToolChain.cs
--- a/Engine/Source/Programs/UnrealBuildTool/Linux/LinuxToolChain.cs
+++ b/Engine/Source/Programs/UnrealBuildTool/Linux/LinuxToolChain.cs
@@ -610,7 +610,7 @@ namespace UnrealBuildTool
 			{
 				// Add the precompiled header file's path to the include path so Clang can find it.
 				// This needs to be before the other include paths to ensure Clang uses it instead of the source header file.
-				var PrecompiledFileExtension = UEBuildPlatform.BuildPlatformDictionary[UnrealTargetPlatform.Linux].GetBinaryExtension(UEBuildBinaryType.PrecompiledHeader);
+				var PrecompiledFileExtension = UEBuildPlatform.GetBuildPlatform(UnrealTargetPlatform.Linux).GetBinaryExtension(UEBuildBinaryType.PrecompiledHeader);
 				PCHArguments += string.Format(" -include \"{0}\"", CompileEnvironment.PrecompiledHeaderFile.AbsolutePath.Replace(PrecompiledFileExtension, ""));
 			}
 
@@ -678,7 +678,7 @@ namespace UnrealBuildTool
 
 				if (CompileEnvironment.Config.PrecompiledHeaderAction == PrecompiledHeaderAction.Create)
 				{
-					var PrecompiledFileExtension = UEBuildPlatform.BuildPlatformDictionary[UnrealTargetPlatform.Linux].GetBinaryExtension(UEBuildBinaryType.PrecompiledHeader);
+					var PrecompiledFileExtension = UEBuildPlatform.GetBuildPlatform(UnrealTargetPlatform.Linux).GetBinaryExtension(UEBuildBinaryType.PrecompiledHeader);
 					// Add the precompiled header file to the produced item list.
 					FileItem PrecompiledHeaderFile = FileItem.GetItemByFileReference(
 						FileReference.Combine(
@@ -701,7 +701,7 @@ namespace UnrealBuildTool
 						CompileAction.PrerequisiteItems.Add(CompileEnvironment.PrecompiledHeaderFile);
 					}
 
-					var ObjectFileExtension = UEBuildPlatform.BuildPlatformDictionary[UnrealTargetPlatform.Linux].GetBinaryExtension(UEBuildBinaryType.Object);
+					var ObjectFileExtension = UEBuildPlatform.GetBuildPlatform(UnrealTargetPlatform.Linux).GetBinaryExtension(UEBuildBinaryType.Object);
 					// Add the object file to the produced item list.
 					FileItem ObjectFile = FileItem.GetItemByFileReference(
 						FileReference.Combine(
diff --git a/Engine/Source/Programs/UnrealBuildTool/Linux/UEBuildLinux.cs b/Engine/Source/Programs/UnrealBuildTool/Linux/UEBuildLinux.cs
--- a/Engine/Source/Programs/UnrealBuildTool/Linux/UEBuildLinux.cs
+++ b/Engine/Source/Programs/UnrealBuildTool/Linux/UEBuildLinux.cs
@@ -11,27 +11,31 @@
 	class LinuxPlatform : UEBuildPlatform
 	{
 		/// <summary>
-		/// This is the SDK version we support
-		/// </summary>
-		static private Dictionary<string, string> ExpectedSDKVersions = new Dictionary<string, string>()
-		{
-			{ "x86_64-unknown-linux-gnu", "v6_clang-3.6.0_ld-2.24_glibc-2.12.2" },
-			{ "arm-unknown-linux-gnueabihf", "arm-unknown-linux-gnueabihf_v5_clang-3.5.0-ld-2.23.1-glibc-2.13" },
-		};
-
-		/// <summary>
-		/// Platform name (embeds architecture for now)
-		/// </summary>
-		static private string TargetPlatformName = "Linux_x64";
-
-		/// <summary>
 		/// Linux architecture (compiler target triplet)
 		/// </summary>
 		// FIXME: for now switching between architectures is hard-coded
-		static private string DefaultArchitecture = "x86_64-unknown-linux-gnu";
+		public const string DefaultArchitecture = "x86_64-unknown-linux-gnu";
 		//static private string DefaultArchitecture = "arm-unknown-linux-gnueabihf";
 
+		LinuxPlatformSDK SDK;
+
 		/// <summary>
+		/// Constructor
+		/// </summary>
+		public LinuxPlatform(LinuxPlatformSDK InSDK) : base(UnrealTargetPlatform.Linux)
+		{
+			SDK = InSDK;
+		}
+
+		/// <summary>
+		/// Whether the required external SDKs are installed for this platform. Could be either a manual install or an AutoSDK.
+		/// </summary>
+		public override SDKStatus HasRequiredSDKsInstalled()
+		{
+			return SDK.HasRequiredSDKsInstalled();
+		}
+
+		/// <summary>
 		/// The current architecture
 		/// </summary>
 		public override string GetActiveArchitecture()
@@ -48,74 +52,4 @@
 			return BinaryName;
 		}
 
-		/// <summary>
-		/// Whether platform supports switching SDKs during runtime
-		/// </summary>
-		/// <returns>true if supports</returns>
-		protected override bool PlatformSupportsAutoSDKs()
-		{
-			return true;
-		}
-
-		/// <summary>
-		/// Returns platform-specific name used in SDK repository
-		/// </summary>
-		/// <returns>path to SDK Repository</returns>
-		public override string GetSDKTargetPlatformName()
-		{
-			return TargetPlatformName;
-		}
-
-		/// <summary>
-		/// Returns SDK string as required by the platform
-		/// </summary>
-		/// <returns>Valid SDK string</returns>
-		protected override string GetRequiredSDKString()
-		{
-			string SDKString;
-			if (!ExpectedSDKVersions.TryGetValue(GetActiveArchitecture(), out SDKString))
-			{
-				throw new BuildException("LinuxPlatform::GetRequiredSDKString: no toolchain set up for architecture '{0}'", GetActiveArchitecture());
-			}
-
-			return SDKString;
-		}
-
-		protected override String GetRequiredScriptVersionString()
-		{
-			return "3.0";
-		}
-
-		protected override bool PreferAutoSDK()
-		{
-			// having LINUX_ROOT set (for legacy reasons or for convenience of cross-compiling certain third party libs) should not make UBT skip AutoSDKs
-			return true;
-		}
-
-		/// <summary>
-		/// Whether the required external SDKs are installed for this platform
-		/// </summary>
-		protected override SDKStatus HasRequiredManualSDKInternal()
-		{
-			if (BuildHostPlatform.Current.Platform == UnrealTargetPlatform.Linux)
-			{
-				return SDKStatus.Valid;
-			}
-
-			string BaseLinuxPath = Environment.GetEnvironmentVariable("LINUX_ROOT");
-
-			// we don't have an LINUX_ROOT specified
-			if (String.IsNullOrEmpty(BaseLinuxPath))
-				return SDKStatus.Invalid;
-
-			// paths to our toolchains
-			BaseLinuxPath = BaseLinuxPath.Replace("\"", "");
-			string ClangPath = Path.Combine(BaseLinuxPath, @"bin\Clang++.exe");
-
-			if (File.Exists(ClangPath))
-				return SDKStatus.Valid;
-
-			return SDKStatus.Invalid;
-		}
-
 		public override bool CanUseXGE()
@@ -125,35 +59,5 @@
 		}
 
 		/// <summary>
-		/// Register the platform with the UEBuildPlatform class
-		/// </summary>
-		protected override void RegisterBuildPlatformInternal()
-		{
-			if ((ProjectFileGenerator.bGenerateProjectFiles == true) || (HasRequiredSDKsInstalled() == SDKStatus.Valid))
-			{
-				bool bRegisterBuildPlatform = true;
-
-				string EngineSourcePath = Path.Combine(ProjectFileGenerator.RootRelativePath, "Engine", "Source");
-				string LinuxTargetPlatformFile = Path.Combine(EngineSourcePath, "Developer", "Linux", "LinuxTargetPlatform", "LinuxTargetPlatform.Build.cs");
-
-				if (File.Exists(LinuxTargetPlatformFile) == false)
-				{
-					bRegisterBuildPlatform = false;
-				}
-
-				if (bRegisterBuildPlatform == true)
-				{
-					// Register this build platform for Linux
-					if (BuildConfiguration.bPrintDebugInfo)
-					{
-						Console.WriteLine("        Registering for {0}", UnrealTargetPlatform.Linux.ToString());
-					}
-					UEBuildPlatform.RegisterBuildPlatform(UnrealTargetPlatform.Linux, this);
-					UEBuildPlatform.RegisterPlatformWithGroup(UnrealTargetPlatform.Linux, UnrealPlatformGroup.Unix);
-				}
-			}
-		}
-
-		/// <summary>
 		/// Retrieve the CPPTargetPlatform for the given UnrealTargetPlatform
 		/// </summary>
@@ -397,4 +301,127 @@
 			return false;
 		}
 	}
+
+	class LinuxPlatformSDK : UEBuildPlatformSDK
+	{
+		/// <summary>
+		/// This is the SDK version we support
+		/// </summary>
+		static private Dictionary<string, string> ExpectedSDKVersions = new Dictionary<string, string>()
+		{
+			{ "x86_64-unknown-linux-gnu", "v6_clang-3.6.0_ld-2.24_glibc-2.12.2" },
+			{ "arm-unknown-linux-gnueabihf", "arm-unknown-linux-gnueabihf_v5_clang-3.5.0-ld-2.23.1-glibc-2.13" },
+		};
+
+		/// <summary>
+		/// Platform name (embeds architecture for now)
+		/// </summary>
+		static private string TargetPlatformName = "Linux_x64";
+
+		/// <summary>
+		/// Whether platform supports switching SDKs during runtime
+		/// </summary>
+		/// <returns>true if supports</returns>
+		protected override bool PlatformSupportsAutoSDKs()
+		{
+			return true;
+		}
+
+		/// <summary>
+		/// Returns platform-specific name used in SDK repository
+		/// </summary>
+		/// <returns>path to SDK Repository</returns>
+		public override string GetSDKTargetPlatformName()
+		{
+			return TargetPlatformName;
+		}
+
+		/// <summary>
+		/// Returns SDK string as required by the platform
+		/// </summary>
+		/// <returns>Valid SDK string</returns>
+		protected override string GetRequiredSDKString()
+		{
+			string SDKString;
+			if (!ExpectedSDKVersions.TryGetValue(LinuxPlatform.DefaultArchitecture, out SDKString))
+			{
+				throw new BuildException("LinuxPlatform::GetRequiredSDKString: no toolchain set up for architecture '{0}'", LinuxPlatform.DefaultArchitecture);
+			}
+
+			return SDKString;
+		}
+
+		protected override String GetRequiredScriptVersionString()
+		{
+			return "3.0";
+		}
+
+		protected override bool PreferAutoSDK()
+		{
+			// having LINUX_ROOT set (for legacy reasons or for convenience of cross-compiling certain third party libs) should not make UBT skip AutoSDKs
+			return true;
+		}
+
+		/// <summary>
+		/// Whether the required external SDKs are installed for this platform
+		/// </summary>
+		protected override SDKStatus HasRequiredManualSDKInternal()
+		{
+			if (BuildHostPlatform.Current.Platform == UnrealTargetPlatform.Linux)
+			{
+				return SDKStatus.Valid;
+			}
+
+			string BaseLinuxPath = Environment.GetEnvironmentVariable("LINUX_ROOT");
+
+			// we don't have an LINUX_ROOT specified
+			if (String.IsNullOrEmpty(BaseLinuxPath))
+				return SDKStatus.Invalid;
+
+			// paths to our toolchains
+			BaseLinuxPath = BaseLinuxPath.Replace("\"", "");
+			string ClangPath = Path.Combine(BaseLinuxPath, @"bin\Clang++.exe");
+
+			if (File.Exists(ClangPath))
+				return SDKStatus.Valid;
+
+			return SDKStatus.Invalid;
+		}
+	}
+
+	class LinuxPlatformFactory : UEBuildPlatformFactory
+	{
+		/// <summary>
+		/// Register the platform with the UEBuildPlatform class
+		/// </summary>
+		public override void RegisterBuildPlatforms()
+		{
+			LinuxPlatformSDK SDK = new LinuxPlatformSDK();
+			SDK.ManageAndValidateSDK();
+
+			if ((ProjectFileGenerator.bGenerateProjectFiles == true) || (SDK.HasRequiredSDKsInstalled() == SDKStatus.Valid))
+			{
+				bool bRegisterBuildPlatform = true;
+
+				string EngineSourcePath = Path.Combine(ProjectFileGenerator.RootRelativePath, "Engine", "Source");
+				string LinuxTargetPlatformFile = Path.Combine(EngineSourcePath, "Developer", "Linux", "LinuxTargetPlatform", "LinuxTargetPlatform.Build.cs");
+
+				if (File.Exists(LinuxTargetPlatformFile) == false)
+				{
+					bRegisterBuildPlatform = false;
+				}
+
+				if (bRegisterBuildPlatform == true)
+				{
+					// Register this build platform for Linux
+					if (BuildConfiguration.bPrintDebugInfo)
+					{
+						Console.WriteLine("        Registering for {0}", UnrealTargetPlatform.Linux.ToString());
+					}
+					UEBuildPlatform.RegisterBuildPlatform(new LinuxPlatform(SDK));
+					UEBuildPlatform.RegisterPlatformWithGroup(UnrealTargetPlatform.Linux, UnrealPlatformGroup.Unix);
+				}
+			}
+		}
+	}
 }
diff --git a/Engine/Source/Programs/UnrealBuildTool/Mac/MacToolChain.cs b/Engine/Source/Programs/UnrealBuildTool/Mac/MacToolChain.cs
--- a/Engine/Source/Programs/UnrealBuildTool/Mac/MacToolChain.cs
+++ b/Engine/Source/Programs/UnrealBuildTool/Mac/MacToolChain.cs
@@ -278,7 +278,7 @@ namespace UnrealBuildTool
 			{
 				// Add the precompiled header file's path to the include path so GCC can find it.
 				// This needs to be before the other include paths to ensure GCC uses it instead of the source header file.
-				var PrecompiledFileExtension = UEBuildPlatform.BuildPlatformDictionary[UnrealTargetPlatform.Mac].GetBinaryExtension(UEBuildBinaryType.PrecompiledHeader);
+				var PrecompiledFileExtension = UEBuildPlatform.GetBuildPlatform(UnrealTargetPlatform.Mac).GetBinaryExtension(UEBuildBinaryType.PrecompiledHeader);
 				PCHArguments.Append(" -include \"");
 				PCHArguments.Append(CompileEnvironment.PrecompiledHeaderFile.AbsolutePath.Replace(PrecompiledFileExtension, ""));
 				PCHArguments.Append("\"");
@@ -377,7 +377,7 @@ namespace UnrealBuildTool
 
 				if (CompileEnvironment.Config.PrecompiledHeaderAction == PrecompiledHeaderAction.Create)
 				{
-					var PrecompiledHeaderExtension = UEBuildPlatform.BuildPlatformDictionary[UnrealTargetPlatform.Mac].GetBinaryExtension(UEBuildBinaryType.PrecompiledHeader);
+					var PrecompiledHeaderExtension = UEBuildPlatform.GetBuildPlatform(UnrealTargetPlatform.Mac).GetBinaryExtension(UEBuildBinaryType.PrecompiledHeader);
 					// Add the precompiled header file to the produced item list.
 					FileItem PrecompiledHeaderFile = FileItem.GetItemByFileReference(
 						FileReference.Combine(
@@ -400,7 +400,7 @@ namespace UnrealBuildTool
 						CompileAction.bIsUsingPCH = true;
 						CompileAction.PrerequisiteItems.Add(CompileEnvironment.PrecompiledHeaderFile);
 					}
-					var ObjectFileExtension = UEBuildPlatform.BuildPlatformDictionary[UnrealTargetPlatform.Mac].GetBinaryExtension(UEBuildBinaryType.Object);
+					var ObjectFileExtension = UEBuildPlatform.GetBuildPlatform(UnrealTargetPlatform.Mac).GetBinaryExtension(UEBuildBinaryType.Object);
 					// Add the object file to the produced item list.
 					FileItem ObjectFile = FileItem.GetItemByFileReference(
 						FileReference.Combine(
diff --git a/Engine/Source/Programs/UnrealBuildTool/Mac/UEBuildMac.cs b/Engine/Source/Programs/UnrealBuildTool/Mac/UEBuildMac.cs
--- a/Engine/Source/Programs/UnrealBuildTool/Mac/UEBuildMac.cs
+++ b/Engine/Source/Programs/UnrealBuildTool/Mac/UEBuildMac.cs
@@ -11,31 +11,26 @@ namespace UnrealBuildTool
 {
 	class MacPlatform : UEBuildPlatform
 	{
-		public override bool CanUseXGE()
+		MacPlatformSDK SDK;
+
+		public MacPlatform(MacPlatformSDK InSDK) : base(UnrealTargetPlatform.Mac)
 		{
-			return false;
+			SDK = InSDK;
 		}
 
-		public override bool CanUseDistcc()
+		public override SDKStatus HasRequiredSDKsInstalled()
 		{
-			return true;
+			return SDK.HasRequiredSDKsInstalled();
 		}
 
-		protected override SDKStatus HasRequiredManualSDKInternal()
+		public override bool CanUseXGE()
 		{
-			return SDKStatus.Valid;
+			return false;
 		}
 
-		/// <summary>
-		/// Register the platform with the UEBuildPlatform class
-		/// </summary>
-		protected override void RegisterBuildPlatformInternal()
+		public override bool CanUseDistcc()
 		{
-			// Register this build platform for Mac
-			Log.TraceVerbose("        Registering for {0}", UnrealTargetPlatform.Mac.ToString());
-			UEBuildPlatform.RegisterBuildPlatform(UnrealTargetPlatform.Mac, this);
-			UEBuildPlatform.RegisterPlatformWithGroup(UnrealTargetPlatform.Mac, UnrealPlatformGroup.Unix);
-			UEBuildPlatform.RegisterPlatformWithGroup(UnrealTargetPlatform.Mac, UnrealPlatformGroup.Apple);
+			return true;
 		}
 
 		/// <summary>
@@ -241,4 +236,30 @@ namespace UnrealBuildTool
 			return true;
 		}
 	}
+
+	class MacPlatformSDK : UEBuildPlatformSDK
+	{
+		protected override SDKStatus HasRequiredManualSDKInternal()
+		{
+			return SDKStatus.Valid;
+		}
+	}
+
+	class MacPlatformFactory : UEBuildPlatformFactory
+	{
+		/// <summary>
+		/// Register the platform with the UEBuildPlatform class
+		/// </summary>
+		public override void RegisterBuildPlatforms()
+		{
+			MacPlatformSDK SDK = new MacPlatformSDK();
+			SDK.ManageAndValidateSDK();
+
+			// Register this build platform for Mac
+			Log.TraceVerbose("        Registering for {0}", UnrealTargetPlatform.Mac.ToString());
+			UEBuildPlatform.RegisterBuildPlatform(new MacPlatform(SDK));
+			UEBuildPlatform.RegisterPlatformWithGroup(UnrealTargetPlatform.Mac, UnrealPlatformGroup.Unix);
+			UEBuildPlatform.RegisterPlatformWithGroup(UnrealTargetPlatform.Mac, UnrealPlatformGroup.Apple);
+		}
+	}
 }
diff --git a/Engine/Source/Programs/UnrealBuildTool/System/ProjectFileGenerator.cs b/Engine/Source/Programs/UnrealBuildTool/System/ProjectFileGenerator.cs
--- a/Engine/Source/Programs/UnrealBuildTool/System/ProjectFileGenerator.cs
+++ b/Engine/Source/Programs/UnrealBuildTool/System/ProjectFileGenerator.cs
@@ -1309,7 +1309,7 @@ namespace UnrealBuildTool
         /// </summary>
         private void AddPS4Projects(MasterProjectFolder Folder)
         {
-			if (UEBuildPlatform.BuildPlatformDictionary.ContainsKey(UnrealTargetPlatform.PS4))
+			if (UEBuildPlatform.IsPlatformAvailable(UnrealTargetPlatform.PS4))
 			{
 				string ProjectFolderName = Path.Combine(EngineRelativePath, "Source", "Programs", "PS4");
 				DirectoryInfo ProjectFolderInfo = new DirectoryInfo(ProjectFolderName);
@@ -1468,13 +1468,6 @@ namespace UnrealBuildTool
 						}
 					}
 				}
-				else
-				{
-					// We have to unregister any build platforms we aren't supporting.
-					// Otherwise, they can add modules that can cause issues when doing tasks
-					// like generating IntelliSense.
-					UEBuildPlatform.UnregisterBuildPlatform(Platform);
-				}
 			}
 
 			// Add all configurations
diff --git a/Engine/Source/Programs/UnrealBuildTool/System/UnrealBuildTool.cs b/Engine/Source/Programs/UnrealBuildTool/System/UnrealBuildTool.cs
--- a/Engine/Source/Programs/UnrealBuildTool/System/UnrealBuildTool.cs
+++ b/Engine/Source/Programs/UnrealBuildTool/System/UnrealBuildTool.cs
@@ -486,11 +486,11 @@ namespace UnrealBuildTool
 					{
 						if (CheckType.IsClass && !CheckType.IsAbstract)
 						{
-							if (CheckType.IsSubclassOf(typeof(UEBuildPlatform)))
+							if (CheckType.IsSubclassOf(typeof(UEBuildPlatformFactory)))
 							{
 								Log.TraceVerbose("    Registering build platform: {0}", CheckType.ToString());
-								var TempInst = (UEBuildPlatform)(UBTAssembly.CreateInstance(CheckType.FullName, true));
-								TempInst.RegisterBuildPlatform();
+								var TempInst = (UEBuildPlatformFactory)(UBTAssembly.CreateInstance(CheckType.FullName, true));
+								TempInst.RegisterBuildPlatforms();
 							}
 						}
 					}
diff --git a/Engine/Source/Programs/UnrealBuildTool/UWP/UEBuildUWP.cs b/Engine/Source/Programs/UnrealBuildTool/UWP/UEBuildUWP.cs
--- a/Engine/Source/Programs/UnrealBuildTool/UWP/UEBuildUWP.cs
+++ b/Engine/Source/Programs/UnrealBuildTool/UWP/UEBuildUWP.cs
@@ -66,23 +66,21 @@ namespace UnrealBuildTool
 			}
 		}
 
-		public override bool RequiresDeployPrepAfterCompile()
+		UWPPlatformSDK SDK;
+
+		public UWPPlatform(UWPPlatformSDK InSDK) : base(UnrealTargetPlatform.UWP)
 		{
-			return true;
+			SDK = InSDK;
 		}
 
-		protected override SDKStatus HasRequiredManualSDKInternal()
+		public override SDKStatus HasRequiredSDKsInstalled()
 		{
-			try
-			{
-				// @todo UWP: wire this up once the rest of cleanup has been resolved. right now, just leaving this to always return valid, like VC does.
-				//WinUWPToolChain.FindBaseVSToolPath();
-				return SDKStatus.Valid;
-			}
-			catch (BuildException)
-			{
-				return SDKStatus.Invalid;
-			}
+			return SDK.HasRequiredSDKsInstalled();
+		}
+
+		public override bool RequiresDeployPrepAfterCompile()
+		{
+			return true;
 		}
 
 		/// <summary>
@@ -140,17 +138,6 @@ namespace UnrealBuildTool
 		}
 
 
-		/// <summary>
-		/// Register the platform with the UEBuildPlatform class
-		/// </summary>
-		protected override void RegisterBuildPlatformInternal()
-		{
-			// Register this build platform for UWP
-			Log.TraceVerbose("        Registering for {0}", UnrealTargetPlatform.UWP.ToString());
-			UEBuildPlatform.RegisterBuildPlatform(UnrealTargetPlatform.UWP, this);
-			UEBuildPlatform.RegisterPlatformWithGroup(UnrealTargetPlatform.UWP, UnrealPlatformGroup.Microsoft);
-		}
-
 		/// <summary>
 		/// Retrieve the CPPTargetPlatform for the given UnrealTargetPlatform
 		/// </summary>
@@ -498,4 +485,39 @@ namespace UnrealBuildTool
 			return true;
 		}
 	}
+
+	public class UWPPlatformSDK : UEBuildPlatformSDK
+	{
+		protected override SDKStatus HasRequiredManualSDKInternal()
+		{
+			try
+			{
+				// @todo UWP: wire this up once the rest of cleanup has been resolved. right now, just leaving this to always return valid, like VC does.
+				//WinUWPToolChain.FindBaseVSToolPath();
+				return SDKStatus.Valid;
+			}
+			catch (BuildException)
+			{
+				return SDKStatus.Invalid;
+			}
+		}
+	}
+
+	class UWPPlatformFactory : UEBuildPlatformFactory
+	{
+		/// <summary>
+		/// Register the platform with the UEBuildPlatform class
+		/// </summary>
+		public override void RegisterBuildPlatforms()
+		{
+			UWPPlatformSDK SDK = new UWPPlatformSDK();
+			SDK.ManageAndValidateSDK();
+
+			// Register this build platform for UWP
+			Log.TraceVerbose("        Registering for {0}", UnrealTargetPlatform.UWP.ToString());
+			UEBuildPlatform.RegisterBuildPlatform(new UWPPlatform(SDK));
+			UEBuildPlatform.RegisterPlatformWithGroup(UnrealTargetPlatform.UWP, UnrealPlatformGroup.Microsoft);
+		}
+
+	}
 }
diff --git a/Engine/Source/Programs/UnrealBuildTool/UWP/UWPToolChain.cs b/Engine/Source/Programs/UnrealBuildTool/UWP/UWPToolChain.cs
--- a/Engine/Source/Programs/UnrealBuildTool/UWP/UWPToolChain.cs
+++ b/Engine/Source/Programs/UnrealBuildTool/UWP/UWPToolChain.cs
@@ -575,7 +575,7 @@ namespace UnrealBuildTool
 					string OriginalPCHHeaderDirectory = Path.GetDirectoryName(SourceFile.AbsolutePath);
 					FileArguments.AppendFormat(" /I \"{0}\"", OriginalPCHHeaderDirectory);
 
-					var PrecompiledFileExtension = UEBuildPlatform.BuildPlatformDictionary[UnrealTargetPlatform.UWP].GetBinaryExtension(UEBuildBinaryType.PrecompiledHeader);
+					var PrecompiledFileExtension = UEBuildPlatform.GetBuildPlatform(UnrealTargetPlatform.UWP).GetBinaryExtension(UEBuildBinaryType.PrecompiledHeader);
 					// Add the precompiled header file to the produced items list.
 					FileItem PrecompiledHeaderFile = FileItem.GetItemByFileReference(
 						FileReference.Combine(
@@ -636,7 +636,7 @@ namespace UnrealBuildTool
 
 				if (bEmitsObjectFile)
 				{
-					var ObjectFileExtension = UEBuildPlatform.BuildPlatformDictionary[UnrealTargetPlatform.UWP].GetBinaryExtension(UEBuildBinaryType.Object);
+					var ObjectFileExtension = UEBuildPlatform.GetBuildPlatform(UnrealTargetPlatform.UWP).GetBinaryExtension(UEBuildBinaryType.Object);
 					// Add the object file to the produced item list.
 					FileItem ObjectFile = FileItem.GetItemByFileReference(
 						FileReference.Combine(
diff --git a/Engine/Source/Programs/UnrealBuildTool/WinRT/UEBuildWinRT.cs b/Engine/Source/Programs/UnrealBuildTool/WinRT/UEBuildWinRT.cs
--- a/Engine/Source/Programs/UnrealBuildTool/WinRT/UEBuildWinRT.cs
+++ b/Engine/Source/Programs/UnrealBuildTool/WinRT/UEBuildWinRT.cs
@@ -16,6 +16,18 @@
 		[XmlConfig]
 		public static bool bCompileWinRT = false;
 
+		WinRTPlatformSDK SDK;
+
+		public WinRTPlatform(UnrealTargetPlatform InPlatform, WinRTPlatformSDK InSDK) : base(InPlatform)
+		{
+			SDK = InSDK;
+		}
+
+		public override SDKStatus HasRequiredSDKsInstalled()
+		{
+			return SDK.HasRequiredSDKsInstalled();
+		}
+
 		public static bool IsVisualStudioInstalled()
 		{
 			string BaseVSToolPath = WindowsPlatform.GetVSComnToolsPath();
@@ -35,62 +47,12 @@
 			}
 			return false;
 		}
+
 		public override bool CanUseXGE()
 		{
 			return false;
 		}
 
 		/// <summary>
-		/// Whether the required external SDKs are installed for this platform
-		/// </summary>
-		protected override SDKStatus HasRequiredManualSDKInternal()
-		{
-			return !Utils.IsRunningOnMono && IsVisualStudioInstalled() ? SDKStatus.Valid : SDKStatus.Invalid;
-		}
-
-		/// <summary>
-		/// Register the platform with the UEBuildPlatform class
-		/// </summary>
-		protected override void RegisterBuildPlatformInternal()
-		{
-			[REDACTED_EMAIL]: Add platform support
-			if (UnrealBuildTool.RunningRocket() || Utils.IsRunningOnMono)
-			{
-				return;
-			}
-
-			if ((ProjectFileGenerator.bGenerateProjectFiles == true) || (IsVisualStudioInstalled() == true))
-			{
-				bool bRegisterBuildPlatform = true;
-
-				// We also need to check for the generated projects... to handle the case where someone generates projects w/out WinRT.
-				// Hardcoding this for now - but ideally it would be dynamically discovered.
-				string EngineSourcePath = Path.Combine(ProjectFileGenerator.EngineRelativePath, "Source");
-				string WinRTRHIFile = Path.Combine(EngineSourcePath, "Runtime", "Windows", "D3D11RHI", "D3D11RHI.build.cs");
-				if (File.Exists(WinRTRHIFile) == false)
-				{
-					bRegisterBuildPlatform = false;
-				}
-
-				if (bRegisterBuildPlatform == true)
-				{
-					// Register this build platform for WinRT
-					Log.TraceVerbose("        Registering for {0}", UnrealTargetPlatform.WinRT.ToString());
-					UEBuildPlatform.RegisterBuildPlatform(UnrealTargetPlatform.WinRT, this);
-					UEBuildPlatform.RegisterPlatformWithGroup(UnrealTargetPlatform.WinRT, UnrealPlatformGroup.Microsoft);
-
-					// For now only register WinRT_ARM is truly a Windows 8 machine.
-					// This will prevent people who do all platform builds from running into the compiler issue.
-					if (WinRTPlatform.IsWindows8() == true)
-					{
-						Log.TraceVerbose("        Registering for {0}", UnrealTargetPlatform.WinRT_ARM.ToString());
-						UEBuildPlatform.RegisterBuildPlatform(UnrealTargetPlatform.WinRT_ARM, this);
-						UEBuildPlatform.RegisterPlatformWithGroup(UnrealTargetPlatform.WinRT_ARM, UnrealPlatformGroup.Microsoft);
-					}
-				}
-			}
-		}
-
-		/// <summary>
 		/// Retrieve the CPPTargetPlatform for the given UnrealTargetPlatform
 		/// </summary>
@@ -648,4 +610,64 @@
 			return true;
 		}
 	}
+
+	public class WinRTPlatformSDK : UEBuildPlatformSDK
+	{
+		/// <summary>
+		/// Whether the required external SDKs are installed for this platform
+		/// </summary>
+		protected override SDKStatus HasRequiredManualSDKInternal()
+		{
+			return !Utils.IsRunningOnMono && WinRTPlatform.IsVisualStudioInstalled() ? SDKStatus.Valid : SDKStatus.Invalid;
+		}
+	}
+
+	public class WinRTPlatformFactory : UEBuildPlatformFactory
+	{
+		/// <summary>
+		/// Register the platform with the UEBuildPlatform class
+		/// </summary>
+		public override void RegisterBuildPlatforms()
+		{
+			[REDACTED_EMAIL]: Add platform support
+			if (UnrealBuildTool.RunningRocket() || Utils.IsRunningOnMono)
+			{
+				return;
+			}
+
+			WinRTPlatformSDK SDK = new WinRTPlatformSDK();
+			SDK.ManageAndValidateSDK();
+
+			if ((ProjectFileGenerator.bGenerateProjectFiles == true) || (WinRTPlatform.IsVisualStudioInstalled() == true))
+			{
+				bool bRegisterBuildPlatform = true;
+
+				// We also need to check for the generated projects... to handle the case where someone generates projects w/out WinRT.
+				// Hardcoding this for now - but ideally it would be dynamically discovered.
+				string EngineSourcePath = Path.Combine(ProjectFileGenerator.EngineRelativePath, "Source");
+				string WinRTRHIFile = Path.Combine(EngineSourcePath, "Runtime", "Windows", "D3D11RHI", "D3D11RHI.build.cs");
+				if (File.Exists(WinRTRHIFile) == false)
+				{
+					bRegisterBuildPlatform = false;
+				}
+
+				if (bRegisterBuildPlatform == true)
+				{
+					// Register this build platform for WinRT
+					Log.TraceVerbose("        Registering for {0}", UnrealTargetPlatform.WinRT.ToString());
+					UEBuildPlatform.RegisterBuildPlatform(new WinRTPlatform(UnrealTargetPlatform.WinRT, SDK));
+					UEBuildPlatform.RegisterPlatformWithGroup(UnrealTargetPlatform.WinRT, UnrealPlatformGroup.Microsoft);
+
+					// For now only register WinRT_ARM is truly a Windows 8 machine.
+					// This will prevent people who do all platform builds from running into the compiler issue.
+					if (WinRTPlatform.IsWindows8() == true)
+					{
+						Log.TraceVerbose("        Registering for {0}", UnrealTargetPlatform.WinRT_ARM.ToString());
+						UEBuildPlatform.RegisterBuildPlatform(new WinRTPlatform(UnrealTargetPlatform.WinRT_ARM, SDK));
+						UEBuildPlatform.RegisterPlatformWithGroup(UnrealTargetPlatform.WinRT_ARM, UnrealPlatformGroup.Microsoft);
+					}
+				}
+			}
+		}
+	}
 }
diff --git a/Engine/Source/Programs/UnrealBuildTool/WinRT/WinRTToolChain.cs b/Engine/Source/Programs/UnrealBuildTool/WinRT/WinRTToolChain.cs
--- a/Engine/Source/Programs/UnrealBuildTool/WinRT/WinRTToolChain.cs
+++ b/Engine/Source/Programs/UnrealBuildTool/WinRT/WinRTToolChain.cs
@@ -523,7 +523,7 @@ namespace UnrealBuildTool
 					string OriginalPCHHeaderDirectory = Path.GetDirectoryName(SourceFile.AbsolutePath);
 					FileArguments += string.Format(" /I \"{0}\"", OriginalPCHHeaderDirectory);
 
-					var PCHExtension = UEBuildPlatform.BuildPlatformDictionary[UnrealTargetPlatform.WinRT].GetBinaryExtension(UEBuildBinaryType.PrecompiledHeader);
+					var PCHExtension = UEBuildPlatform.GetBuildPlatform(UnrealTargetPlatform.WinRT).GetBinaryExtension(UEBuildBinaryType.PrecompiledHeader);
 
 					// Add the precompiled header file to the produced items list.
 					FileItem PrecompiledHeaderFile = FileItem.GetItemByFileReference(
@@ -580,7 +580,7 @@ namespace UnrealBuildTool
 					CompileAction.StatusDescription = Path.GetFileName(SourceFile.AbsolutePath);
 				}
 
-				var ObjectFileExtension = UEBuildPlatform.BuildPlatformDictionary[UnrealTargetPlatform.WinRT].GetBinaryExtension(UEBuildBinaryType.Object);
+				var ObjectFileExtension = UEBuildPlatform.GetBuildPlatform(UnrealTargetPlatform.WinRT).GetBinaryExtension(UEBuildBinaryType.Object);
 
 				// Add the object file to the produced item list.
 				FileItem ObjectFile = FileItem.GetItemByFileReference(
diff --git a/Engine/Source/Programs/UnrealBuildTool/Windows/UEBuildWindows.cs b/Engine/Source/Programs/UnrealBuildTool/Windows/UEBuildWindows.cs
--- a/Engine/Source/Programs/UnrealBuildTool/Windows/UEBuildWindows.cs
+++ b/Engine/Source/Programs/UnrealBuildTool/Windows/UEBuildWindows.cs
@@ -149,17 +149,31 @@ namespace UnrealBuildTool
 			}
 		}
 
+		WindowsPlatformSDK SDK;
+
 		/// <summary>
-		/// The current architecture
+		/// Constructor
 		/// </summary>
-		public override string GetActiveArchitecture()
+		/// <param name="InPlatform">Creates a windows platform with the given enum value</param>
+		public WindowsPlatform(UnrealTargetPlatform InPlatform, WindowsPlatformSDK InSDK) : base(InPlatform)
 		{
-			return IsWindowsXPSupported() ? "_XP" : base.GetActiveArchitecture();
+			SDK = InSDK;
 		}
 
-		protected override SDKStatus HasRequiredManualSDKInternal()
+		/// <summary>
+		/// Whether the required external SDKs are installed for this platform. Could be either a manual install or an AutoSDK.
+		/// </summary>
+		public override SDKStatus HasRequiredSDKsInstalled()
 		{
-			return SDKStatus.Valid;
+			return SDK.HasRequiredSDKsInstalled();
+		}
+
+		/// <summary>
+		/// The current architecture
+		/// </summary>
+		public override string GetActiveArchitecture()
+		{
+			return IsWindowsXPSupported() ? "_XP" : base.GetActiveArchitecture();
 		}
 
 		/// <summary>
@@ -241,23 +255,6 @@ namespace UnrealBuildTool
 			}
 		}
 
-		/// <summary>
-		/// Register the platform with the UEBuildPlatform class
-		/// </summary>
-		protected override void RegisterBuildPlatformInternal()
-		{
-			// Register this build platform for both Win64 and Win32
-			Log.TraceVerbose("        Registering for {0}", UnrealTargetPlatform.Win64.ToString());
-			UEBuildPlatform.RegisterBuildPlatform(UnrealTargetPlatform.Win64, this);
-			UEBuildPlatform.RegisterPlatformWithGroup(UnrealTargetPlatform.Win64, UnrealPlatformGroup.Windows);
-			UEBuildPlatform.RegisterPlatformWithGroup(UnrealTargetPlatform.Win64, UnrealPlatformGroup.Microsoft);
-
-			Log.TraceVerbose("        Registering for {0}", UnrealTargetPlatform.Win32.ToString());
-			UEBuildPlatform.RegisterBuildPlatform(UnrealTargetPlatform.Win32, this);
-			UEBuildPlatform.RegisterPlatformWithGroup(UnrealTargetPlatform.Win32, UnrealPlatformGroup.Windows);
-			UEBuildPlatform.RegisterPlatformWithGroup(UnrealTargetPlatform.Win32, UnrealPlatformGroup.Microsoft);
-		}
-
 		/// <summary>
 		/// Retrieve the CPPTargetPlatform for the given UnrealTargetPlatform
 		/// </summary>
@@ -803,4 +800,35 @@ namespace UnrealBuildTool
 			return false;
 		}
 	}
+
+	public class WindowsPlatformSDK : UEBuildPlatformSDK
+	{
+		protected override SDKStatus HasRequiredManualSDKInternal()
+		{
+			return SDKStatus.Valid;
+		}
+	}
+
+	class WindowsPlatformFactory : UEBuildPlatformFactory
+	{
+		/// <summary>
+		/// Register the platform with the UEBuildPlatform class
+		/// </summary>
+		public override void RegisterBuildPlatforms()
+		{
+			WindowsPlatformSDK SDK = new WindowsPlatformSDK();
+			SDK.ManageAndValidateSDK();
+
+			// Register this build platform for both Win64 and Win32
+			Log.TraceVerbose("        Registering for {0}", UnrealTargetPlatform.Win64.ToString());
+			UEBuildPlatform.RegisterBuildPlatform(new WindowsPlatform(UnrealTargetPlatform.Win64, SDK));
+			UEBuildPlatform.RegisterPlatformWithGroup(UnrealTargetPlatform.Win64, UnrealPlatformGroup.Windows);
+			UEBuildPlatform.RegisterPlatformWithGroup(UnrealTargetPlatform.Win64, UnrealPlatformGroup.Microsoft);
+
+			Log.TraceVerbose("        Registering for {0}", UnrealTargetPlatform.Win32.ToString());
+			UEBuildPlatform.RegisterBuildPlatform(new WindowsPlatform(UnrealTargetPlatform.Win32, SDK));
+			UEBuildPlatform.RegisterPlatformWithGroup(UnrealTargetPlatform.Win32, UnrealPlatformGroup.Windows);
+			UEBuildPlatform.RegisterPlatformWithGroup(UnrealTargetPlatform.Win32, UnrealPlatformGroup.Microsoft);
+		}
+	}
 }
diff --git a/Engine/Source/Programs/UnrealBuildTool/Windows/VCToolChain.cs b/Engine/Source/Programs/UnrealBuildTool/Windows/VCToolChain.cs
--- a/Engine/Source/Programs/UnrealBuildTool/Windows/VCToolChain.cs
+++ b/Engine/Source/Programs/UnrealBuildTool/Windows/VCToolChain.cs
@@ -1121,7 +1121,7 @@ namespace UnrealBuildTool
 					string OriginalPCHHeaderDirectory = Path.GetDirectoryName(SourceFile.AbsolutePath);
 					AddIncludePath(FileArguments, OriginalPCHHeaderDirectory);
 
-					var PrecompiledFileExtension = UEBuildPlatform.BuildPlatformDictionary[UnrealTargetPlatform.Win64].GetBinaryExtension(UEBuildBinaryType.PrecompiledHeader);
+					var PrecompiledFileExtension = UEBuildPlatform.GetBuildPlatform(UnrealTargetPlatform.Win64).GetBinaryExtension(UEBuildBinaryType.PrecompiledHeader);
 					// Add the precompiled header file to the produced items list.
 					FileItem PrecompiledHeaderFile = FileItem.GetItemByFileReference(
 						FileReference.Combine(
@@ -1215,7 +1215,7 @@ namespace UnrealBuildTool
 
 				if (bEmitsObjectFile)
 				{
-					var ObjectFileExtension = UEBuildPlatform.BuildPlatformDictionary[UnrealTargetPlatform.Win64].GetBinaryExtension(UEBuildBinaryType.Object);
+					var ObjectFileExtension = UEBuildPlatform.GetBuildPlatform(UnrealTargetPlatform.Win64).GetBinaryExtension(UEBuildBinaryType.Object);
 					// Add the object file to the produced item list.
 					FileItem ObjectFile = FileItem.GetItemByFileReference(
 						FileReference.Combine(
