Disable the compiler toolchain log message when we're generating project files on Linux; we're not actually compiling anything (just creating a large number of targets), so it's not helpful.

diff --git a/Engine/Source/Programs/UnrealBuildTool/Linux/LinuxToolChain.cs b/Engine/Source/Programs/UnrealBuildTool/Linux/LinuxToolChain.cs
--- a/Engine/Source/Programs/UnrealBuildTool/Linux/LinuxToolChain.cs
+++ b/Engine/Source/Programs/UnrealBuildTool/Linux/LinuxToolChain.cs
@@ -188,7 +188,7 @@ namespace UnrealBuildTool
 				// icl?
 			}
 
-			if (!CrossCompiling())
+			if (!CrossCompiling() && !ProjectFileGenerator.bGenerateProjectFiles)
 			{
 				Console.WriteLine("Using {0} version '{1}' (string), {2} (major), {3} (minor), {4} (patch)",
 					String.IsNullOrEmpty(ClangPath) ? "gcc" : "clang",
