Split the ModifyModuleRules() method into two parts; one which is called only for the active platform (ModifyModuleRulesForActivePlatform), and one which is called on all platforms for a host platform (ModifyModuleRulesForOtherPlatform).

diff --git a/Engine/Source/Programs/UnrealBuildTool/Android/UEBuildAndroid.cs b/Engine/Source/Programs/UnrealBuildTool/Android/UEBuildAndroid.cs
--- a/Engine/Source/Programs/UnrealBuildTool/Android/UEBuildAndroid.cs
+++ b/Engine/Source/Programs/UnrealBuildTool/Android/UEBuildAndroid.cs
@@ -151,7 +151,25 @@ namespace UnrealBuildTool
 		{
 		}
 
-		public override void ModifyModuleRules(string ModuleName, ModuleRules Rules, TargetInfo Target)
+		/// <summary>
+		/// Modify the rules for a newly created module, in a target that's being built for this platform.
+		/// This is not required - but allows for hiding details of a particular platform.
+		/// </summary>
+		/// <param name="ModuleName">The name of the module</param>
+		/// <param name="Rules">The module rules</param>
+		/// <param name="Target">The target being build</param>
+		public override void ModifyModuleRulesForActivePlatform(string ModuleName, ModuleRules Rules, TargetInfo Target)
+		{
+		}
+
+		/// <summary>
+		/// Modify the rules for a newly created module, where the target is a different host platform.
+		/// This is not required - but allows for hiding details of a particular platform.
+		/// </summary>
+		/// <param name="ModuleName">The name of the module</param>
+		/// <param name="Rules">The module rules</param>
+		/// <param name="Target">The target being build</param>
+		public override void ModifyModuleRulesForOtherPlatform(string ModuleName, ModuleRules Rules, TargetInfo Target)
 		{
 			if ((Target.Platform == UnrealTargetPlatform.Win32) || (Target.Platform == UnrealTargetPlatform.Win64) || (Target.Platform == UnrealTargetPlatform.Mac))
 			{
diff --git a/Engine/Source/Programs/UnrealBuildTool/Configuration/UEBuildPlatform.cs b/Engine/Source/Programs/UnrealBuildTool/Configuration/UEBuildPlatform.cs
--- a/Engine/Source/Programs/UnrealBuildTool/Configuration/UEBuildPlatform.cs
+++ b/Engine/Source/Programs/UnrealBuildTool/Configuration/UEBuildPlatform.cs
@@ -174,9 +174,13 @@ namespace UnrealBuildTool
 		{
 			foreach (var PlatformEntry in BuildPlatformDictionary)
 			{
-				if (Only == UnrealTargetPlatform.Unknown || PlatformEntry.Key == Only || PlatformEntry.Key == Target.Platform)
+				if(PlatformEntry.Key == Target.Platform)
 				{
-					PlatformEntry.Value.ModifyModuleRules(ModuleName, Rules, Target);
+					PlatformEntry.Value.ModifyModuleRulesForActivePlatform(ModuleName, Rules, Target);
+				}
+				else if (Only == UnrealTargetPlatform.Unknown || PlatformEntry.Key == Only)
+				{
+					PlatformEntry.Value.ModifyModuleRulesForOtherPlatform(ModuleName, Rules, Target);
 				}
 			}
 		}
@@ -358,14 +362,24 @@ namespace UnrealBuildTool
 		}
 
 		/// <summary>
-		/// Modify the newly created module passed in for this platform.
-		/// This is not required - but allows for hiding details of a
-		/// particular platform.
+		/// Modify the rules for a newly created module, in a target that's being built for this platform.
+		/// This is not required - but allows for hiding details of a particular platform.
 		/// </summary>
-		/// <param name="">ModuleName  The name of the module</param>
-		/// <param name="Rules">  The module rules</param>
-		/// <param name="Target">  The target being build</param>
-		public virtual void ModifyModuleRules(string ModuleName, ModuleRules Rules, TargetInfo Target)
+		/// <param name="ModuleName">The name of the module</param>
+		/// <param name="Rules">The module rules</param>
+		/// <param name="Target">The target being build</param>
+		public virtual void ModifyModuleRulesForActivePlatform(string ModuleName, ModuleRules Rules, TargetInfo Target)
+		{
+		}
+
+		/// <summary>
+		/// Modify the rules for a newly created module, where the target is a different host platform.
+		/// This is not required - but allows for hiding details of a particular platform.
+		/// </summary>
+		/// <param name="ModuleName">The name of the module</param>
+		/// <param name="Rules">The module rules</param>
+		/// <param name="Target">The target being build</param>
+		public virtual void ModifyModuleRulesForOtherPlatform(string ModuleName, ModuleRules Rules, TargetInfo Target)
 		{
 		}
 
diff --git a/Engine/Source/Programs/UnrealBuildTool/HTML5/UEBuildHTML5.cs b/Engine/Source/Programs/UnrealBuildTool/HTML5/UEBuildHTML5.cs
--- a/Engine/Source/Programs/UnrealBuildTool/HTML5/UEBuildHTML5.cs
+++ b/Engine/Source/Programs/UnrealBuildTool/HTML5/UEBuildHTML5.cs
@@ -209,31 +209,38 @@ namespace UnrealBuildTool
 		}
 
 		/// <summary>
-		/// Modify the newly created module passed in for this platform.
-		/// This is not required - but allows for hiding details of a
-		/// particular platform.
+		/// Modify the rules for a newly created module, in a target that's being built for this platform.
+		/// This is not required - but allows for hiding details of a particular platform.
 		/// </summary>
-		/// <param name="">Name   The name of the module</param>
-		/// <param name="Rules">  The module rules</param>
-		/// <param name="Target">  The target being build</param>
-		public override void ModifyModuleRules(string ModuleName, ModuleRules Rules, TargetInfo Target)
+		/// <param name="ModuleName">The name of the module</param>
+		/// <param name="Rules">The module rules</param>
+		/// <param name="Target">The target being build</param>
+		public override void ModifyModuleRulesForActivePlatform(string ModuleName, ModuleRules Rules, TargetInfo Target)
 		{
-			if (Target.Platform == UnrealTargetPlatform.HTML5)
+			if (ModuleName == "Core")
 			{
-				if (ModuleName == "Core")
-				{
-					Rules.PublicIncludePaths.Add("Runtime/Core/Public/HTML5");
-					Rules.PublicDependencyModuleNames.Add("zlib");
-				}
-				else if (ModuleName == "Engine")
-				{
-					Rules.PrivateDependencyModuleNames.Add("zlib");
-					Rules.PrivateDependencyModuleNames.Add("UElibPNG");
-					Rules.PublicDependencyModuleNames.Add("UEOgg");
-					Rules.PublicDependencyModuleNames.Add("Vorbis");
-				}
+				Rules.PublicIncludePaths.Add("Runtime/Core/Public/HTML5");
+				Rules.PublicDependencyModuleNames.Add("zlib");
+			}
+			else if (ModuleName == "Engine")
+			{
+				Rules.PrivateDependencyModuleNames.Add("zlib");
+				Rules.PrivateDependencyModuleNames.Add("UElibPNG");
+				Rules.PublicDependencyModuleNames.Add("UEOgg");
+				Rules.PublicDependencyModuleNames.Add("Vorbis");
 			}
-			else if (Target.Platform == UnrealTargetPlatform.Win32 || Target.Platform == UnrealTargetPlatform.Win64 || Target.Platform == UnrealTargetPlatform.Mac)
+		}
+
+		/// <summary>
+		/// Modify the rules for a newly created module, where the target is a different host platform.
+		/// This is not required - but allows for hiding details of a particular platform.
+		/// </summary>
+		/// <param name="ModuleName">The name of the module</param>
+		/// <param name="Rules">The module rules</param>
+		/// <param name="Target">The target being build</param>
+		public override void ModifyModuleRulesForOtherPlatform(string ModuleName, ModuleRules Rules, TargetInfo Target)
+		{
+			if (Target.Platform == UnrealTargetPlatform.Win32 || Target.Platform == UnrealTargetPlatform.Win64 || Target.Platform == UnrealTargetPlatform.Mac)
 			{
 				// allow standalone tools to use targetplatform modules, without needing Engine
 				if ((!UEBuildConfiguration.bBuildRequiresCookedData
diff --git a/Engine/Source/Programs/UnrealBuildTool/IOS/UEBuildIOS.cs b/Engine/Source/Programs/UnrealBuildTool/IOS/UEBuildIOS.cs
--- a/Engine/Source/Programs/UnrealBuildTool/IOS/UEBuildIOS.cs
+++ b/Engine/Source/Programs/UnrealBuildTool/IOS/UEBuildIOS.cs
@@ -387,14 +387,24 @@ namespace UnrealBuildTool
 		}
 
 		/// <summary>
-		/// Modify the newly created module passed in for this platform.
-		/// This is not required - but allows for hiding details of a
-		/// particular platform.
+		/// Modify the rules for a newly created module, in a target that's being built for this platform.
+		/// This is not required - but allows for hiding details of a particular platform.
 		/// </summary>
-		/// <param name="">Name   The name of the module</param>
-		/// <param name="Rules">  The module rules</param>
-		/// <param name="Target">  The target being build</param>
-		public override void ModifyModuleRules(string ModuleName, ModuleRules Rules, TargetInfo Target)
+		/// <param name="ModuleName">The name of the module</param>
+		/// <param name="Rules">The module rules</param>
+		/// <param name="Target">The target being build</param>
+		public override void ModifyModuleRulesForActivePlatform(string ModuleName, ModuleRules Rules, TargetInfo Target)
+		{
+		}
+
+		/// <summary>
+		/// Modify the rules for a newly created module, where the target is a different host platform.
+		/// This is not required - but allows for hiding details of a particular platform.
+		/// </summary>
+		/// <param name="ModuleName">The name of the module</param>
+		/// <param name="Rules">The module rules</param>
+		/// <param name="Target">The target being build</param>
+		public override void ModifyModuleRulesForOtherPlatform(string ModuleName, ModuleRules Rules, TargetInfo Target)
 		{
 			if ((Target.Platform == UnrealTargetPlatform.Win32) || (Target.Platform == UnrealTargetPlatform.Win64) || (Target.Platform == UnrealTargetPlatform.Mac))
 			{
diff --git a/Engine/Source/Programs/UnrealBuildTool/Linux/UEBuildLinux.cs b/Engine/Source/Programs/UnrealBuildTool/Linux/UEBuildLinux.cs
--- a/Engine/Source/Programs/UnrealBuildTool/Linux/UEBuildLinux.cs
+++ b/Engine/Source/Programs/UnrealBuildTool/Linux/UEBuildLinux.cs
@@ -164,14 +164,51 @@ namespace UnrealBuildTool
 		}
 
 		/// <summary>
-		/// Modify the newly created module passed in for this platform.
-		/// This is not required - but allows for hiding details of a
-		/// particular platform.
+		/// Modify the rules for a newly created module, in a target that's being built for this platform.
+		/// This is not required - but allows for hiding details of a particular platform.
 		/// </summary>
-		/// <param name="">Name   The name of the module</param>
-		/// <param name="Rules">  The module rules</param>
-		/// <param name="Target">  The target being build</param>
-		public override void ModifyModuleRules(string ModuleName, ModuleRules Rules, TargetInfo Target)
+		/// <param name="ModuleName">The name of the module</param>
+		/// <param name="Rules">The module rules</param>
+		/// <param name="Target">The target being build</param>
+		public override void ModifyModuleRulesForActivePlatform(string ModuleName, ModuleRules Rules, TargetInfo Target)
+		{
+			bool bBuildShaderFormats = UEBuildConfiguration.bForceBuildShaderFormats;
+
+			if (!UEBuildConfiguration.bBuildRequiresCookedData)
+			{
+				if (ModuleName == "TargetPlatform")
+				{
+					bBuildShaderFormats = true;
+				}
+			}
+
+			// allow standalone tools to use target platform modules, without needing Engine
+			if (ModuleName == "TargetPlatform")
+			{
+				if (UEBuildConfiguration.bForceBuildTargetPlatforms)
+				{
+					Rules.DynamicallyLoadedModuleNames.Add("LinuxTargetPlatform");
+					Rules.DynamicallyLoadedModuleNames.Add("LinuxNoEditorTargetPlatform");
+					Rules.DynamicallyLoadedModuleNames.Add("LinuxServerTargetPlatform");
+					Rules.DynamicallyLoadedModuleNames.Add("AllDesktopTargetPlatform");
+				}
+
+				if (bBuildShaderFormats)
+				{
+					// Rules.DynamicallyLoadedModuleNames.Add("ShaderFormatD3D");
+					Rules.DynamicallyLoadedModuleNames.Add("ShaderFormatOpenGL");
+				}
+			}
+		}
+
+		/// <summary>
+		/// Modify the rules for a newly created module, where the target is a different host platform.
+		/// This is not required - but allows for hiding details of a particular platform.
+		/// </summary>
+		/// <param name="ModuleName">The name of the module</param>
+		/// <param name="Rules">The module rules</param>
+		/// <param name="Target">The target being build</param>
+		public override void ModifyModuleRulesForOtherPlatform(string ModuleName, ModuleRules Rules, TargetInfo Target)
 		{
 			if ((Target.Platform == UnrealTargetPlatform.Win32) || (Target.Platform == UnrealTargetPlatform.Win64))
 			{
@@ -196,36 +233,6 @@ namespace UnrealBuildTool
 					Rules.PlatformSpecificDynamicallyLoadedModuleNames.Add("LinuxServerTargetPlatform");
 				}
 			}
-			else if (Target.Platform == UnrealTargetPlatform.Linux)
-			{
-				bool bBuildShaderFormats = UEBuildConfiguration.bForceBuildShaderFormats;
-
-				if (!UEBuildConfiguration.bBuildRequiresCookedData)
-				{
-					if (ModuleName == "TargetPlatform")
-					{
-						bBuildShaderFormats = true;
-					}
-				}
-
-				// allow standalone tools to use target platform modules, without needing Engine
-				if (ModuleName == "TargetPlatform")
-				{
-					if (UEBuildConfiguration.bForceBuildTargetPlatforms)
-					{
-						Rules.DynamicallyLoadedModuleNames.Add("LinuxTargetPlatform");
-						Rules.DynamicallyLoadedModuleNames.Add("LinuxNoEditorTargetPlatform");
-						Rules.DynamicallyLoadedModuleNames.Add("LinuxServerTargetPlatform");
-						Rules.DynamicallyLoadedModuleNames.Add("AllDesktopTargetPlatform");
-					}
-
-					if (bBuildShaderFormats)
-					{
-						// Rules.DynamicallyLoadedModuleNames.Add("ShaderFormatD3D");
-						Rules.DynamicallyLoadedModuleNames.Add("ShaderFormatOpenGL");
-					}
-				}
-			}
 		}
 
 		/// <summary>
diff --git a/Engine/Source/Programs/UnrealBuildTool/Mac/UEBuildMac.cs b/Engine/Source/Programs/UnrealBuildTool/Mac/UEBuildMac.cs
--- a/Engine/Source/Programs/UnrealBuildTool/Mac/UEBuildMac.cs
+++ b/Engine/Source/Programs/UnrealBuildTool/Mac/UEBuildMac.cs
@@ -81,41 +81,56 @@ namespace UnrealBuildTool
 			return BuildConfiguration.bGeneratedSYMFile || BuildConfiguration.bUsePDBFiles ? ".dSYM" : "";
 		}
 
-		public override void ModifyModuleRules(string ModuleName, ModuleRules Rules, TargetInfo Target)
+		/// <summary>
+		/// Modify the rules for a newly created module, in a target that's being built for this platform.
+		/// This is not required - but allows for hiding details of a particular platform.
+		/// </summary>
+		/// <param name="ModuleName">The name of the module</param>
+		/// <param name="Rules">The module rules</param>
+		/// <param name="Target">The target being build</param>
+		public override void ModifyModuleRulesForActivePlatform(string ModuleName, ModuleRules Rules, TargetInfo Target)
 		{
-			if (Target.Platform == UnrealTargetPlatform.Mac)
+			bool bBuildShaderFormats = UEBuildConfiguration.bForceBuildShaderFormats;
+
+			if (!UEBuildConfiguration.bBuildRequiresCookedData)
 			{
-				bool bBuildShaderFormats = UEBuildConfiguration.bForceBuildShaderFormats;
+				if (ModuleName == "TargetPlatform")
+				{
+					bBuildShaderFormats = true;
+				}
+			}
 
-				if (!UEBuildConfiguration.bBuildRequiresCookedData)
+			// allow standalone tools to use target platform modules, without needing Engine
+			if (ModuleName == "TargetPlatform")
+			{
+				if (UEBuildConfiguration.bForceBuildTargetPlatforms)
 				{
-					if (ModuleName == "TargetPlatform")
-					{
-						bBuildShaderFormats = true;
-					}
+					Rules.DynamicallyLoadedModuleNames.Add("MacTargetPlatform");
+					Rules.DynamicallyLoadedModuleNames.Add("MacNoEditorTargetPlatform");
+					Rules.DynamicallyLoadedModuleNames.Add("MacClientTargetPlatform");
+					Rules.DynamicallyLoadedModuleNames.Add("MacServerTargetPlatform");
+					Rules.DynamicallyLoadedModuleNames.Add("AllDesktopTargetPlatform");
 				}
 
-				// allow standalone tools to use target platform modules, without needing Engine
-				if (ModuleName == "TargetPlatform")
+				if (bBuildShaderFormats)
 				{
-					if (UEBuildConfiguration.bForceBuildTargetPlatforms)
-					{
-						Rules.DynamicallyLoadedModuleNames.Add("MacTargetPlatform");
-						Rules.DynamicallyLoadedModuleNames.Add("MacNoEditorTargetPlatform");
-						Rules.DynamicallyLoadedModuleNames.Add("MacClientTargetPlatform");
-						Rules.DynamicallyLoadedModuleNames.Add("MacServerTargetPlatform");
-						Rules.DynamicallyLoadedModuleNames.Add("AllDesktopTargetPlatform");
-					}
-
-					if (bBuildShaderFormats)
-					{
-						// Rules.DynamicallyLoadedModuleNames.Add("ShaderFormatD3D");
-						Rules.DynamicallyLoadedModuleNames.Add("ShaderFormatOpenGL");
-					}
+					// Rules.DynamicallyLoadedModuleNames.Add("ShaderFormatD3D");
+					Rules.DynamicallyLoadedModuleNames.Add("ShaderFormatOpenGL");
 				}
 			}
 		}
 
+		/// <summary>
+		/// Modify the rules for a newly created module, where the target is a different host platform.
+		/// This is not required - but allows for hiding details of a particular platform.
+		/// </summary>
+		/// <param name="ModuleName">The name of the module</param>
+		/// <param name="Rules">The module rules</param>
+		/// <param name="Target">The target being build</param>
+		public override void ModifyModuleRulesForOtherPlatform(string ModuleName, ModuleRules Rules, TargetInfo Target)
+		{
+		}
+
 		/// <summary>
 		/// Setup the target environment for building
 		/// </summary>
diff --git a/Engine/Source/Programs/UnrealBuildTool/UWP/UEBuildUWP.cs b/Engine/Source/Programs/UnrealBuildTool/UWP/UEBuildUWP.cs
--- a/Engine/Source/Programs/UnrealBuildTool/UWP/UEBuildUWP.cs
+++ b/Engine/Source/Programs/UnrealBuildTool/UWP/UEBuildUWP.cs
@@ -258,53 +258,68 @@ namespace UnrealBuildTool
 			UEBuildConfiguration.bCompileICU = true;
 		}
 
-		public override void ModifyModuleRules(string ModuleName, ModuleRules Rules, TargetInfo Target)
+		/// <summary>
+		/// Modify the rules for a newly created module, in a target that's being built for this platform.
+		/// This is not required - but allows for hiding details of a particular platform.
+		/// </summary>
+		/// <param name="ModuleName">The name of the module</param>
+		/// <param name="Rules">The module rules</param>
+		/// <param name="Target">The target being build</param>
+		public override void ModifyModuleRulesForActivePlatform(string ModuleName, ModuleRules Rules, TargetInfo Target)
 		{
-			if ((Target.Platform == UnrealTargetPlatform.UWP))
+			if (ModuleName == "Core")
 			{
-				if (ModuleName == "Core")
-				{
-					Rules.PrivateDependencyModuleNames.Add("UWPSDK");
-				}
-				else if (ModuleName == "Engine")
-				{
-					Rules.PrivateDependencyModuleNames.Add("zlib");
-					Rules.PrivateDependencyModuleNames.Add("UElibPNG");
-					Rules.PublicDependencyModuleNames.Add("UEOgg");
-					Rules.PublicDependencyModuleNames.Add("Vorbis");
-				}
-				else if (ModuleName == "D3D11RHI")
-				{
-					Rules.Definitions.Add("D3D11_WITH_DWMAPI=0");
-					Rules.Definitions.Add("WITH_DX_PERF=0");
-				}
-				else if (ModuleName == "DX11")
-				{
-					// Clear out all the Windows include paths and libraries...
-					// The UWPSDK module handles proper paths and libs for UWP.
-					// However, the D3D11RHI module will include the DX11 module.
-					Rules.PublicIncludePaths.Clear();
-					Rules.PublicLibraryPaths.Clear();
-					Rules.PublicAdditionalLibraries.Clear();
-					Rules.Definitions.Remove("WITH_D3DX_LIBS=1");
-					Rules.Definitions.Add("WITH_D3DX_LIBS=0");
-					Rules.PublicAdditionalLibraries.Remove("X3DAudio.lib");
-					Rules.PublicAdditionalLibraries.Remove("XAPOFX.lib");
-				}
-				else if (ModuleName == "XAudio2")
-				{
-					Rules.Definitions.Add("XAUDIO_SUPPORTS_XMA2WAVEFORMATEX=0");
-					Rules.Definitions.Add("XAUDIO_SUPPORTS_DEVICE_DETAILS=0");
-					Rules.Definitions.Add("XAUDIO2_SUPPORTS_MUSIC=0");
-					Rules.Definitions.Add("XAUDIO2_SUPPORTS_SENDLIST=0");
-					Rules.PublicAdditionalLibraries.Add("XAudio2.lib");
-				}
-				else if (ModuleName == "DX11Audio")
-				{
-					Rules.PublicAdditionalLibraries.Remove("X3DAudio.lib");
-					Rules.PublicAdditionalLibraries.Remove("XAPOFX.lib");
-				}
+				Rules.PrivateDependencyModuleNames.Add("UWPSDK");
 			}
+			else if (ModuleName == "Engine")
+			{
+				Rules.PrivateDependencyModuleNames.Add("zlib");
+				Rules.PrivateDependencyModuleNames.Add("UElibPNG");
+				Rules.PublicDependencyModuleNames.Add("UEOgg");
+				Rules.PublicDependencyModuleNames.Add("Vorbis");
+			}
+			else if (ModuleName == "D3D11RHI")
+			{
+				Rules.Definitions.Add("D3D11_WITH_DWMAPI=0");
+				Rules.Definitions.Add("WITH_DX_PERF=0");
+			}
+			else if (ModuleName == "DX11")
+			{
+				// Clear out all the Windows include paths and libraries...
+				// The UWPSDK module handles proper paths and libs for UWP.
+				// However, the D3D11RHI module will include the DX11 module.
+				Rules.PublicIncludePaths.Clear();
+				Rules.PublicLibraryPaths.Clear();
+				Rules.PublicAdditionalLibraries.Clear();
+				Rules.Definitions.Remove("WITH_D3DX_LIBS=1");
+				Rules.Definitions.Add("WITH_D3DX_LIBS=0");
+				Rules.PublicAdditionalLibraries.Remove("X3DAudio.lib");
+				Rules.PublicAdditionalLibraries.Remove("XAPOFX.lib");
+			}
+			else if (ModuleName == "XAudio2")
+			{
+				Rules.Definitions.Add("XAUDIO_SUPPORTS_XMA2WAVEFORMATEX=0");
+				Rules.Definitions.Add("XAUDIO_SUPPORTS_DEVICE_DETAILS=0");
+				Rules.Definitions.Add("XAUDIO2_SUPPORTS_MUSIC=0");
+				Rules.Definitions.Add("XAUDIO2_SUPPORTS_SENDLIST=0");
+				Rules.PublicAdditionalLibraries.Add("XAudio2.lib");
+			}
+			else if (ModuleName == "DX11Audio")
+			{
+				Rules.PublicAdditionalLibraries.Remove("X3DAudio.lib");
+				Rules.PublicAdditionalLibraries.Remove("XAPOFX.lib");
+			}
+		}
+
+		/// <summary>
+		/// Modify the rules for a newly created module, where the target is a different host platform.
+		/// This is not required - but allows for hiding details of a particular platform.
+		/// </summary>
+		/// <param name="ModuleName">The name of the module</param>
+		/// <param name="Rules">The module rules</param>
+		/// <param name="Target">The target being build</param>
+		public override void ModifyModuleRulesForOtherPlatform(string ModuleName, ModuleRules Rules, TargetInfo Target)
+		{
 		}
 
 		/// <summary>
diff --git a/Engine/Source/Programs/UnrealBuildTool/WinRT/UEBuildWinRT.cs b/Engine/Source/Programs/UnrealBuildTool/WinRT/UEBuildWinRT.cs
--- a/Engine/Source/Programs/UnrealBuildTool/WinRT/UEBuildWinRT.cs
+++ b/Engine/Source/Programs/UnrealBuildTool/WinRT/UEBuildWinRT.cs
@@ -191,245 +191,251 @@ namespace UnrealBuildTool
 		}
 
 		/// <summary>
-		/// Modify the newly created module passed in for this platform.
-		/// This is not required - but allows for hiding details of a
-		/// particular platform.
+		/// Modify the rules for a newly created module, in a target that's being built for this platform.
+		/// This is not required - but allows for hiding details of a particular platform.
 		/// </summary>
-		/// <param name="">Name   The name of the module</param>
-		/// <param name="Rules">  The module rules</param>
-		/// <param name="Target">  The target being build</param>
-		public override void ModifyModuleRules(string ModuleName, ModuleRules Rules, TargetInfo Target)
+		/// <param name="ModuleName">The name of the module</param>
+		/// <param name="Rules">The module rules</param>
+		/// <param name="Target">The target being build</param>
+		public override void ModifyModuleRulesForActivePlatform(string ModuleName, ModuleRules Rules, TargetInfo Target)
 		{
-			if ((Target.Platform == UnrealTargetPlatform.WinRT) ||
-				(Target.Platform == UnrealTargetPlatform.WinRT_ARM))
+			if (ModuleName == "Core")
+			{
+				Rules.PublicIncludePaths.Add("Runtime/Core/Public/WinRT");
+				Rules.PublicDependencyModuleNames.Add("zlib");
+			}
+			else if (ModuleName == "Engine")
+			{
+				Rules.PrivateDependencyModuleNames.Add("zlib");
+				Rules.PrivateDependencyModuleNames.Add("UElibPNG");
+				Rules.PublicDependencyModuleNames.Add("UEOgg");
+				Rules.PublicDependencyModuleNames.Add("Vorbis");
+			}
+			else if (ModuleName == "Launch")
+			{
+			}
+			else if (ModuleName == "D3D11RHI")
 			{
-				if (ModuleName == "Core")
+				Rules.Definitions.Add("D3D11_CUSTOM_VIEWPORT_CONSTRUCTOR=1");
+				// To enable platform specific D3D11 RHI Types
+				Rules.PrivateIncludePaths.Add("Runtime/Windows/D3D11RHI/Private/WinRT");
+				// Hack to enable AllowWindowsPlatformTypes.h/HideWindowsPlatformTypes.h
+				Rules.PublicIncludePaths.Add("Runtime/Core/Public/Windows");
+			}
+			else if (ModuleName == "Sockets")
+			{
+				// Hack to enable AllowWindowsPlatformTypes.h/HideWindowsPlatformTypes.h
+				Rules.PublicIncludePaths.Add("Runtime/Core/Public/Windows");
+			}
+			else if (ModuleName == "PhysX")
+			{
+				string PhysXDir = UEBuildConfiguration.UEThirdPartySourceDirectory + "PhysX/PhysX-3.3/";
+
+				Rules.PublicIncludePaths.Add("include/foundation/WinRT");
+				if (Target.Platform == UnrealTargetPlatform.WinRT)
 				{
-					Rules.PublicIncludePaths.Add("Runtime/Core/Public/WinRT");
-					Rules.PublicDependencyModuleNames.Add("zlib");
+					Rules.PublicLibraryPaths.Add(PhysXDir + "Lib/WinRT");
 				}
-				else if (ModuleName == "Engine")
+				else
 				{
-					Rules.PrivateDependencyModuleNames.Add("zlib");
-					Rules.PrivateDependencyModuleNames.Add("UElibPNG");
-					Rules.PublicDependencyModuleNames.Add("UEOgg");
-					Rules.PublicDependencyModuleNames.Add("Vorbis");
+					Rules.PublicLibraryPaths.Add(PhysXDir + "Lib/WinRT/ARM");
 				}
-				else if (ModuleName == "Launch")
+
+				if (Target.Configuration == UnrealTargetConfiguration.Debug)
 				{
+					Rules.PublicAdditionalLibraries.Add("PhysX3DEBUG.lib");
+					Rules.PublicAdditionalLibraries.Add("PhysX3ExtensionsDEBUG.lib");
+					Rules.PublicAdditionalLibraries.Add("PhysX3CookingDEBUG.lib");
+					Rules.PublicAdditionalLibraries.Add("PhysX3CommonDEBUG.lib");
+					Rules.PublicAdditionalLibraries.Add("PhysX3VehicleDEBUG.lib");
+					Rules.PublicAdditionalLibraries.Add("PxTaskDEBUG.lib");
+					Rules.PublicAdditionalLibraries.Add("PhysXVisualDebuggerSDKDEBUG.lib");
+					Rules.PublicAdditionalLibraries.Add("PhysXProfileSDKDEBUG.lib");
 				}
-				else if (ModuleName == "D3D11RHI")
+				else if (Target.Configuration == UnrealTargetConfiguration.Development)
 				{
-					Rules.Definitions.Add("D3D11_CUSTOM_VIEWPORT_CONSTRUCTOR=1");
-					// To enable platform specific D3D11 RHI Types
-					Rules.PrivateIncludePaths.Add("Runtime/Windows/D3D11RHI/Private/WinRT");
-					// Hack to enable AllowWindowsPlatformTypes.h/HideWindowsPlatformTypes.h
-					Rules.PublicIncludePaths.Add("Runtime/Core/Public/Windows");
+					Rules.PublicAdditionalLibraries.Add("PhysX3PROFILE.lib");
+					Rules.PublicAdditionalLibraries.Add("PhysX3ExtensionsPROFILE.lib");
+					Rules.PublicAdditionalLibraries.Add("PhysX3CookingPROFILE.lib");
+					Rules.PublicAdditionalLibraries.Add("PhysX3CommonPROFILE.lib");
+					Rules.PublicAdditionalLibraries.Add("PhysX3VehiclePROFILE.lib");
+					Rules.PublicAdditionalLibraries.Add("PxTaskPROFILE.lib");
+					Rules.PublicAdditionalLibraries.Add("PhysXVisualDebuggerSDKPROFILE.lib");
+					Rules.PublicAdditionalLibraries.Add("PhysXProfileSDKPROFILE.lib");
 				}
-				else if (ModuleName == "Sockets")
+				else // Test or Shipping
 				{
-					// Hack to enable AllowWindowsPlatformTypes.h/HideWindowsPlatformTypes.h
-					Rules.PublicIncludePaths.Add("Runtime/Core/Public/Windows");
+					Rules.PublicAdditionalLibraries.Add("PhysX3.lib");
+					Rules.PublicAdditionalLibraries.Add("PhysX3Extensions.lib");
+					Rules.PublicAdditionalLibraries.Add("PhysX3Cooking.lib");
+					Rules.PublicAdditionalLibraries.Add("PhysX3Common.lib");
+					Rules.PublicAdditionalLibraries.Add("PhysX3Vehicle.lib");
+					Rules.PublicAdditionalLibraries.Add("PxTask.lib");
+					Rules.PublicAdditionalLibraries.Add("PhysXVisualDebuggerSDK.lib");
+					Rules.PublicAdditionalLibraries.Add("PhysXProfileSDK.lib");
 				}
-				else if (ModuleName == "PhysX")
-				{
-					string PhysXDir = UEBuildConfiguration.UEThirdPartySourceDirectory + "PhysX/PhysX-3.3/";
-
-					Rules.PublicIncludePaths.Add("include/foundation/WinRT");
-					if (Target.Platform == UnrealTargetPlatform.WinRT)
-					{
-						Rules.PublicLibraryPaths.Add(PhysXDir + "Lib/WinRT");
-					}
-					else
-					{
-						Rules.PublicLibraryPaths.Add(PhysXDir + "Lib/WinRT/ARM");
-					}
+			}
+			else if (ModuleName == "APEX")
+			{
+				Rules.Definitions.Remove("APEX_STATICALLY_LINKED=0");
+				Rules.Definitions.Add("APEX_STATICALLY_LINKED=1");
 
-					if (Target.Configuration == UnrealTargetConfiguration.Debug)
-					{
-						Rules.PublicAdditionalLibraries.Add("PhysX3DEBUG.lib");
-						Rules.PublicAdditionalLibraries.Add("PhysX3ExtensionsDEBUG.lib");
-						Rules.PublicAdditionalLibraries.Add("PhysX3CookingDEBUG.lib");
-						Rules.PublicAdditionalLibraries.Add("PhysX3CommonDEBUG.lib");
-						Rules.PublicAdditionalLibraries.Add("PhysX3VehicleDEBUG.lib");
-						Rules.PublicAdditionalLibraries.Add("PxTaskDEBUG.lib");
-						Rules.PublicAdditionalLibraries.Add("PhysXVisualDebuggerSDKDEBUG.lib");
-						Rules.PublicAdditionalLibraries.Add("PhysXProfileSDKDEBUG.lib");
-					}
-					else if (Target.Configuration == UnrealTargetConfiguration.Development)
-					{
-						Rules.PublicAdditionalLibraries.Add("PhysX3PROFILE.lib");
-						Rules.PublicAdditionalLibraries.Add("PhysX3ExtensionsPROFILE.lib");
-						Rules.PublicAdditionalLibraries.Add("PhysX3CookingPROFILE.lib");
-						Rules.PublicAdditionalLibraries.Add("PhysX3CommonPROFILE.lib");
-						Rules.PublicAdditionalLibraries.Add("PhysX3VehiclePROFILE.lib");
-						Rules.PublicAdditionalLibraries.Add("PxTaskPROFILE.lib");
-						Rules.PublicAdditionalLibraries.Add("PhysXVisualDebuggerSDKPROFILE.lib");
-						Rules.PublicAdditionalLibraries.Add("PhysXProfileSDKPROFILE.lib");
-					}
-					else // Test or Shipping
-					{
-						Rules.PublicAdditionalLibraries.Add("PhysX3.lib");
-						Rules.PublicAdditionalLibraries.Add("PhysX3Extensions.lib");
-						Rules.PublicAdditionalLibraries.Add("PhysX3Cooking.lib");
-						Rules.PublicAdditionalLibraries.Add("PhysX3Common.lib");
-						Rules.PublicAdditionalLibraries.Add("PhysX3Vehicle.lib");
-						Rules.PublicAdditionalLibraries.Add("PxTask.lib");
-						Rules.PublicAdditionalLibraries.Add("PhysXVisualDebuggerSDK.lib");
-						Rules.PublicAdditionalLibraries.Add("PhysXProfileSDK.lib");
-					}
+				string APEXDir = UEBuildConfiguration.UEThirdPartySourceDirectory + "PhysX/APEX-1.3/";
+				if (Target.Platform == UnrealTargetPlatform.WinRT)
+				{
+					Rules.PublicLibraryPaths.Add(APEXDir + "lib/WinRT");
 				}
-				else if (ModuleName == "APEX")
+				else
 				{
-					Rules.Definitions.Remove("APEX_STATICALLY_LINKED=0");
-					Rules.Definitions.Add("APEX_STATICALLY_LINKED=1");
-
-					string APEXDir = UEBuildConfiguration.UEThirdPartySourceDirectory + "PhysX/APEX-1.3/";
-					if (Target.Platform == UnrealTargetPlatform.WinRT)
-					{
-						Rules.PublicLibraryPaths.Add(APEXDir + "lib/WinRT");
-					}
-					else
-					{
-						Rules.PublicLibraryPaths.Add(APEXDir + "lib/WinRT/ARM");
-					}
+					Rules.PublicLibraryPaths.Add(APEXDir + "lib/WinRT/ARM");
+				}
 
-					if (Target.Configuration == UnrealTargetConfiguration.Debug)
-					{
-						Rules.PublicAdditionalLibraries.Add("ApexCommonDEBUG.lib");
-						Rules.PublicAdditionalLibraries.Add("ApexFrameworkDEBUG.lib");
-						Rules.PublicAdditionalLibraries.Add("ApexSharedDEBUG.lib");
-						Rules.PublicAdditionalLibraries.Add("APEX_DestructibleDEBUG.lib");
+				if (Target.Configuration == UnrealTargetConfiguration.Debug)
+				{
+					Rules.PublicAdditionalLibraries.Add("ApexCommonDEBUG.lib");
+					Rules.PublicAdditionalLibraries.Add("ApexFrameworkDEBUG.lib");
+					Rules.PublicAdditionalLibraries.Add("ApexSharedDEBUG.lib");
+					Rules.PublicAdditionalLibraries.Add("APEX_DestructibleDEBUG.lib");
 
-					}
-					else if (Target.Configuration == UnrealTargetConfiguration.Development)
-					{
-						Rules.PublicAdditionalLibraries.Add("ApexCommonPROFILE.lib");
-						Rules.PublicAdditionalLibraries.Add("ApexFrameworkPROFILE.lib");
-						Rules.PublicAdditionalLibraries.Add("ApexSharedPROFILE.lib");
-						Rules.PublicAdditionalLibraries.Add("APEX_DestructiblePROFILE.lib");
-					}
-					else // Test or Shipping
-					{
-						Rules.PublicAdditionalLibraries.Add("ApexCommon.lib");
-						Rules.PublicAdditionalLibraries.Add("ApexFramework.lib");
-						Rules.PublicAdditionalLibraries.Add("ApexShared.lib");
-						Rules.PublicAdditionalLibraries.Add("APEX_Destructible.lib");
-					}
 				}
-				else if (ModuleName == "FreeType2")
+				else if (Target.Configuration == UnrealTargetConfiguration.Development)
 				{
-					string FreeType2Path = UEBuildConfiguration.UEThirdPartySourceDirectory + "FreeType2/FreeType2-2.4.12/";
-					if (Target.Platform == UnrealTargetPlatform.WinRT)
-					{
-						Rules.PublicLibraryPaths.Add(FreeType2Path + "Lib/WinRT/Win64");
-					}
-					else
-					{
-						Rules.PublicLibraryPaths.Add(FreeType2Path + "Lib/WinRT/ARM");
-					}
-					Rules.PublicAdditionalLibraries.Add("freetype2412MT.lib");
+					Rules.PublicAdditionalLibraries.Add("ApexCommonPROFILE.lib");
+					Rules.PublicAdditionalLibraries.Add("ApexFrameworkPROFILE.lib");
+					Rules.PublicAdditionalLibraries.Add("ApexSharedPROFILE.lib");
+					Rules.PublicAdditionalLibraries.Add("APEX_DestructiblePROFILE.lib");
 				}
-				else if (ModuleName == "UElibPNG")
+				else // Test or Shipping
 				{
-					string libPNGPath = UEBuildConfiguration.UEThirdPartySourceDirectory + "libPNG/libPNG-1.5.2";
-					if (Target.Platform == UnrealTargetPlatform.WinRT)
-					{
-						Rules.PublicLibraryPaths.Add(libPNGPath + "/lib/WinRT/Win64");
-					}
-					else
-					{
-						Rules.PublicLibraryPaths.Add(libPNGPath + "/lib/WinRT/ARM");
-					}
-					Rules.PublicAdditionalLibraries.Add("libpng125.lib");
+					Rules.PublicAdditionalLibraries.Add("ApexCommon.lib");
+					Rules.PublicAdditionalLibraries.Add("ApexFramework.lib");
+					Rules.PublicAdditionalLibraries.Add("ApexShared.lib");
+					Rules.PublicAdditionalLibraries.Add("APEX_Destructible.lib");
 				}
-				else if (ModuleName == "DX11")
+			}
+			else if (ModuleName == "FreeType2")
+			{
+				string FreeType2Path = UEBuildConfiguration.UEThirdPartySourceDirectory + "FreeType2/FreeType2-2.4.12/";
+				if (Target.Platform == UnrealTargetPlatform.WinRT)
 				{
-					// Clear out all the Windows include paths and libraries...
-					// The WinRTSDK module handles proper paths and libs for WinRT.
-					// However, the D3D11RHI module will include the DX11 module.
-					Rules.PublicIncludePaths.Clear();
-					Rules.PublicLibraryPaths.Clear();
-					Rules.PublicAdditionalLibraries.Clear();
-					Rules.Definitions.Remove("WITH_D3DX_LIBS=1");
-					Rules.Definitions.Add("D3D11_WITH_DWMAPI=0");
-					Rules.Definitions.Add("WITH_D3DX_LIBS=0");
-					Rules.Definitions.Add("WITH_DX_PERF=0");
-					Rules.PublicAdditionalLibraries.Remove("X3DAudio.lib");
-					Rules.PublicAdditionalLibraries.Remove("XAPOFX.lib");
+					Rules.PublicLibraryPaths.Add(FreeType2Path + "Lib/WinRT/Win64");
 				}
-				else if (ModuleName == "XInput")
+				else
 				{
-					Rules.PublicAdditionalLibraries.Add("XInput.lib");
+					Rules.PublicLibraryPaths.Add(FreeType2Path + "Lib/WinRT/ARM");
 				}
-				else if (ModuleName == "XAudio2")
+				Rules.PublicAdditionalLibraries.Add("freetype2412MT.lib");
+			}
+			else if (ModuleName == "UElibPNG")
+			{
+				string libPNGPath = UEBuildConfiguration.UEThirdPartySourceDirectory + "libPNG/libPNG-1.5.2";
+				if (Target.Platform == UnrealTargetPlatform.WinRT)
 				{
-					Rules.Definitions.Add("XAUDIO_SUPPORTS_XMA2WAVEFORMATEX=0");
-					Rules.Definitions.Add("XAUDIO_SUPPORTS_DEVICE_DETAILS=0");
-					Rules.Definitions.Add("XAUDIO2_SUPPORTS_MUSIC=0");
-					Rules.Definitions.Add("XAUDIO2_SUPPORTS_SENDLIST=0");
-					Rules.PublicAdditionalLibraries.Add("XAudio2.lib");
-					// Hack to enable AllowWindowsPlatformTypes.h/HideWindowsPlatformTypes.h
-					Rules.PublicIncludePaths.Add("Runtime/Core/Public/Windows");
+					Rules.PublicLibraryPaths.Add(libPNGPath + "/lib/WinRT/Win64");
 				}
-				else if (ModuleName == "UEOgg")
+				else
 				{
-					string OggPath = UEBuildConfiguration.UEThirdPartySourceDirectory + "Ogg/libogg-1.2.2/";
-					if (Target.Platform == UnrealTargetPlatform.WinRT)
-					{
-						Rules.PublicLibraryPaths.Add(OggPath + "WinRT/VS2012/WinRT/x64/Release");
-					}
-					else
-					{
-						Rules.PublicLibraryPaths.Add(OggPath + "WinRT/VS2012/WinRT/ARM/Release");
-					}
-					Rules.PublicAdditionalLibraries.Add("libogg_static.lib");
+					Rules.PublicLibraryPaths.Add(libPNGPath + "/lib/WinRT/ARM");
 				}
-				else if (ModuleName == "Vorbis")
+				Rules.PublicAdditionalLibraries.Add("libpng125.lib");
+			}
+			else if (ModuleName == "DX11")
+			{
+				// Clear out all the Windows include paths and libraries...
+				// The WinRTSDK module handles proper paths and libs for WinRT.
+				// However, the D3D11RHI module will include the DX11 module.
+				Rules.PublicIncludePaths.Clear();
+				Rules.PublicLibraryPaths.Clear();
+				Rules.PublicAdditionalLibraries.Clear();
+				Rules.Definitions.Remove("WITH_D3DX_LIBS=1");
+				Rules.Definitions.Add("D3D11_WITH_DWMAPI=0");
+				Rules.Definitions.Add("WITH_D3DX_LIBS=0");
+				Rules.Definitions.Add("WITH_DX_PERF=0");
+				Rules.PublicAdditionalLibraries.Remove("X3DAudio.lib");
+				Rules.PublicAdditionalLibraries.Remove("XAPOFX.lib");
+			}
+			else if (ModuleName == "XInput")
+			{
+				Rules.PublicAdditionalLibraries.Add("XInput.lib");
+			}
+			else if (ModuleName == "XAudio2")
+			{
+				Rules.Definitions.Add("XAUDIO_SUPPORTS_XMA2WAVEFORMATEX=0");
+				Rules.Definitions.Add("XAUDIO_SUPPORTS_DEVICE_DETAILS=0");
+				Rules.Definitions.Add("XAUDIO2_SUPPORTS_MUSIC=0");
+				Rules.Definitions.Add("XAUDIO2_SUPPORTS_SENDLIST=0");
+				Rules.PublicAdditionalLibraries.Add("XAudio2.lib");
+				// Hack to enable AllowWindowsPlatformTypes.h/HideWindowsPlatformTypes.h
+				Rules.PublicIncludePaths.Add("Runtime/Core/Public/Windows");
+			}
+			else if (ModuleName == "UEOgg")
+			{
+				string OggPath = UEBuildConfiguration.UEThirdPartySourceDirectory + "Ogg/libogg-1.2.2/";
+				if (Target.Platform == UnrealTargetPlatform.WinRT)
 				{
-					string VorbisPath = UEBuildConfiguration.UEThirdPartySourceDirectory + "Vorbis/libvorbis-1.3.2/";
-					if (Target.Platform == UnrealTargetPlatform.WinRT)
-					{
-						Rules.PublicLibraryPaths.Add(VorbisPath + "WinRT/VS2012/WinRT/x64/Release");
-					}
-					else
-					{
-						Rules.PublicLibraryPaths.Add(VorbisPath + "WinRT/VS2012/WinRT/ARM/Release");
-					}
-					Rules.PublicAdditionalLibraries.Add("libvorbis_static.lib");
+					Rules.PublicLibraryPaths.Add(OggPath + "WinRT/VS2012/WinRT/x64/Release");
 				}
-				else if (ModuleName == "VorbisFile")
+				else
 				{
-					string VorbisPath = UEBuildConfiguration.UEThirdPartySourceDirectory + "Vorbis/libvorbis-1.3.2/";
-					if (Target.Platform == UnrealTargetPlatform.WinRT)
-					{
-						Rules.PublicLibraryPaths.Add(VorbisPath + "WinRT/VS2012/WinRT/x64/Release");
-					}
-					else
-					{
-						Rules.PublicLibraryPaths.Add(VorbisPath + "WinRT/VS2012/WinRT/ARM/Release");
-					}
-					Rules.PublicAdditionalLibraries.Add("libvorbisfile_static.lib");
+					Rules.PublicLibraryPaths.Add(OggPath + "WinRT/VS2012/WinRT/ARM/Release");
 				}
-				else if (ModuleName == "DX11Audio")
+				Rules.PublicAdditionalLibraries.Add("libogg_static.lib");
+			}
+			else if (ModuleName == "Vorbis")
+			{
+				string VorbisPath = UEBuildConfiguration.UEThirdPartySourceDirectory + "Vorbis/libvorbis-1.3.2/";
+				if (Target.Platform == UnrealTargetPlatform.WinRT)
 				{
-					Rules.PublicAdditionalLibraries.Remove("X3DAudio.lib");
-					Rules.PublicAdditionalLibraries.Remove("XAPOFX.lib");
+					Rules.PublicLibraryPaths.Add(VorbisPath + "WinRT/VS2012/WinRT/x64/Release");
 				}
-				else if (ModuleName == "zlib")
+				else
 				{
-					if (Target.Platform == UnrealTargetPlatform.WinRT)
-					{
-						Rules.PublicLibraryPaths.Add(UEBuildConfiguration.UEThirdPartySourceDirectory + "zlib/zlib-1.2.5/Lib/WinRT/Win64");
-					}
-					else
-					{
-						Rules.PublicLibraryPaths.Add(UEBuildConfiguration.UEThirdPartySourceDirectory + "zlib/zlib-1.2.5/Lib/WinRT/ARM");
-					}
-					Rules.PublicAdditionalLibraries.Add("zlib125.lib");
+					Rules.PublicLibraryPaths.Add(VorbisPath + "WinRT/VS2012/WinRT/ARM/Release");
+				}
+				Rules.PublicAdditionalLibraries.Add("libvorbis_static.lib");
+			}
+			else if (ModuleName == "VorbisFile")
+			{
+				string VorbisPath = UEBuildConfiguration.UEThirdPartySourceDirectory + "Vorbis/libvorbis-1.3.2/";
+				if (Target.Platform == UnrealTargetPlatform.WinRT)
+				{
+					Rules.PublicLibraryPaths.Add(VorbisPath + "WinRT/VS2012/WinRT/x64/Release");
+				}
+				else
+				{
+					Rules.PublicLibraryPaths.Add(VorbisPath + "WinRT/VS2012/WinRT/ARM/Release");
+				}
+				Rules.PublicAdditionalLibraries.Add("libvorbisfile_static.lib");
+			}
+			else if (ModuleName == "DX11Audio")
+			{
+				Rules.PublicAdditionalLibraries.Remove("X3DAudio.lib");
+				Rules.PublicAdditionalLibraries.Remove("XAPOFX.lib");
+			}
+			else if (ModuleName == "zlib")
+			{
+				if (Target.Platform == UnrealTargetPlatform.WinRT)
+				{
+					Rules.PublicLibraryPaths.Add(UEBuildConfiguration.UEThirdPartySourceDirectory + "zlib/zlib-1.2.5/Lib/WinRT/Win64");
+				}
+				else
+				{
+					Rules.PublicLibraryPaths.Add(UEBuildConfiguration.UEThirdPartySourceDirectory + "zlib/zlib-1.2.5/Lib/WinRT/ARM");
 				}
+				Rules.PublicAdditionalLibraries.Add("zlib125.lib");
 			}
-			else if ((Target.Platform == UnrealTargetPlatform.Win32) || (Target.Platform == UnrealTargetPlatform.Win64))
+		}
+
+		/// <summary>
+		/// Modify the rules for a newly created module, where the target is a different host platform.
+		/// This is not required - but allows for hiding details of a particular platform.
+		/// </summary>
+		/// <param name="ModuleName">The name of the module</param>
+		/// <param name="Rules">The module rules</param>
+		/// <param name="Target">The target being build</param>
+		public override void ModifyModuleRulesForOtherPlatform(string ModuleName, ModuleRules Rules, TargetInfo Target)
+		{
+			if ((Target.Platform == UnrealTargetPlatform.Win32) || (Target.Platform == UnrealTargetPlatform.Win64))
 			{
 				//              bool bBuildShaderFormats = UEBuildConfiguration.bForceBuildShaderFormats;
 				// 				if (!UEBuildConfiguration.bBuildRequiresCookedData)
diff --git a/Engine/Source/Programs/UnrealBuildTool/Windows/UEBuildWindows.cs b/Engine/Source/Programs/UnrealBuildTool/Windows/UEBuildWindows.cs
--- a/Engine/Source/Programs/UnrealBuildTool/Windows/UEBuildWindows.cs
+++ b/Engine/Source/Programs/UnrealBuildTool/Windows/UEBuildWindows.cs
@@ -463,62 +463,77 @@ namespace UnrealBuildTool
 			UEBuildConfiguration.bCompileICU = true;
 		}
 
-		public override void ModifyModuleRules(string ModuleName, ModuleRules Rules, TargetInfo Target)
+		/// <summary>
+		/// Modify the rules for a newly created module, in a target that's being built for this platform.
+		/// This is not required - but allows for hiding details of a particular platform.
+		/// </summary>
+		/// <param name="ModuleName">The name of the module</param>
+		/// <param name="Rules">The module rules</param>
+		/// <param name="Target">The target being build</param>
+		public override void ModifyModuleRulesForActivePlatform(string ModuleName, ModuleRules Rules, TargetInfo Target)
 		{
-			if ((Target.Platform == UnrealTargetPlatform.Win32) || (Target.Platform == UnrealTargetPlatform.Win64))
+			bool bBuildShaderFormats = UEBuildConfiguration.bForceBuildShaderFormats;
+
+			if (!UEBuildConfiguration.bBuildRequiresCookedData)
 			{
-				bool bBuildShaderFormats = UEBuildConfiguration.bForceBuildShaderFormats;
+				if (ModuleName == "TargetPlatform")
+				{
+					bBuildShaderFormats = true;
+				}
+			}
 
-				if (!UEBuildConfiguration.bBuildRequiresCookedData)
+			// allow standalone tools to use target platform modules, without needing Engine
+			if (ModuleName == "TargetPlatform")
+			{
+				if (UEBuildConfiguration.bForceBuildTargetPlatforms)
 				{
-					if (ModuleName == "TargetPlatform")
-					{
-						bBuildShaderFormats = true;
-					}
+					Rules.DynamicallyLoadedModuleNames.Add("WindowsTargetPlatform");
+					Rules.DynamicallyLoadedModuleNames.Add("WindowsNoEditorTargetPlatform");
+					Rules.DynamicallyLoadedModuleNames.Add("WindowsServerTargetPlatform");
+					Rules.DynamicallyLoadedModuleNames.Add("WindowsClientTargetPlatform");
+					Rules.DynamicallyLoadedModuleNames.Add("AllDesktopTargetPlatform");
 				}
 
-				// allow standalone tools to use target platform modules, without needing Engine
-				if (ModuleName == "TargetPlatform")
+				if (bBuildShaderFormats)
 				{
-					if (UEBuildConfiguration.bForceBuildTargetPlatforms)
-					{
-						Rules.DynamicallyLoadedModuleNames.Add("WindowsTargetPlatform");
-						Rules.DynamicallyLoadedModuleNames.Add("WindowsNoEditorTargetPlatform");
-						Rules.DynamicallyLoadedModuleNames.Add("WindowsServerTargetPlatform");
-						Rules.DynamicallyLoadedModuleNames.Add("WindowsClientTargetPlatform");
-						Rules.DynamicallyLoadedModuleNames.Add("AllDesktopTargetPlatform");
-					}
+					Rules.DynamicallyLoadedModuleNames.Add("ShaderFormatD3D");
+					Rules.DynamicallyLoadedModuleNames.Add("ShaderFormatOpenGL");
 
-					if (bBuildShaderFormats)
+					//#todo-rco: Remove when public
+					string VulkanSDKPath = Environment.GetEnvironmentVariable("VK_SDK_PATH");
 					{
-						Rules.DynamicallyLoadedModuleNames.Add("ShaderFormatD3D");
-						Rules.DynamicallyLoadedModuleNames.Add("ShaderFormatOpenGL");
-
-						//#todo-rco: Remove when public
-						string VulkanSDKPath = Environment.GetEnvironmentVariable("VK_SDK_PATH");
+						if (!String.IsNullOrEmpty(VulkanSDKPath))
 						{
-							if (!String.IsNullOrEmpty(VulkanSDKPath))
-							{
-								Rules.DynamicallyLoadedModuleNames.Add("VulkanShaderFormat");
-							}
+							Rules.DynamicallyLoadedModuleNames.Add("VulkanShaderFormat");
 						}
 					}
 				}
+			}
 
-				if (ModuleName == "D3D11RHI")
-				{
-					// To enable platform specific D3D11 RHI Types
-					Rules.PrivateIncludePaths.Add("Runtime/Windows/D3D11RHI/Private/Windows");
-				}
+			if (ModuleName == "D3D11RHI")
+			{
+				// To enable platform specific D3D11 RHI Types
+				Rules.PrivateIncludePaths.Add("Runtime/Windows/D3D11RHI/Private/Windows");
+			}
 
-				if (ModuleName == "D3D12RHI")
-				{
-					// To enable platform specific D3D12 RHI Types
-					Rules.PrivateIncludePaths.Add("Runtime/Windows/D3D12RHI/Private/Windows");
-				}
+			if (ModuleName == "D3D12RHI")
+			{
+				// To enable platform specific D3D12 RHI Types
+				Rules.PrivateIncludePaths.Add("Runtime/Windows/D3D12RHI/Private/Windows");
 			}
 		}
 
+		/// <summary>
+		/// Modify the rules for a newly created module, where the target is a different host platform.
+		/// This is not required - but allows for hiding details of a particular platform.
+		/// </summary>
+		/// <param name="ModuleName">The name of the module</param>
+		/// <param name="Rules">The module rules</param>
+		/// <param name="Target">The target being build</param>
+		public override void ModifyModuleRulesForOtherPlatform(string ModuleName, ModuleRules Rules, TargetInfo Target)
+		{
+		}
+
 		/// <summary>
 		/// Setup the target environment for building
 		/// </summary>
