Rearranged link command building in MacToolChain to not require third party libs to have "lib" prefix

diff --git a/Engine/Source/Programs/UnrealBuildTool/Mac/MacToolChain.cs b/Engine/Source/Programs/UnrealBuildTool/Mac/MacToolChain.cs
--- a/Engine/Source/Programs/UnrealBuildTool/Mac/MacToolChain.cs
+++ b/Engine/Source/Programs/UnrealBuildTool/Mac/MacToolChain.cs
@@ -679,25 +679,7 @@ namespace UnrealBuildTool
 						continue;
 					}
 
-					if (Path.GetFileName(AdditionalLibrary).StartsWith("lib"))
-					{
-						LinkCommand += string.Format(" \"{0}\"", ConvertPath(Path.GetFullPath(AdditionalLibrary)));
-						if (Path.GetExtension(AdditionalLibrary) == ".dylib")
-						{
-							ThirdPartyLibraries.Add(AdditionalLibrary);
-						}
-
-						if (BuildHostPlatform.Current.Platform != UnrealTargetPlatform.Mac)
-						{
-							// copy over libs we may need
-							FileItem ShadowFile = FileItem.GetExistingItemByPath(AdditionalLibrary);
-							if (ShadowFile != null)
-							{
-								QueueFileForBatchUpload(ShadowFile);
-							}
-						}
-					}
-					else if (Path.GetDirectoryName(AdditionalLibrary) != "" &&
+					if (Path.GetDirectoryName(AdditionalLibrary) != "" &&
 							 (Path.GetDirectoryName(AdditionalLibrary).Contains("Binaries/Mac") ||
 							 Path.GetDirectoryName(AdditionalLibrary).Contains("Binaries\\Mac")))
 					{
@@ -722,10 +704,28 @@ namespace UnrealBuildTool
 					{
 						LinkCommand += string.Format(" \'{0}\'", AdditionalLibrary);
 					}
-					else
+					else  if (string.IsNullOrEmpty(Path.GetDirectoryName(AdditionalLibrary)) && string.IsNullOrEmpty(Path.GetExtension(AdditionalLibrary)))
 					{
 						LinkCommand += string.Format(" -l{0}", AdditionalLibrary);
 					}
+					else
+					{
+						LinkCommand += string.Format(" \"{0}\"", ConvertPath(Path.GetFullPath(AdditionalLibrary)));
+						if (Path.GetExtension(AdditionalLibrary) == ".dylib")
+						{
+							ThirdPartyLibraries.Add(AdditionalLibrary);
+						}
+
+						if (BuildHostPlatform.Current.Platform != UnrealTargetPlatform.Mac)
+						{
+							// copy over libs we may need
+							FileItem ShadowFile = FileItem.GetExistingItemByPath(AdditionalLibrary);
+							if (ShadowFile != null)
+							{
+								QueueFileForBatchUpload(ShadowFile);
+							}
+						}
+					}
 
 					AddLibraryPathToRPaths(AdditionalLibrary, AbsolutePath, ref RPaths, ref LinkCommand, bIsBuildingAppBundle);
 				}
