Fix min ruleset warning in VS2015

diff --git a/Engine/Source/Programs/UnrealBuildTool/UnrealBuildTool.csproj b/Engine/Source/Programs/UnrealBuildTool/UnrealBuildTool.csproj
--- a/Engine/Source/Programs/UnrealBuildTool/UnrealBuildTool.csproj
+++ b/Engine/Source/Programs/UnrealBuildTool/UnrealBuildTool.csproj
@@ -43,8 +43,6 @@
     <DebugType>pdbonly</DebugType>
     <PlatformTarget>AnyCPU</PlatformTarget>
     <ErrorReport>prompt</ErrorReport>
-    <CodeAnalysisIgnoreBuiltInRuleSets>true</CodeAnalysisIgnoreBuiltInRuleSets>
-    <CodeAnalysisIgnoreBuiltInRules>true</CodeAnalysisIgnoreBuiltInRules>
     <AllowUnsafeBlocks>false</AllowUnsafeBlocks>
     <Prefer32Bit>false</Prefer32Bit>
   </PropertyGroup>
@@ -58,8 +56,6 @@
     <DebugType>pdbonly</DebugType>
     <PlatformTarget>x86</PlatformTarget>
     <ErrorReport>prompt</ErrorReport>
-    <CodeAnalysisIgnoreBuiltInRuleSets>true</CodeAnalysisIgnoreBuiltInRuleSets>
-    <CodeAnalysisIgnoreBuiltInRules>true</CodeAnalysisIgnoreBuiltInRules>
   </PropertyGroup>
   <PropertyGroup Condition=" '$(Configuration)|$(Platform)' == 'Debug|AnyCPU' ">
     <OutputPath>..\..\..\Binaries\DotNET\</OutputPath>
@@ -68,8 +64,6 @@
     <DebugType>pdbonly</DebugType>
     <PlatformTarget>AnyCPU</PlatformTarget>
     <ErrorReport>prompt</ErrorReport>
-    <CodeAnalysisIgnoreBuiltInRuleSets>true</CodeAnalysisIgnoreBuiltInRuleSets>
-    <CodeAnalysisIgnoreBuiltInRules>true</CodeAnalysisIgnoreBuiltInRules>
     <Prefer32Bit>false</Prefer32Bit>
   </PropertyGroup>
   <PropertyGroup Condition="'$(Configuration)|$(Platform)' == 'Debug|x86'">
