Merging using UE4-To-UE4-Releases-4.10
Generate proper proguard-project.txt file with optional additions from Build/Android/ProguardAdditions.txt
#ue4
#android

diff --git a/Engine/Source/Developer/Android/AndroidPlatformEditor/Private/AndroidTargetSettingsCustomization.cpp b/Engine/Source/Developer/Android/AndroidPlatformEditor/Private/AndroidTargetSettingsCustomization.cpp
--- a/Engine/Source/Developer/Android/AndroidPlatformEditor/Private/AndroidTargetSettingsCustomization.cpp
+++ b/Engine/Source/Developer/Android/AndroidPlatformEditor/Private/AndroidTargetSettingsCustomization.cpp
@@ -378,7 +378,7 @@ void FAndroidTargetSettingsCustomization::CopySetupFilesIntoProject()
 		}
 
 		// and copy the other files (aren't required)
-		SourceControlHelpers::CopyFileUnderSourceControl(GameProguardPath, EngineProguardPath, LOCTEXT("Proguard", "Proguard Settings"), /*out*/ ErrorMessage);
+		//SourceControlHelpers::CopyFileUnderSourceControl(GameProguardPath, EngineProguardPath, LOCTEXT("Proguard", "Proguard Settings"), /*out*/ ErrorMessage);
 	}
 
 	SavedLayoutBuilder->ForceRefreshDetails();
diff --git a/Engine/Source/Programs/UnrealBuildTool/Android/UEDeployAndroid.cs b/Engine/Source/Programs/UnrealBuildTool/Android/UEDeployAndroid.cs
--- a/Engine/Source/Programs/UnrealBuildTool/Android/UEDeployAndroid.cs
+++ b/Engine/Source/Programs/UnrealBuildTool/Android/UEDeployAndroid.cs
@@ -492,10 +492,10 @@ namespace UnrealBuildTool
 		{
 			switch (UE4Arch)
 			{
-				case "-armv7": return "armeabi-v7a";
-				case "-arm64": return "arm64-v8a";
-				case "-x64": return "x86_64";
-				case "-x86": return "x86";
+				case "-armv7":	return "armeabi-v7a";
+                case "-arm64":  return "arm64-v8a";
+				case "-x64":	return "x86_64";
+				case "-x86":	return "x86";
 
 				default: throw new BuildException("Unknown UE4 architecture {0}", UE4Arch);
 			}
@@ -506,15 +506,15 @@ namespace UnrealBuildTool
 			switch (NDKArch)
 			{
 				case "armeabi-v7a": return "-armv7";
-				case "arm64-v8a": return "-arm64";
-				case "x86": return "-x86";
-				case "arm64": return "-arm64";
+                case "arm64-v8a":   return "-arm64";
+                case "x86":         return "-x86";
+                case "arm64":       return "-arm64";
 				case "x86_64":
-				case "x64": return "-x64";
-
-				//				default: throw new BuildException("Unknown NDK architecture '{0}'", NDKArch);
-				// future-proof by returning armv7 for unknown
-				default: return "-armv7";
+				case "x64":			return "-x64";
+					
+//				default: throw new BuildException("Unknown NDK architecture '{0}'", NDKArch);
+                // future-proof by returning armv7 for unknown
+                default:            return "-armv7";
 			}
 		}
 
@@ -1231,6 +1231,9 @@ namespace UnrealBuildTool
 					Text.AppendLine("\t\t\t" + Line);
 				}
 			}
+
+			// add plugin Activity additions HERE
+	
 			Text.AppendLine("\t\t</activity>");
 
 			// For OBB download support
@@ -1279,9 +1282,11 @@ namespace UnrealBuildTool
 				}
 			}
 
-			// Required for OBB download support
-			Text.AppendLine("\t\t<service android:name=\"OBBDownloaderService\" />");
-			Text.AppendLine("\t\t<receiver android:name=\"AlarmReceiver\" />");
+			// add plugin Application additions HERE
+
+            // Required for OBB download support
+            Text.AppendLine("\t\t<service android:name=\"OBBDownloaderService\" />");
+            Text.AppendLine("\t\t<receiver android:name=\"AlarmReceiver\" />");
 
 			Text.AppendLine("\t</application>");
 
@@ -1338,7 +1343,40 @@ namespace UnrealBuildTool
 					}
 				}
 			}
+
+			// add plugin Requirements additions HERE
+
 			Text.AppendLine("</manifest>");
+
+			// allow plugins to modify final manifest HERE
+
+			return Text.ToString();
+		}
+
+		private string GenerateProguard(string EngineSourcePath, string GameBuildFilesPath)
+		{
+			StringBuilder Text = new StringBuilder();
+
+			string ProguardFile = Path.Combine(EngineSourcePath, "proguard-project.txt");
+			if (File.Exists(ProguardFile))
+			{
+				foreach (string Line in File.ReadAllLines(ProguardFile))
+				{
+					Text.AppendLine(Line);
+				}
+			}
+
+			string ProguardAdditionsFile = Path.Combine(GameBuildFilesPath, "ProguardAdditions.txt");
+			if (File.Exists(ProguardAdditionsFile))
+			{
+				foreach (string Line in File.ReadAllLines(ProguardAdditionsFile))
+				{
+					Text.AppendLine(Line);
+				}
+			}
+
+			// add plugin additions HERE
+
 			return Text.ToString();
 		}
 
@@ -1414,6 +1452,14 @@ namespace UnrealBuildTool
 				WriteJavaOBBDataFile(UE4OBBDataFileName, PackageName, new List<string> { ObbFileLocation });
 			}
 
+			// Make sure any existing proguard file in project is NOT used (back it up)
+			string ProjectBuildProguardFile = Path.Combine(GameBuildFilesPath, "proguard-project.txt");
+			if (File.Exists(ProjectBuildProguardFile))
+			{
+				string ProjectBackupProguardFile = Path.Combine(GameBuildFilesPath, "proguard-project.backup");
+				File.Move(ProjectBuildProguardFile, ProjectBackupProguardFile);
+			}
+
 			WriteJavaDownloadSupportFiles(UE4DownloadShimFileName, templates, new Dictionary<string, string>{
                 { "$$GameName$$", ProjectName },
                 { "$$PublicKey$$", GetPublicKey() }, 
@@ -1601,6 +1647,15 @@ namespace UnrealBuildTool
 			CopyFileDirectory(GameBuildFilesPath + "/NotForLicensees", UE4BuildPath, Replacements);
 			CopyFileDirectory(GameBuildFilesPath + "/NoRedist", UE4BuildPath, Replacements);
 
+			// Generate the Proguard file contents and write it
+			string ProguardContents = GenerateProguard(UE4BuildFilesPath, GameBuildFilesPath);
+			string ProguardFilename = UE4BuildPath + "/proguard-project.txt";
+			if (File.Exists(ProguardFilename))
+			{
+				File.Delete(ProguardFilename);
+			}
+			File.WriteAllText(ProguardFilename, ProguardContents);
+
 			//Now keep the splash screen images matching orientation requested
 			PickSplashScreenOrientation(UE4BuildPath);
 
