Merging using UE4-To-UE4-Releases-4.10
Replace escaped characters in Android application display name
#ue4
#android

diff --git a/Engine/Source/Programs/UnrealBuildTool/Android/UEDeployAndroid.cs b/Engine/Source/Programs/UnrealBuildTool/Android/UEDeployAndroid.cs
--- a/Engine/Source/Programs/UnrealBuildTool/Android/UEDeployAndroid.cs
+++ b/Engine/Source/Programs/UnrealBuildTool/Android/UEDeployAndroid.cs
@@ -633,8 +633,8 @@ namespace UnrealBuildTool
 				ApplicationDisplayName = ProjectName;
 			}
 
-			// make sure name does not have < or >
-			ApplicationDisplayName = ApplicationDisplayName.Replace("<", "(").Replace(">", ")");
+			// replace escaped characters (note: changes &# pattern before &, then patches back to allow escaped character codes in the string)
+			ApplicationDisplayName = ApplicationDisplayName.Replace("&#", "$@#$").Replace("&", "&amp;").Replace("'", "&#39;").Replace("\"", "&quot;").Replace("<", "&lt;").Replace(">", "&gt;").Replace("$@#$", "&#");
 
 			// if it doesn't exist, need to repackage
 			if (!File.Exists(StringsXMLPath))
