UBT: do not try to access registry unnecessarily.

- Can error out on Linux.

diff --git a/Engine/Source/Programs/UnrealBuildTool/Windows/UEBuildWindows.cs b/Engine/Source/Programs/UnrealBuildTool/Windows/UEBuildWindows.cs
--- a/Engine/Source/Programs/UnrealBuildTool/Windows/UEBuildWindows.cs
+++ b/Engine/Source/Programs/UnrealBuildTool/Windows/UEBuildWindows.cs
@@ -578,6 +578,11 @@ namespace UnrealBuildTool
 		/// <returns>Common tools path.</returns>
 		public static string GetVSComnToolsPath(WindowsCompiler Compiler)
 		{
+			if (BuildHostPlatform.Current.Platform != UnrealTargetPlatform.Win64 && BuildHostPlatform.Current.Platform != UnrealTargetPlatform.Win32)
+			{
+				return null;
+			}
+
 			int VSVersion;
 
 			switch (Compiler)
