Move .dSYM files out of the application bundle to prevent a crash on launch under 10.10 and earlier. This will cause regeneration of .dSYMs whenever you build, even if nothing changes, because in order for them to be built at all we have to tell UBT that they are alongside their .dylib counterparts. Successfully preflighted and should fix UE-20070.

diff --git a/Engine/Source/Programs/UnrealBuildTool/Mac/MacToolChain.cs b/Engine/Source/Programs/UnrealBuildTool/Mac/MacToolChain.cs
--- a/Engine/Source/Programs/UnrealBuildTool/Mac/MacToolChain.cs
+++ b/Engine/Source/Programs/UnrealBuildTool/Mac/MacToolChain.cs
@@ -112,6 +112,7 @@ namespace UnrealBuildTool
 			Result += " -fasm-blocks";
 
 			Result += " -Wall -Werror";
+			//Result += " -Wsign-compare"; // fed up of not seeing the signed/unsigned warnings we get on Windows - lets enable them here too.
 
 			Result += " -Wno-unused-variable";
 			Result += " -Wno-unused-value";
@@ -1153,6 +1154,16 @@ namespace UnrealBuildTool
 		/// <param name="MachOBinary">FileItem describing the executable or dylib to generate debug info for</param>
 		public FileItem GenerateDebugInfo(FileItem MachOBinary)
 		{
+			string BinaryPath = MachOBinary.AbsolutePath;
+			if(BinaryPath.Contains(".app"))
+			{
+				while(BinaryPath.Contains(".app"))
+				{
+					BinaryPath = Path.GetDirectoryName(BinaryPath);
+				}
+				BinaryPath = Path.Combine(BinaryPath, Path.GetFileName(Path.ChangeExtension(MachOBinary.AbsolutePath, ".dSYM")));
+			}
+
 			// Make a file item for the source and destination files
 			string FullDestPath = Path.ChangeExtension(MachOBinary.AbsolutePath, ".dSYM");
 
@@ -1178,10 +1189,23 @@ namespace UnrealBuildTool
 
 			// Deletes ay existing file on the building machine,
 			// note that the source and dest are switched from a copy command
-			GenDebugAction.CommandArguments = string.Format("-c 'rm -rf \"{2}\"; \"{0}\"dsymutil -f \"{1}\" -o \"{2}\"'",
-				ToolchainDir,
-				InputFile.AbsolutePath,
-				DestFile.AbsolutePath);
+			if(BinaryPath == MachOBinary.AbsolutePath)
+			{
+				GenDebugAction.CommandArguments = string.Format("-c 'rm -rf \"{2}\"; \"{0}\"dsymutil -f \"{1}\" -o \"{2}\"'",
+					ToolchainDir,
+					InputFile.AbsolutePath,
+					DestFile.AbsolutePath);
+			}
+			else
+			{
+				FileItem MovedFile = FileItem.GetItemByPath(BinaryPath);
+				FileItem FinalFile = LocalToRemoteFileItem(MovedFile, false);
+				GenDebugAction.CommandArguments = string.Format("-c 'rm -rf \"{2}\"; \"{0}\"dsymutil -f \"{1}\" -o \"{2}\"; mv \"{2}\" \"{3}\"'",
+					ToolchainDir,
+					InputFile.AbsolutePath,
+					DestFile.AbsolutePath,
+					FinalFile.AbsolutePath);
+			}
 			GenDebugAction.PrerequisiteItems.Add(InputFile);
 			GenDebugAction.ProducedItems.Add(DestFile);
 			GenDebugAction.StatusDescription = GenDebugAction.CommandArguments;
@@ -1315,26 +1339,40 @@ namespace UnrealBuildTool
 
 		public override void ModifyBuildProducts(UEBuildBinary Binary, Dictionary<FileReference, BuildProductType> BuildProducts)
 		{
-			// The cross-platform code adds .dSYMs for static libraries, which is just wrong, so
-			// eliminate them here for now.
-			string DebugExtension = UEBuildPlatform.GetBuildPlatform(Binary.Target.Platform).GetDebugInfoExtension(Binary.Config.Type);
-			if (DebugExtension == ".dSYM")
+			if (BuildConfiguration.bGeneratedSYMFile == true || BuildConfiguration.bUsePDBFiles == true)
 			{
 				KeyValuePair<FileReference, BuildProductType>[] BuildProductsArray = BuildProducts.ToArray();
 
 				foreach (KeyValuePair<FileReference, BuildProductType> BuildProductPair in BuildProductsArray)
 				{
-					if (BuildProductPair.Key.HasExtension(DebugExtension))
+					string DebugExtension = "";
+					switch (BuildProductPair.Value)
 					{
-						BuildProducts.Remove(BuildProductPair.Key);
+						case BuildProductType.Executable:
+							DebugExtension = UEBuildPlatform.GetBuildPlatform(Binary.Target.Platform).GetDebugInfoExtension(UEBuildBinaryType.Executable);
+							break;
+						case BuildProductType.DynamicLibrary:
+							DebugExtension = UEBuildPlatform.GetBuildPlatform(Binary.Target.Platform).GetDebugInfoExtension(UEBuildBinaryType.DynamicLinkLibrary);
+							break;
 					}
-				}
-
-				foreach (KeyValuePair<FileReference, BuildProductType> BuildProductPair in BuildProductsArray)
-				{
-					if (BuildProductPair.Value == BuildProductType.Executable || BuildProductPair.Value == BuildProductType.DynamicLibrary)
+					if (DebugExtension == ".dSYM")
+					{
+						string BinaryPath = BuildProductPair.Key.FullName;
+						if(BinaryPath.Contains(".app"))
+						{
+							while(BinaryPath.Contains(".app"))
+							{
+								BinaryPath = Path.GetDirectoryName(BinaryPath);
+							}
+							BinaryPath = Path.Combine(BinaryPath, BuildProductPair.Key.GetFileName());
+							BinaryPath = Path.ChangeExtension(BinaryPath, DebugExtension);
+							FileReference Ref = new FileReference(BinaryPath);
+							BuildProducts.Add(Ref, BuildProductType.SymbolFile);
+						}
+					}
+					else if(BuildProductPair.Value == BuildProductType.SymbolFile && BuildProductPair.Key.FullName.Contains(".app"))
 					{
-						BuildProducts.Add(BuildProductPair.Key.ChangeExtension(DebugExtension), BuildProductType.SymbolFile);
+						BuildProducts.Remove(BuildProductPair.Key);
 					}
 				}
 			}
@@ -1456,6 +1494,20 @@ namespace UnrealBuildTool
 				foreach (UEBuildBinary Binary in InTarget.AppBinaries)
 				{
 					BuiltBinaries.Add(Path.GetFullPath(Binary.ToString()));
+
+					string DebugExtension = UEBuildPlatform.GetBuildPlatform(Binary.Target.Platform).GetDebugInfoExtension(Binary.Config.Type);
+					if (DebugExtension == ".dSYM")
+					{
+						Dictionary<FileReference, BuildProductType> BuildProducts = new Dictionary<FileReference, BuildProductType>();
+						Binary.GetBuildProducts(this, BuildProducts);
+						foreach (KeyValuePair<FileReference, BuildProductType> BuildProductPair in BuildProducts)
+						{
+							if (BuildProductPair.Value == BuildProductType.SymbolFile)
+							{
+								BuiltBinaries.Add(BuildProductPair.Key.FullName);
+							}
+						}
+					}
 				}
 
 				string IntermediateDirectory = InTarget.EngineIntermediateDirectory.FullName;
diff --git a/Engine/Source/Programs/UnrealBuildTool/Mac/UEBuildMac.cs b/Engine/Source/Programs/UnrealBuildTool/Mac/UEBuildMac.cs
--- a/Engine/Source/Programs/UnrealBuildTool/Mac/UEBuildMac.cs
+++ b/Engine/Source/Programs/UnrealBuildTool/Mac/UEBuildMac.cs
@@ -201,7 +201,21 @@ namespace UnrealBuildTool
 		/// <returns>string    The debug info extension (i.e. 'pdb')</returns>
 		public override string GetDebugInfoExtension(UEBuildBinaryType InBinaryType)
 		{
-			return BuildConfiguration.bGeneratedSYMFile || BuildConfiguration.bUsePDBFiles ? ".dSYM" : "";
+			switch (InBinaryType)
+			{
+				case UEBuildBinaryType.DynamicLinkLibrary:
+					return BuildConfiguration.bGeneratedSYMFile || BuildConfiguration.bUsePDBFiles ? ".dSYM" : "";
+				case UEBuildBinaryType.Executable:
+					return BuildConfiguration.bGeneratedSYMFile || BuildConfiguration.bUsePDBFiles ? ".dSYM" : "";
+				case UEBuildBinaryType.StaticLibrary:
+					return "";
+				case UEBuildBinaryType.Object:
+					return "";
+				case UEBuildBinaryType.PrecompiledHeader:
+					return "";
+				default:
+					return "";
+			}
 		}
 
 		/// <summary>
