Merging PR-1568 from 4.10

diff --git a/Engine/Source/Programs/UnrealBuildTool/IOS/IOSToolChain.cs b/Engine/Source/Programs/UnrealBuildTool/IOS/IOSToolChain.cs
--- a/Engine/Source/Programs/UnrealBuildTool/IOS/IOSToolChain.cs
+++ b/Engine/Source/Programs/UnrealBuildTool/IOS/IOSToolChain.cs
@@ -394,6 +394,7 @@ namespace UnrealBuildTool
 			Result += " -dead_strip";
 			Result += " -miphoneos-version-min=" + PlatformContext.GetRunTimeVersion();
 			Result += " -Wl,-no_pie";
+			Result += " -stdlib=libc++";
 			//			Result += " -v";
 
 			// link in the frameworks
