Merging PR-1578 from 4.10

diff --git a/Engine/Source/Programs/IOS/iPhonePackager/SSHCommandHelper.cs b/Engine/Source/Programs/IOS/iPhonePackager/SSHCommandHelper.cs
--- a/Engine/Source/Programs/IOS/iPhonePackager/SSHCommandHelper.cs
+++ b/Engine/Source/Programs/IOS/iPhonePackager/SSHCommandHelper.cs
@@ -260,7 +260,7 @@ namespace iPhonePackager
 			// make simple rsync commandline to send a file
 			RsyncProcess.StartInfo.FileName = RSyncExePath;
 			RsyncProcess.StartInfo.Arguments = string.Format(
-				"-zae \"ssh {0}\" --rsync-path=\"mkdir -p {1} && rsync\" '{2}' {3}@{4}:'{1}/{5}'",
+				"-zae \"ssh {0}\" --rsync-path=\"mkdir -p {1} && rsync\" --chmod=ug=rwX,o=rxX '{2}' {3}@{4}:'{1}/{5}'",
 				SSHAuthentication,
 				RemoteDir,
 				ConvertPathToCygwin(LocalPath),
diff --git a/Engine/Source/Programs/UnrealBuildTool/ToolChain/RemoteToolChain.cs b/Engine/Source/Programs/UnrealBuildTool/ToolChain/RemoteToolChain.cs
--- a/Engine/Source/Programs/UnrealBuildTool/ToolChain/RemoteToolChain.cs
+++ b/Engine/Source/Programs/UnrealBuildTool/ToolChain/RemoteToolChain.cs
@@ -825,7 +825,7 @@ namespace UnrealBuildTool
 			// make simple rsync commandline to send a file
 			RsyncProcess.StartInfo.FileName = ResolvedRSyncExe;
 			RsyncProcess.StartInfo.Arguments = string.Format(
-				"-zae \"{0}\" --rsync-path=\"mkdir -p {1} && rsync\" '{2}' {3}@{4}:'{1}/{5}'",
+				"-zae \"{0}\" --rsync-path=\"mkdir -p {1} && rsync\" --chmod=ug=rwX,o=rxX '{2}' {3}@{4}:'{1}/{5}'",
 				ResolvedRsyncAuthentication,
 				RemoteDir,
 				ConvertPathToCygwin(LocalPath),
