Merging fix for failure to package when having an & in the bundle display name from 4.10

diff --git a/Engine/Source/Programs/UnrealBuildTool/IOS/UEDeployIOS.cs b/Engine/Source/Programs/UnrealBuildTool/IOS/UEDeployIOS.cs
--- a/Engine/Source/Programs/UnrealBuildTool/IOS/UEDeployIOS.cs
+++ b/Engine/Source/Programs/UnrealBuildTool/IOS/UEDeployIOS.cs
@@ -247,6 +247,14 @@ namespace UnrealBuildTool
 			string ExtraData = "";
 			Ini.GetString("/Script/IOSRuntimeSettings.IOSRuntimeSettings", "AdditionalPlistData", out ExtraData);
 
+			// create the final display name, including converting all entities for XML use
+			string FinalDisplayName = BundleDisplayName.Replace("[PROJECT_NAME]", ProjectName).Replace("_", "");
+			FinalDisplayName = FinalDisplayName.Replace("&", "&amp;");
+			FinalDisplayName = FinalDisplayName.Replace("\"", "&quot;");
+			FinalDisplayName = FinalDisplayName.Replace("\'", "&apos;");
+			FinalDisplayName = FinalDisplayName.Replace("<", "&lt;");
+			FinalDisplayName = FinalDisplayName.Replace(">", "&gt;");
+
 			// generate the plist file
 			StringBuilder Text = new StringBuilder();
 			Text.AppendLine("<?xml version=\"1.0\" encoding=\"UTF-8\"?>");
@@ -272,7 +280,7 @@ namespace UnrealBuildTool
 			Text.AppendLine("\t<key>CFBundleDevelopmentRegion</key>");
 			Text.AppendLine("\t<string>English</string>");
 			Text.AppendLine("\t<key>CFBundleDisplayName</key>");
-			Text.AppendLine(string.Format("\t<string>{0}</string>", BundleDisplayName.Replace("[PROJECT_NAME]", ProjectName).Replace("_", "")));
+			Text.AppendLine(string.Format("\t<string>{0}</string>",FinalDisplayName));
 			Text.AppendLine("\t<key>CFBundleExecutable</key>");
 			Text.AppendLine(string.Format("\t<string>{0}</string>", bIsUE4Game ? "UE4Game" : GameName));
 			Text.AppendLine("\t<key>CFBundleIdentifier</key>");
