Merging using UE4-To-UE4-Releases-4.10
Default launch image theme fallback is black to prevent stretch/squash of the image if started in wrong orientation
#ue4
#android

diff --git a/Engine/Build/Android/Java/res/values/styles.xml b/Engine/Build/Android/Java/res/values/styles.xml
--- a/Engine/Build/Android/Java/res/values/styles.xml
+++ b/Engine/Build/Android/Java/res/values/styles.xml
@@ -1,6 +1,10 @@
 <?xml version="1.0" encoding="utf-8"?>
 <resources>
 
+    <style name="UE4BaseTheme" parent = "@android:style/Theme.Black.NoTitleBar.Fullscreen" />
+
+    <style name="UE4SplashTheme" parent = "UE4BaseTheme" />
+
     <style name="NotificationText">
         <item name="android:textColor">?android:attr/textColorPrimary</item>
     </style>
diff --git a/Engine/Source/Programs/UnrealBuildTool/Android/UEDeployAndroid.cs b/Engine/Source/Programs/UnrealBuildTool/Android/UEDeployAndroid.cs
--- a/Engine/Source/Programs/UnrealBuildTool/Android/UEDeployAndroid.cs
+++ b/Engine/Source/Programs/UnrealBuildTool/Android/UEDeployAndroid.cs
@@ -1240,6 +1240,8 @@ namespace UnrealBuildTool
 			if (bShowLaunchImage)
 			{
 				Text.AppendLine("\t\t<activity android:name=\".DownloaderActivity\"");
+				Text.AppendLine(string.Format("\t\t          android:screenOrientation=\"{0}\"", ConvertOrientationIniValue(Orientation)));
+				Text.AppendLine("\t\t          android:configChanges=\"screenSize|orientation|keyboardHidden|keyboard\"");
 				Text.AppendLine("\t\t          android:theme=\"@style/UE4SplashTheme\" />");
 			}
 			else
