Merging CL#2726541 using //depot/UE4-To-//UE4/Main

Adding option to make module output publicly distributable so that OnlineSubsystemMCP dll can be distributed for launcher.

diff --git a/Engine/Source/Programs/UnrealBuildTool/Configuration/UEBuildBinary.cs b/Engine/Source/Programs/UnrealBuildTool/Configuration/UEBuildBinary.cs
--- a/Engine/Source/Programs/UnrealBuildTool/Configuration/UEBuildBinary.cs
+++ b/Engine/Source/Programs/UnrealBuildTool/Configuration/UEBuildBinary.cs
@@ -396,7 +396,7 @@ namespace UnrealBuildTool
 					if (OutputDistributionLevel < DistributionLevel)
 					{
 						var JoinedModuleNames = String.Join(",", DependantModuleNames[(int)DistributionLevel]);
-						throw new BuildException("Output file \"{0}\" has distribution level of \"{1}\" but has direct dependencies on modules with distribution level of \"{2}\" ({3}).\nEither change to dynamic dependencies, set BinariesSubFolder/ExeBinariesSubFolder to \"{2}\" or set bOutputPubliclyDistributable to true in the target.cs file.",
+						throw new BuildException("Output file \"{0}\" has distribution level of \"{1}\" but has direct dependencies on modules with distribution level of \"{2}\" ({3}).\nEither change to dynamic dependencies, set BinariesSubFolder/ExeBinariesSubFolder to \"{2}\" or set bOutputPubliclyDistributable to true in the target.cs or build.cs file.",
 							OutputFilePath, OutputDistributionLevel.ToString(), DistributionLevel.ToString(), JoinedModuleNames);
 					}
 				}
diff --git a/Engine/Source/Programs/UnrealBuildTool/Configuration/UEBuildModule.cs b/Engine/Source/Programs/UnrealBuildTool/Configuration/UEBuildModule.cs
--- a/Engine/Source/Programs/UnrealBuildTool/Configuration/UEBuildModule.cs
+++ b/Engine/Source/Programs/UnrealBuildTool/Configuration/UEBuildModule.cs
@@ -154,10 +154,13 @@ namespace UnrealBuildTool
 
 			UEBuildModuleDistribution DistributionLevel = UEBuildModuleDistribution.Public;
 
-			// Keep checking as long as we haven't reached the maximum level
-			for (int PathIndex = 0; PathIndex < PathsToCheck.Count && DistributionLevel != UEBuildModuleDistribution.NoRedist; ++PathIndex)
+			if (!Rules.bOutputPubliclyDistributable)
 			{
-				DistributionLevel = Utils.Max(DistributionLevel, UEBuildModule.GetModuleDistributionLevelBasedOnLocation(PathsToCheck[PathIndex]));
+				// Keep checking as long as we haven't reached the maximum level
+				for (int PathIndex = 0; PathIndex < PathsToCheck.Count && DistributionLevel != UEBuildModuleDistribution.NoRedist; ++PathIndex)
+				{
+					DistributionLevel = Utils.Max(DistributionLevel, UEBuildModule.GetModuleDistributionLevelBasedOnLocation(PathsToCheck[PathIndex]));
+				}
 			}
 
 			return DistributionLevel;
diff --git a/Engine/Source/Programs/UnrealBuildTool/System/RulesCompiler.cs b/Engine/Source/Programs/UnrealBuildTool/System/RulesCompiler.cs
--- a/Engine/Source/Programs/UnrealBuildTool/System/RulesCompiler.cs
+++ b/Engine/Source/Programs/UnrealBuildTool/System/RulesCompiler.cs
@@ -237,6 +237,13 @@ namespace UnrealBuildTool
 		/// </summary>
 		public bool bBuildLocallyWithSNDBS = false;
 
+		/// <summary>
+		/// Whether the output from this module can be publicly distributed, even if it has code/
+		/// dependencies on modules that are not (i.e. CarefullyRedist, NotForLicensees, NoRedist).
+		/// This should be used when you plan to release binaries but not source.
+		/// </summary>
+		public bool bOutputPubliclyDistributable = false;
+
 		/// <summary>
         /// List of modules names (no path needed) with header files that our module's public headers needs access to, but we don't need to "import" or link against.
 		/// </summary>
