fix for build failure when staging PS4

diff --git a/Engine/Source/Programs/UnrealBuildTool/Configuration/EngineConfiguration.cs b/Engine/Source/Programs/UnrealBuildTool/Configuration/EngineConfiguration.cs
--- a/Engine/Source/Programs/UnrealBuildTool/Configuration/EngineConfiguration.cs
+++ b/Engine/Source/Programs/UnrealBuildTool/Configuration/EngineConfiguration.cs
@@ -53,7 +53,7 @@ namespace UnrealBuildTool
 		static List<string> RequiredSections = new List<string>(){"AppxManifest", "CommonSettings", "/Script/AndroidRuntimeSettings.AndroidRuntimeSettings", "/Script/AndroidPlatformEditor.AndroidSDKSettings",
 																	"/Script/BuildSettings.BuildSettings", "/Script/IOSRuntimeSettings.IOSRuntimeSettings", "/Script/WindowsTargetPlatform.WindowsTargetSettings",
 																	"/Script/UnrealEd.ProjectPackagingSettings", "/Script/PS4PlatformEditor.PS4TargetSettings", "/Script/XboxOneTargetPlatform.XboxOneTargetSettings",
-																	"/Script/HTML5PlatformEditor.HTML5TargetSettings"};
+																	"/Script/HTML5PlatformEditor.HTML5TargetSettings","PS4SymbolServer"};
 
 		// static creation functions for ini files
 		public static ConfigCacheIni CreateConfigCacheIni(UnrealTargetPlatform Platform, string BaseIniName, DirectoryReference ProjectDirectory, DirectoryReference EngineDirectory = null)
