Copying

//UE4/Dev-Platform/Engine/Source/Programs/UnrealBuildTool/Configuration/EngineConfiguration.cs

to //UE4/Main/Engine/Source/Programs/UnrealBuildTool/Configuration/EngineConfiguration.cs

diff --git a/Engine/Source/Programs/UnrealBuildTool/Configuration/EngineConfiguration.cs b/Engine/Source/Programs/UnrealBuildTool/Configuration/EngineConfiguration.cs
--- a/Engine/Source/Programs/UnrealBuildTool/Configuration/EngineConfiguration.cs
+++ b/Engine/Source/Programs/UnrealBuildTool/Configuration/EngineConfiguration.cs
@@ -84,6 +84,13 @@ namespace UnrealBuildTool
 		/// </summary>
 		public class IniValues : List<string>
 		{
+			public IniValues()
+			{
+			}
+			public IniValues(IniValues Other)
+				: base(Other)
+			{
+			}
 			public override string ToString()
 			{
 				return String.Join(",", ToArray());
@@ -98,6 +105,14 @@ namespace UnrealBuildTool
 			public IniSection()
 				: base(StringComparer.InvariantCultureIgnoreCase)
 			{ }
+			public IniSection(IniSection Other)
+				: this()
+			{
+				foreach (var Pair in Other)
+				{
+					Add(Pair.Key, new IniValues(Pair.Value));
+				}
+			}
 			public override string ToString()
 			{
 				return "IniSection";
@@ -192,7 +207,7 @@ namespace UnrealBuildTool
 			{
 				foreach (var Pair in BaseCache.Sections)
 				{
-					Sections.Add(Pair.Key, Pair.Value);
+					Sections.Add(Pair.Key, new IniSection(Pair.Value));
 				}
 			}
 			if (EngineOnly)
