String literal compiler extensions enabled for Android, HTML5, iOS and Mac.

diff --git a/Engine/Source/Programs/UnrealBuildTool/Android/AndroidToolChain.cs b/Engine/Source/Programs/UnrealBuildTool/Android/AndroidToolChain.cs
--- a/Engine/Source/Programs/UnrealBuildTool/Android/AndroidToolChain.cs
+++ b/Engine/Source/Programs/UnrealBuildTool/Android/AndroidToolChain.cs
@@ -346,6 +346,8 @@ namespace UnrealBuildTool
 			Result += " -Wno-unknown-pragmas";			// probably should kill this one, sign of another issue in PhysX?
 			Result += " -Wno-invalid-offsetof";			// needed to suppress warnings about using offsetof on non-POD types.
 			Result += " -Wno-logical-op-parentheses";	// needed for external headers we can't change
+			// we use this feature to allow static FNames.
+			Result += " -Wno-gnu-string-literal-operator-template";
 
 			if (CompileEnvironment.Config.bEnableShadowVariableWarning)
 			{
diff --git a/Engine/Source/Programs/UnrealBuildTool/HTML5/HTML5ToolChain.cs b/Engine/Source/Programs/UnrealBuildTool/HTML5/HTML5ToolChain.cs
--- a/Engine/Source/Programs/UnrealBuildTool/HTML5/HTML5ToolChain.cs
+++ b/Engine/Source/Programs/UnrealBuildTool/HTML5/HTML5ToolChain.cs
@@ -63,7 +63,10 @@ namespace UnrealBuildTool
 			// okay, in UE4, we'd fix the code for these, but in UE3, not worth it
 			Result += " -Wno-logical-op-parentheses"; // appErrorf triggers this
 			Result += " -Wno-array-bounds"; // some VectorLoads go past the end of the array, but it's okay in that case
-			Result += " -Wno-invalid-offsetof"; // too many warnings kills windows clang. 
+			// needed to suppress warnings about using offsetof on non-POD types.
+			Result += " -Wno-invalid-offsetof";
+			// we use this feature to allow static FNames.
+			Result += " -Wno-gnu-string-literal-operator-template";
 
 			if (bEnableShadowVariableWarning)
 			{
diff --git a/Engine/Source/Programs/UnrealBuildTool/IOS/IOSToolChain.cs b/Engine/Source/Programs/UnrealBuildTool/IOS/IOSToolChain.cs
--- a/Engine/Source/Programs/UnrealBuildTool/IOS/IOSToolChain.cs
+++ b/Engine/Source/Programs/UnrealBuildTool/IOS/IOSToolChain.cs
@@ -201,7 +201,10 @@ namespace UnrealBuildTool
 			Result += " -Wno-tautological-compare";
 			//This will prevent the issue of warnings for unused private variables.
 			Result += " -Wno-unused-private-field";
-			Result += " -Wno-invalid-offsetof"; // needed to suppress warnings about using offsetof on non-POD types.
+			// needed to suppress warnings about using offsetof on non-POD types.
+			Result += " -Wno-invalid-offsetof";
+			// we use this feature to allow static FNames.
+			Result += " -Wno-gnu-string-literal-operator-template";
 
 			if (IOSSDKVersionFloat >= 9.0)
 			{
diff --git a/Engine/Source/Programs/UnrealBuildTool/Mac/MacToolChain.cs b/Engine/Source/Programs/UnrealBuildTool/Mac/MacToolChain.cs
--- a/Engine/Source/Programs/UnrealBuildTool/Mac/MacToolChain.cs
+++ b/Engine/Source/Programs/UnrealBuildTool/Mac/MacToolChain.cs
@@ -124,7 +124,10 @@ namespace UnrealBuildTool
 			Result += " -Wno-tautological-compare";
 			// This will prevent the issue of warnings for unused private variables.
 			Result += " -Wno-unused-private-field";
-			Result += " -Wno-invalid-offsetof"; // needed to suppress warnings about using offsetof on non-POD types.
+			// needed to suppress warnings about using offsetof on non-POD types.
+			Result += " -Wno-invalid-offsetof";
+			// we use this feature to allow static FNames.
+			Result += " -Wno-gnu-string-literal-operator-template";
 
 			if (MacOSSDKVersionFloat < 10.9f && MacOSSDKVersionFloat >= 10.11f)
 			{
