Remove compatibility shims from trying to directly import games, and fix up UAT compile warnings. Also remove PromotionArchive stuff, now that the functionality is replaced by UGS.

diff --git a/Engine/Source/Programs/AutomationTool/GUBP/LegacyBranchSetup.cs b/Engine/Source/Programs/AutomationTool/GUBP/LegacyBranchSetup.cs
--- a/Engine/Source/Programs/AutomationTool/GUBP/LegacyBranchSetup.cs
+++ b/Engine/Source/Programs/AutomationTool/GUBP/LegacyBranchSetup.cs
@@ -20,13 +20,6 @@ partial class GUBP
 		public bool bForceIncrementalCompile;
         public JobInfo JobInfo;
 
-		// Included for compatibility with //UE4/Main import
-		[Obsolete]
-		public bool bPreflightBuild
-		{
-			get { return JobInfo.IsPreflight; }
-		}
-
 		public Dictionary<string, GUBPNode> GUBPNodes = new Dictionary<string, GUBPNode>();
 		public Dictionary<string, GUBPAggregateNode> GUBPAggregates = new Dictionary<string, GUBPAggregateNode>();
 
diff --git a/Engine/Source/Programs/AutomationTool/GUBP/LegacyNodes.cs b/Engine/Source/Programs/AutomationTool/GUBP/LegacyNodes.cs
--- a/Engine/Source/Programs/AutomationTool/GUBP/LegacyNodes.cs
+++ b/Engine/Source/Programs/AutomationTool/GUBP/LegacyNodes.cs
@@ -20,20 +20,6 @@ partial class GUBP
         public List<string> AllDependencyBuildProducts = null;
         public string AgentSharingGroup = "";
 
-		// Included for compatibility during //UE4/Main import
-		[Obsolete]
-		public virtual string GameNameIfAnyForTempStorage()
-		{
-			return "";
-		}
-
-		// Included for compatibility during //UE4/Main import
-		[Obsolete]
-		public virtual string RootIfAnyForTempStorage()
-		{
-			return "";
-		}
-
         public virtual string GetFullName()
         {
             throw new AutomationException("Unimplemented GetFullName.");
diff --git a/Engine/Source/Programs/UnrealBuildTool/System/TargetReceipt.cs b/Engine/Source/Programs/UnrealBuildTool/System/TargetReceipt.cs
--- a/Engine/Source/Programs/UnrealBuildTool/System/TargetReceipt.cs
+++ b/Engine/Source/Programs/UnrealBuildTool/System/TargetReceipt.cs
@@ -455,25 +455,4 @@ namespace UnrealBuildTool
 			}
 		}
 	}
-
-	/// <summary>
-	/// Legacy class for compatibility in //UE4/Main import.
-	/// </summary>
-	[Obsolete]
-	public class BuildReceipt : TargetReceipt
-	{
-		public BuildReceipt(TargetReceipt Other) : base(Other)
-		{
-		}
-
-		public static new BuildReceipt Read(string FileName)
-		{
-			return new BuildReceipt(TargetReceipt.Read(FileName));
-		}
-
-		public void ExpandPathVariables(string EngineDir, string ProjectDir)
-		{
-			base.ExpandPathVariables(new DirectoryReference(EngineDir), new DirectoryReference(ProjectDir));
-		}
-	}
 }
