[INTEGRATE] Change 2733699 by [REDACTED_PERSON][REDACTED_HANDLE] on 2015/10/19 12:07:44

	Add a build configuration setting to disable the editor dependency on ShaderCompileWorker when generating project files.

diff --git a/Engine/Source/Programs/UnrealBuildTool/Configuration/UEBuildConfiguration.cs b/Engine/Source/Programs/UnrealBuildTool/Configuration/UEBuildConfiguration.cs
--- a/Engine/Source/Programs/UnrealBuildTool/Configuration/UEBuildConfiguration.cs
+++ b/Engine/Source/Programs/UnrealBuildTool/Configuration/UEBuildConfiguration.cs
@@ -289,6 +289,12 @@ namespace UnrealBuildTool
 		/// </summary>
 		public static string PreferredSubPlatform = "";
 
+		/// <summary>
+		/// Whether to include a dependency on ShaderCompileWorker when generating project files for the editor.
+		/// </summary>
+		[XmlConfig]
+		public static bool bEditorDependsOnShaderCompileWorker;
+
 		/// <summary>
 		/// Sets the configuration back to defaults.
 		/// </summary>
@@ -330,6 +336,7 @@ namespace UnrealBuildTool
 			bSkipLinkingWhenNothingToCompile = false;
 			bCompileCEF3 = true;
 			PreferredSubPlatform = "";
+			bEditorDependsOnShaderCompileWorker = true;
 		}
 
 		/// <summary>
diff --git a/Engine/Source/Programs/UnrealBuildTool/System/VCProjectFileGenerator.cs b/Engine/Source/Programs/UnrealBuildTool/System/VCProjectFileGenerator.cs
--- a/Engine/Source/Programs/UnrealBuildTool/System/VCProjectFileGenerator.cs
+++ b/Engine/Source/Programs/UnrealBuildTool/System/VCProjectFileGenerator.cs
@@ -474,7 +474,7 @@ namespace UnrealBuildTool
 							Dependencies.Add(UBTProject);
 							Dependencies.AddRange(UBTProject.DependsOnProjects);
 						}
-						if (CurProject.IsGeneratedProject && ShaderCompileWorkerProject != null && CurProject.ProjectTargets.Any(x => x.TargetRules != null && x.TargetRules.Type == TargetRules.TargetType.Editor))
+						if (UEBuildConfiguration.bEditorDependsOnShaderCompileWorker && CurProject.IsGeneratedProject && ShaderCompileWorkerProject != null && CurProject.ProjectTargets.Any(x => x.TargetRules != null && x.TargetRules.Type == TargetRules.TargetType.Editor))
 						{
 							Dependencies.Add(ShaderCompileWorkerProject);
 						}
