UBT fast iteration: No longer invalidate includes for non-game/engine targets (UE-19993)
- When switching between targets that never share build products, there's no need to rescan all C++ headers
- This allows SCW, UBT and other programs to re-use cached includes when switching targets
- Game projects and the engine itself still invalidate cached includes.
- Unfortunately when UBT decides to invalidate includes, we don't yet know whether we are building modular/monolithic, so conservatively we always invalidate includes when switching between games or the engine in the same workspace

- NOTE:  I'm leaving the JIRA for this open, as we don't want to actually be invoking UBT twice at all for optimal iteration times.  Ben has work in progress code to allow SCW and game targets to be compiled in the same invocation, so this can remain open until then.

diff --git a/Engine/Source/Programs/UnrealBuildTool/System/UnrealBuildTool.cs b/Engine/Source/Programs/UnrealBuildTool/System/UnrealBuildTool.cs
--- a/Engine/Source/Programs/UnrealBuildTool/System/UnrealBuildTool.cs
+++ b/Engine/Source/Programs/UnrealBuildTool/System/UnrealBuildTool.cs
@@ -1623,27 +1623,16 @@ namespace UnrealBuildTool
 				{
 					if (BuildConfiguration.bUseUBTMakefiles)
 					{
-						// If we're building UHT without a Mutex, we'll need to assume that we're building the same targets and that no caches
-						// should be invalidated for this run.  This is important when UBT is invoked from within UBT in order to compile
-						// UHT.  In that (very common) case we definitely don't want to have to rebuild our cache from scratch.
-						bool bMustAssumeSameTargets = false;
-
-						if (TargetDescs[0].TargetName.Equals("UnrealHeaderTool", StringComparison.InvariantCultureIgnoreCase))
+						// Only the modular editor and game targets will share build products.  Unfortunately, we can't determine at
+						// at this point whether we're dealing with modular or monolithic game binaries, so we opt to always invalidate
+						// cached includes if the target we're switching to is either a game target (has project file) or "UE4Editor".
+						bool bMightHaveSharedBuildProducts = 
+							ProjectFile != null ||	// Is this a game? (has a .uproject file for the target)
+							TargetDescs[ 0 ].TargetName.Equals( "UE4Editor", StringComparison.InvariantCultureIgnoreCase );	// Is the engine?
+						if( bMightHaveSharedBuildProducts )
 						{
-							int NoMutexArgumentIndex;
-							if (Utils.ParseCommandLineFlag(Arguments, "-NoMutex", out NoMutexArgumentIndex))
-							{
-								bMustAssumeSameTargets = true;
-							}
-						}
+							bool bIsBuildingSameTargetsAsLastTime = false;
 
-						bool bIsBuildingSameTargetsAsLastTime = false;
-						if (bMustAssumeSameTargets)
-						{
-							bIsBuildingSameTargetsAsLastTime = true;
-						}
-						else
-						{
 							string TargetCollectionName = MakeTargetCollectionName(TargetDescs);
 
 							string LastBuiltTargetsFileName = bIsHotReload ? "HotReloadLastBuiltTargets.txt" : "LastBuiltTargets.txt";
@@ -1664,7 +1653,6 @@ namespace UnrealBuildTool
 								File.WriteAllText(LastBuiltTargetsFilePath, TargetCollectionName, Encoding.UTF8);
 							}
 
-
 							if (!bIsBuildingSameTargetsAsLastTime)
 							{
 								// Can't use super fast include checking unless we're building the same set of targets as last time, because
