Fix 2 second delay at the end of any UBT build with XGE, caused by trying to determine number of physical processors for telemetry. Regression from Dev-Build merge.

This slowdown still affects non-XGE builds, but is not a regression - we use the number of physical processors to derive the number of compiler instances that should be run in parallel, and have done for a long time. Bug for this case is tracked in UEB-549.

diff --git a/Engine/Source/Programs/UnrealBuildTool/System/ActionGraph.cs b/Engine/Source/Programs/UnrealBuildTool/System/ActionGraph.cs
--- a/Engine/Source/Programs/UnrealBuildTool/System/ActionGraph.cs
+++ b/Engine/Source/Programs/UnrealBuildTool/System/ActionGraph.cs
@@ -470,7 +470,7 @@ namespace UnrealBuildTool
 						"OS", Environment.OSVersion.ToString(),
 						"MachineName", Environment.MachineName,
 						"NumLogicalCores", Environment.ProcessorCount.ToString(),
-						"NumPhysicalCores", Utils.GetPhysicalProcessorCount().ToString(),
+						"NumPhysicalCores", (bUsedXGE? -1 : Utils.GetPhysicalProcessorCount()).ToString(),
 						"Targets", TargetInfoForTelemetry,
 						"NumActions", ActionsToExecute.Count.ToString(),
 						"NumCompileActions", ActionsToExecute.Count(x => x.ActionType == ActionType.Compile).ToString(),
