Disable -fastpdb by default; seems to be causing missing symbols when debugging in VS2015.

diff --git a/Engine/Source/Developer/HotReload/Private/HotReload.cpp b/Engine/Source/Developer/HotReload/Private/HotReload.cpp
--- a/Engine/Source/Developer/HotReload/Private/HotReload.cpp
+++ b/Engine/Source/Developer/HotReload/Private/HotReload.cpp
@@ -1512,8 +1512,8 @@ FString FHotReloadModule::MakeUBTArgumentsForModuleCompiling()
 		}
 	}
 
-	// Use new FastPDB option to cut down linking time
-	AdditionalArguments += TEXT(" -FastPDB");
+	// Use new FastPDB option to cut down linking time. Currently disabled due to problems with missing symbols in VS2015.
+	// AdditionalArguments += TEXT(" -FastPDB");
 
 	return AdditionalArguments;
 }
diff --git a/Engine/Source/Programs/UnrealBuildTool/Configuration/BuildConfiguration.cs b/Engine/Source/Programs/UnrealBuildTool/Configuration/BuildConfiguration.cs
--- a/Engine/Source/Programs/UnrealBuildTool/Configuration/BuildConfiguration.cs
+++ b/Engine/Source/Programs/UnrealBuildTool/Configuration/BuildConfiguration.cs
@@ -699,7 +699,7 @@ namespace UnrealBuildTool
 			bUCACheckUObjectThreadSafety = false;
 
 			// Use Fast PDB linking by default in projects but not all builds
-			bAddFastPDBToProjects = true;
+			bAddFastPDBToProjects = false;
 			bUseFastPDBLinking = false;
 
 			// The default for normal Mac users should be to use DistCode which installs as an Xcode plugin and provides dynamic host management
