Add -unattended flag to UHT invocation; we never want it to open message boxes. Currently causing issues in CIS due to Dev-Core merge.

diff --git a/Engine/Source/Programs/UnrealBuildTool/System/ExternalExecution.cs b/Engine/Source/Programs/UnrealBuildTool/System/ExternalExecution.cs
--- a/Engine/Source/Programs/UnrealBuildTool/System/ExternalExecution.cs
+++ b/Engine/Source/Programs/UnrealBuildTool/System/ExternalExecution.cs
@@ -773,7 +773,7 @@ namespace UnrealBuildTool
 					System.IO.File.WriteAllText(ModuleInfoFileName.FullName, fastJSON.JSON.Instance.ToJSON(Manifest, new fastJSON.JSONParameters { UseExtensions = false }));
 
 					string CmdLine = (Target.ProjectFile != null) ? "\"" + Target.ProjectFile.FullName + "\"" : Target.GetTargetName();
-					CmdLine += " \"" + ModuleInfoFileName + "\" -LogCmds=\"loginit warning, logexit warning, logdatabase error\"";
+					CmdLine += " \"" + ModuleInfoFileName + "\" -LogCmds=\"loginit warning, logexit warning, logdatabase error\" -Unattended";
 					if (UnrealBuildTool.IsEngineInstalled())
 					{
 						CmdLine += " -installed";
