Always add the 32-bit toolchain to the PATH environment variable. If we're VS2013 Express to cross-compile for a 64-bit target using a 32-bit toolchain, it needs the PATH environment variable updating to include the standard 32-bit toolchain so it can find MSPDB120.dll.

diff --git a/Engine/Source/Programs/UnrealBuildTool/Windows/VCEnvironment.cs b/Engine/Source/Programs/UnrealBuildTool/Windows/VCEnvironment.cs
--- a/Engine/Source/Programs/UnrealBuildTool/Windows/VCEnvironment.cs
+++ b/Engine/Source/Programs/UnrealBuildTool/Windows/VCEnvironment.cs
@@ -97,6 +97,14 @@ namespace UnrealBuildTool
 			LibraryLinkerPath = GetLibraryLinkerToolPath(LinkerVSToolPath);
 			ResourceCompilerPath = GetResourceCompilerToolPath(Platform, bSupportWindowsXP);
 
+            // Make sure the base 32-bit VS tool path is in the PATH, regardless of which configuration we're using. The toolchain may need to reference support DLLs from this directory (eg. mspdb120.dll).
+            string PathEnvironmentVariable = Environment.GetEnvironmentVariable("PATH");
+            if (!PathEnvironmentVariable.Split(';').Any(x => String.Compare(x, VSToolPath32Bit, true) == 0))
+            {
+                PathEnvironmentVariable = VSToolPath32Bit + ";" + PathEnvironmentVariable;
+                Environment.SetEnvironmentVariable("PATH", PathEnvironmentVariable);
+            }
+
 			// Setup the INCLUDE environment variable
 			List<string> IncludePaths = GetVisualCppIncludePaths(VisualCppDir, UniversalCRTDir, UniversalCRTVersion, NetFxSDKExtensionDir, WindowsSDKDir, WindowsSDKLibVersion);
 			if(InitialIncludePaths != null)
