Copying //UE4/Dev-Core to //UE4/Main

diff --git a/Engine/Source/Developer/Profiler/Private/Widgets/SProfilerMiniView.h b/Engine/Source/Developer/Profiler/Private/Widgets/SProfilerMiniView.h
--- a/Engine/Source/Developer/Profiler/Private/Widgets/SProfilerMiniView.h
+++ b/Engine/Source/Developer/Profiler/Private/Widgets/SProfilerMiniView.h
@@ -296,7 +296,7 @@ protected:
 
 	const int32 FindMiniViewSampleIndex( const int32 FrameIndex ) const
 	{
-		// @TODO yrx 2014-04-07 Naive implementation, consider using binary search.
+		// #Profiler: 2014-04-07 Naive implementation, consider using binary search.
 		for( int32 SampleIndex = 0; SampleIndex < MiniViewSamples.Num(); ++SampleIndex )
 		{
 			const FMiniViewSample& MiniViewSample = MiniViewSamples[SampleIndex];
diff --git a/Engine/Source/Developer/Profiler/Private/Widgets/SProfilerThreadView.cpp b/Engine/Source/Developer/Profiler/Private/Widgets/SProfilerThreadView.cpp
--- a/Engine/Source/Developer/Profiler/Private/Widgets/SProfilerThreadView.cpp
+++ b/Engine/Source/Developer/Profiler/Private/Widgets/SProfilerThreadView.cpp
@@ -300,7 +300,7 @@ void SProfilerThreadView::DrawUIStackNodes() const
 		}
 	}
 
-	// @TODO yrx 2014-04-29 Separate layer for makebox, makeshadowtext, maketext.
+	// #Profiler: 2014-04-29 Separate layer for makebox, makeshadowtext, maketext.
 	PaintState->LayerId++;
 
 	const float MarkerPosYOffsetPx = ((float)NUM_PIXELS_PER_ROW - PaintState->SummaryFont8Height)*0.5f;
diff --git a/Engine/Source/Developer/Profiler/Private/Widgets/SProfilerThreadView.h b/Engine/Source/Developer/Profiler/Private/Widgets/SProfilerThreadView.h
--- a/Engine/Source/Developer/Profiler/Private/Widgets/SProfilerThreadView.h
+++ b/Engine/Source/Developer/Profiler/Private/Widgets/SProfilerThreadView.h
@@ -45,7 +45,7 @@ class SProfilerThreadView : public SCompoundWidget
 		 *	One partition must have a least one frame.
 		 */
 		NUM_DATA_PARTITIONS = 16,
-		// @TODO yrx 2014-04-25 Dynamic data partitioning
+		// #Profiler: 2014-04-25 Dynamic data partitioning
 
 		/**
 		 *	Maximum zoom value for time axis.
diff --git a/Engine/Source/Developer/Profiler/Public/ProfilerCommon.h b/Engine/Source/Developer/Profiler/Public/ProfilerCommon.h
--- a/Engine/Source/Developer/Profiler/Public/ProfilerCommon.h
+++ b/Engine/Source/Developer/Profiler/Public/ProfilerCommon.h
@@ -35,12 +35,6 @@ typedef TSharedPtr<class FProfilerSession> FProfilerSessionPtr;
 /** Type definition for shared references to instances of FProfilerSession. */
 typedef TSharedRef<class FProfilerSession> FProfilerSessionRef;
 
-/** Type definition for shared references to const instances of FProfilerSession. */
-typedef TSharedRef<const class FProfilerSession> FProfilerSessionRefConst;
-
-/** Type definition for weak references to instances of FProfilerSession. */
-typedef TWeakPtr<class FProfilerSession> FProfilerSessionWeak;
-
 /** Type definition for shared pointers to instances of FProfilerStatMetaData. */
 typedef TSharedPtr<class FProfilerStatMetaData> FProfilerStatMetaDataPtr;
 
diff --git a/Engine/Source/Developer/ProfilerClient/Private/ProfilerClientManager.cpp b/Engine/Source/Developer/ProfilerClient/Private/ProfilerClientManager.cpp
--- a/Engine/Source/Developer/ProfilerClient/Private/ProfilerClientManager.cpp
+++ b/Engine/Source/Developer/ProfilerClient/Private/ProfilerClientManager.cpp
@@ -598,7 +598,7 @@ bool FProfilerClientManager::HandleMessagesTicker( float DeltaTime )
 				break;
 			}
 
-			FScopeLogTime SLT( "HandleMessagesTicker" );
+			//FScopeLogTime SLT( "HandleMessagesTicker" );
 
 			const int64 FrameNum = Frames[Index];
 			const TArray<uint8>* const Data = Connection.ReceivedData.FindChecked( FrameNum );
@@ -734,7 +734,6 @@ void FProfilerClientManager::SendDataToGame( TArray<uint8>* DataToGame, int64 Fr
 	}
 }
 
-// #YRX_tag 2015-11-19 Convert to thread safe pointers?
 void FProfilerClientManager::SendProfilerDataFrameToGame( FProfilerDataFrame* NewData, FStatMetaData* MetaDataPtr )
 {
 	if (MetaDataPtr)
diff --git a/Engine/Source/Developer/ProfilerClient/Private/ProfilerClientManager.h b/Engine/Source/Developer/ProfilerClient/Private/ProfilerClientManager.h
--- a/Engine/Source/Developer/ProfilerClient/Private/ProfilerClientManager.h
+++ b/Engine/Source/Developer/ProfilerClient/Private/ProfilerClientManager.h
@@ -211,7 +211,7 @@ private:
 	TArray<FGuid> PendingInstances;
 
 	/** Service connections */
-	// #YRX_tag 2015-11-19 should be only one active connection.
+	// #Profiler: 2015-11-19 should be only one active connection.
 	TMap<FGuid, FServiceConnection> Connections;
 
 	struct FReceivedFileInfo
diff --git a/Engine/Source/Editor/LevelEditor/Private/LevelViewportActions.cpp b/Engine/Source/Editor/LevelEditor/Private/LevelViewportActions.cpp
--- a/Engine/Source/Editor/LevelEditor/Private/LevelViewportActions.cpp
+++ b/Engine/Source/Editor/LevelEditor/Private/LevelViewportActions.cpp
@@ -263,7 +263,7 @@ void FLevelViewportCommands::RegisterCommands()
 
 void FLevelViewportCommands::HandleNewStatGroup(const TArray<FStatNameAndInfo>& NameAndInfos)
 {
-	// @TODO yrx 2014-10-17 FStatNameAndInfo should be private and visible only to stats2 system
+	// #Stats: FStatNameAndInfo should be private and visible only to stats2 system
 	for (int32 InfoIdx = 0; InfoIdx < NameAndInfos.Num(); InfoIdx++)
 	{
 		const FStatNameAndInfo& NameAndInfo = NameAndInfos[InfoIdx];
diff --git a/Engine/Source/Programs/AutomationTool/SharedUtils.cs b/Engine/Source/Programs/AutomationTool/SharedUtils.cs
--- a/Engine/Source/Programs/AutomationTool/SharedUtils.cs
+++ b/Engine/Source/Programs/AutomationTool/SharedUtils.cs
@@ -15,7 +15,7 @@ namespace AutomationTool
 	{
 		/// <summary>
 		/// Parses the command line string and returns an array of passed arguments.
-		/// Unlike the default parsing algorithm, this will trean \r\n characters
+		/// Unlike the default parsing algorithm, this will treat \r\n characters
 		/// just like spaces.
 		/// </summary>
 		/// <returns>List of command line arguments.</returns>
@@ -25,19 +25,13 @@ namespace AutomationTool
 			var Args = new List<string>();
 			StringBuilder Arg = new StringBuilder(CmdLine.Length);
 			bool bQuote = false;
-			bool bEscape = false;
+			int bEscape = 0;
 			for (int Index = 0; Index < CmdLine.Length; ++Index)
 			{
 				bool bCanAppend = true;
 				char C = CmdLine[Index];
 				if (!bQuote && Char.IsWhiteSpace(C))
 				{
-					// There can be an escape character here when passing a Windows-style path with a trailing backslash (eg. "-MyArg1=D:\UE4\ -MyArg2"). We don't want to push it into the next argument.
-					if(bEscape)
-					{
-						Arg.Append('\\');
-						bEscape = false;
-					}
 					if (Arg.Length > 0)
 					{
 						Args.Add(Arg.ToString());
@@ -48,29 +42,29 @@ namespace AutomationTool
 				else if (C == '\\' && Index < (CmdLine.Length - 1))
 				{
 					// Escape character
-					bEscape = true;
+					bEscape++;
 					bCanAppend = false;
 				}
 				else if (C == '\"')
 				{
-					bCanAppend = bEscape;
-					if (!bEscape)
+					bCanAppend = bEscape > 0;
+					if (bEscape == 0)
 					{
 						bQuote = !bQuote;
 					}
 					else
 					{
 						// Consume the scape character
-						bEscape = false;
+						bEscape--;
 					}
 				}
 				if (bCanAppend)
 				{
-					if (bEscape)
+					if (bEscape > 0)
 					{
 						// Unused escape character.
-						Arg.Append('\\');
-						bEscape = false;
+						Arg.Append('\\', bEscape);
+						bEscape = 0;
 					}
 					Arg.Append(C);
 				}
diff --git a/Engine/Source/Programs/CrashReporter/CrashReportClient/Private/CrashDescription.cpp b/Engine/Source/Programs/CrashReporter/CrashReportClient/Private/CrashDescription.cpp
--- a/Engine/Source/Programs/CrashReporter/CrashReportClient/Private/CrashDescription.cpp
+++ b/Engine/Source/Programs/CrashReporter/CrashReportClient/Private/CrashDescription.cpp
@@ -11,7 +11,7 @@
 #include "Runtime/Analytics/Analytics/Public/Interfaces/IAnalyticsProvider.h"
 #include "EngineBuildSettings.h"
 
-// #YRX_Crash: 2015-07-23 Move crashes from [REDACTED_USER_HOME]\AppData\Local\Microsoft\Windows\WER\ReportQueue to [REDACTED_USER_HOME]\AppData\Local\CrashReportClient\Saved
+// #CrashReport: 2015-07-23 Move crashes from [REDACTED_USER_HOME]\AppData\Local\Microsoft\Windows\WER\ReportQueue to [REDACTED_USER_HOME]\AppData\Local\CrashReportClient\Saved
 
 /*-----------------------------------------------------------------------------
 	FCrashProperty
diff --git a/Engine/Source/Programs/CrashReporter/CrashReportCommon/Common.cs b/Engine/Source/Programs/CrashReporter/CrashReportCommon/Common.cs
--- a/Engine/Source/Programs/CrashReporter/CrashReportCommon/Common.cs
+++ b/Engine/Source/Programs/CrashReporter/CrashReportCommon/Common.cs
@@ -132,6 +132,10 @@ namespace Tools.CrashReporter.CrashReportCommon
 		[XmlElement]
 		public string[] UserDescription = null;
 
+		/// <summary>A string representing the user activity hint text from the application.</summary>
+		[XmlElement]
+		public string UserActivityHint = "";
+
 		/// <summary>The error message, can be assertion message, ensure message or message from the fatal error.</summary>
 		[XmlElement]
 		public string[] ErrorMessage = null;
diff --git a/Engine/Source/Programs/CrashReporter/CrashReportCommon/CrashContext.cs b/Engine/Source/Programs/CrashReporter/CrashReportCommon/CrashContext.cs
--- a/Engine/Source/Programs/CrashReporter/CrashReportCommon/CrashContext.cs
+++ b/Engine/Source/Programs/CrashReporter/CrashReportCommon/CrashContext.cs
@@ -65,6 +65,7 @@ namespace Tools.CrashReporter.CrashReportCommon
 		"CallStack"
 		"SourceContext"
 		"UserDescription"
+		"UserActivityHint" 
 		"ErrorMessage"
 		"CrashDumpMode"
 		"Misc.NumberOfCores"
@@ -458,6 +459,10 @@ namespace Tools.CrashReporter.CrashReportCommon
 			return UnescapedValue.Split( new char[] { '\n' }, StringSplitOptions.RemoveEmptyEntries );
 		}
 
+		/// <summary></summary>
+		[XmlElement]
+		public string UserActivityHint;
+
 		/// <summary></summary>
 		[XmlElement]
 		public string SourceContext;
diff --git a/Engine/Source/Programs/CrashReporter/CrashReportCommon/CrashReportCommon.csproj b/Engine/Source/Programs/CrashReporter/CrashReportCommon/CrashReportCommon.csproj
--- a/Engine/Source/Programs/CrashReporter/CrashReportCommon/CrashReportCommon.csproj
+++ b/Engine/Source/Programs/CrashReporter/CrashReportCommon/CrashReportCommon.csproj
@@ -37,6 +37,17 @@
     <CodeAnalysisIgnoreGeneratedCode>true</CodeAnalysisIgnoreGeneratedCode>
     <RunCodeAnalysis>true</RunCodeAnalysis>
   </PropertyGroup>
+  <PropertyGroup Condition="'$(Configuration)|$(Platform)' == 'Release|AnyCPU'">
+    <OutputPath>bin\Release\</OutputPath>
+    <DefineConstants>CODE_ANALYSIS;TRACE</DefineConstants>
+    <DocumentationFile>..\..\..\..\Binaries\DotNET\CrashReportCommon.xml</DocumentationFile>
+    <Optimize>true</Optimize>
+    <DebugType>pdbonly</DebugType>
+    <PlatformTarget>AnyCPU</PlatformTarget>
+    <RunCodeAnalysis>true</RunCodeAnalysis>
+    <ErrorReport>prompt</ErrorReport>
+    <CodeAnalysisRuleSet>MinimumRecommendedRules.ruleset</CodeAnalysisRuleSet>
+  </PropertyGroup>
   <ItemGroup>
     <Reference Include="System" />
     <Reference Include="System.Core" />
diff --git a/Engine/Source/Programs/CrashReporter/CrashReportCommon/ReportData.cs b/Engine/Source/Programs/CrashReporter/CrashReportCommon/ReportData.cs
--- a/Engine/Source/Programs/CrashReporter/CrashReportCommon/ReportData.cs
+++ b/Engine/Source/Programs/CrashReporter/CrashReportCommon/ReportData.cs
@@ -16,6 +16,13 @@ namespace Tools.CrashReporter.CrashReportCommon
 			int MinusPos = StringEngineVersion.IndexOf( '-' );
 			int PlusPos = StringEngineVersion.IndexOf( '+' );
 
+			if (MinusPos < 0 || PlusPos < 0 || PlusPos < MinusPos)
+			{
+				VersionNumber = string.Empty;
+				Branch = string.Empty;
+				return;
+			}
+
 			VersionNumber = StringEngineVersion.Substring( 0, MinusPos );
 			string ChangelistSring = StringEngineVersion.Substring( MinusPos + 1, PlusPos - MinusPos - 1 );
 			uint.TryParse( ChangelistSring, out Changelist );
diff --git a/Engine/Source/Programs/CrashReporter/CrashReportProcess/App.config b/Engine/Source/Programs/CrashReporter/CrashReportProcess/App.config
--- a/Engine/Source/Programs/CrashReporter/CrashReportProcess/App.config
+++ b/Engine/Source/Programs/CrashReporter/CrashReportProcess/App.config
@@ -12,7 +12,7 @@
 	<applicationSettings>
 		<Tools.CrashReporter.CrashReportProcess.Properties.Settings>
    <setting name="ProcessedReports" serializeAs="String">
-    <value>\\devweb-02\Sites\CrashReporterFiles</value>
+    <value>E:\Sites\CrashReporterFiles</value>
    </setting>
    <setting name="DaysToSunsetReport" serializeAs="String">
     <value>14</value>
@@ -21,23 +21,35 @@
     <value>CrashReporterUE4</value>
    </setting>
    <setting name="ProcessedVideos" serializeAs="String">
-    <value>\\devweb-02\ReportVideos</value>
+    <value>F:\CrashReporterVideos</value>
    </setting>
    <setting name="DepotRoot" serializeAs="String">
-    <value>D:\P4EPIC</value>
+    <value>F:\depot</value>
    </setting>
    <setting name="InternalLandingZone" serializeAs="String">
-    <value>D:\!CrashReport\DMZ</value>
+    <value>F:\InternalCrashReports\CrashReportRepository</value>
    </setting>
    <setting name="ExternalLandingZone" serializeAs="String">
-    <value />
+    <value>\\na-web-cluster.dmz.epicgames.net\crash_reporter\Rep</value>
    </setting>
    <setting name="InvalidReportsDirectory" serializeAs="String">
-    <value>D:\!CrashReport\InvalidCrashReports</value>
+    <value>F:\InvalidCrashReports</value>
    </setting>
    <setting name="DepotIndex" serializeAs="String">
     <value>crash-01</value>
    </setting>
+   <setting name="SyncBinariesFromDepot" serializeAs="String">
+    <value>Stream\UE4\Main\Engine\Binaries\Win64\MinidumpDiagnostics*</value>
+   </setting>
+   <setting name="SyncConfigFromDepot" serializeAs="String">
+    <value>Stream\UE4\Main\Engine\Config\...</value>
+   </setting>
+   <setting name="SyncThirdPartyFromDepot" serializeAs="String">
+    <value>Stream\UE4\Main\Engine\Binaries\ThirdParty\OpenSSL\...</value>
+   </setting>
+   <setting name="MDDBinariesFolderInDepot" serializeAs="String">
+    <value>Stream\UE4\Main\Engine\Binaries\Win64</value>
+   </setting>
   </Tools.CrashReporter.CrashReportProcess.Properties.Settings>
 	</applicationSettings>
 </configuration>
diff --git a/Engine/Source/Programs/CrashReporter/CrashReportProcess/Configs/CrashReportProcess.exe.config b/Engine/Source/Programs/CrashReporter/CrashReportProcess/Configs/CrashReportProcess.exe.config
--- a/Engine/Source/Programs/CrashReporter/CrashReportProcess/Configs/CrashReportProcess.exe.config
+++ b/Engine/Source/Programs/CrashReporter/CrashReportProcess/Configs/CrashReportProcess.exe.config
@@ -7,12 +7,30 @@
     </configSections>
     <applicationSettings>
         <Tools.CrashReporter.CrashReportProcess.Properties.Settings>
-            <setting name="InternalLandingZone" serializeAs="String">
-                <value>F:\InternalCrashReports\CrashReportRepository</value>
-            </setting>
-            <setting name="ExternalLandingZone" serializeAs="String">
-                <value>\\na-web-cluster.dmz.epicgames.net\crash_reporter\CrashReportRepository</value>
-            </setting>
+			<setting name="ProcessedReports" serializeAs="String">
+				<value>E:\Sites\CrashReporterFiles</value>
+			</setting>
+			<setting name="DaysToSunsetReport" serializeAs="String">
+				<value>14</value>
+			</setting>
+			<setting name="DepotRoot" serializeAs="String">
+				<value>F:\depot</value>
+			</setting>
+			<setting name="CrashReportWebSite" serializeAs="String">
+				<value>CrashReporterUE4</value>
+			</setting>
+			<setting name="ProcessedVideos" serializeAs="String">
+				<value>F:\CrashReporterVideos</value>
+			</setting>
+			<setting name="InternalLandingZone" serializeAs="String">
+				<value>F:\InternalCrashReports\CrashReportRepository</value>
+			</setting>
+			<setting name="ExternalLandingZone" serializeAs="String">
+				<value>\\na-web-cluster.dmz.epicgames.net\crash_reporter\Rep</value>
+			</setting>
+			<setting name="InvalidReportsDirectory" serializeAs="String">
+				<value>F:\InvalidCrashReports</value>
+			</setting>
 			<setting name="DepotIndex" serializeAs="String">
 				<value>crash-01</value>
 			</setting>
diff --git a/Engine/Source/Programs/CrashReporter/CrashReportProcess/CrashReportProcess.csproj b/Engine/Source/Programs/CrashReporter/CrashReportProcess/CrashReportProcess.csproj
--- a/Engine/Source/Programs/CrashReporter/CrashReportProcess/CrashReportProcess.csproj
+++ b/Engine/Source/Programs/CrashReporter/CrashReportProcess/CrashReportProcess.csproj
@@ -49,6 +49,18 @@
   <PropertyGroup>
     <ManifestCertificateThumbprint>7B0CF23A7AEEF0637C57BC4BD76F2F43119EFAD6</ManifestCertificateThumbprint>
   </PropertyGroup>
+  <PropertyGroup Condition="'$(Configuration)|$(Platform)' == 'Release|AnyCPU'">
+    <OutputPath>bin\Release\</OutputPath>
+    <DefineConstants>TRACE</DefineConstants>
+    <AllowUnsafeBlocks>true</AllowUnsafeBlocks>
+    <DocumentationFile>..\..\..\..\Binaries\DotNET\CrashReportProcess.xml</DocumentationFile>
+    <Optimize>true</Optimize>
+    <DebugType>pdbonly</DebugType>
+    <PlatformTarget>AnyCPU</PlatformTarget>
+    <RunCodeAnalysis>true</RunCodeAnalysis>
+    <ErrorReport>prompt</ErrorReport>
+    <CodeAnalysisRuleSet>MinimumRecommendedRules.ruleset</CodeAnalysisRuleSet>
+  </PropertyGroup>
   <ItemGroup>
     <Reference Include="System" />
     <Reference Include="System.Configuration.Install" />
@@ -88,9 +100,12 @@
     <Compile Include="ReportProcessor.cs" />
     <Compile Include="ReportQueue.cs" />
     <Compile Include="ReportWatcher.cs" />
+    <Compile Include="Settings.cs" />
   </ItemGroup>
   <ItemGroup>
-    <None Include="App.config" />
+    <None Include="App.config">
+      <SubType>Designer</SubType>
+    </None>
     <None Include="Properties\Settings.settings">
       <Generator>SettingsSingleFileGenerator</Generator>
       <LastGenOutput>Settings.Designer.cs</LastGenOutput>
diff --git a/Engine/Source/Programs/CrashReporter/CrashReportProcess/Properties/Settings.settings b/Engine/Source/Programs/CrashReporter/CrashReportProcess/Properties/Settings.settings
--- a/Engine/Source/Programs/CrashReporter/CrashReportProcess/Properties/Settings.settings
+++ b/Engine/Source/Programs/CrashReporter/CrashReportProcess/Properties/Settings.settings
@@ -3,7 +3,7 @@
   <Profiles />
   <Settings>
     <Setting Name="ProcessedReports" Type="System.String" Scope="Application">
-      <Value Profile="(Default)">\\devweb-02\Sites\CrashReporterFiles</Value>
+      <Value Profile="(Default)">E:\Sites\CrashReporterFiles</Value>
     </Setting>
     <Setting Name="DaysToSunsetReport" Type="System.Int32" Scope="Application">
       <Value Profile="(Default)">14</Value>
@@ -12,22 +12,34 @@
       <Value Profile="(Default)">CrashReporterUE4</Value>
     </Setting>
     <Setting Name="ProcessedVideos" Type="System.String" Scope="Application">
-      <Value Profile="(Default)">\\devweb-02\ReportVideos</Value>
+      <Value Profile="(Default)">F:\CrashReporterVideos</Value>
     </Setting>
     <Setting Name="DepotRoot" Type="System.String" Scope="Application">
-      <Value Profile="(Default)">D:\P4EPIC</Value>
+      <Value Profile="(Default)">F:\depot</Value>
     </Setting>
     <Setting Name="InternalLandingZone" Type="System.String" Scope="Application">
-      <Value Profile="(Default)">D:\!CrashReport\DMZ</Value>
+      <Value Profile="(Default)">F:\InternalCrashReports\CrashReportRepository</Value>
     </Setting>
     <Setting Name="ExternalLandingZone" Type="System.String" Scope="Application">
-      <Value Profile="(Default)" />
+      <Value Profile="(Default)">\\na-web-cluster.dmz.epicgames.net\crash_reporter\Rep</Value>
     </Setting>
     <Setting Name="InvalidReportsDirectory" Type="System.String" Scope="Application">
-      <Value Profile="(Default)">D:\!CrashReport\InvalidCrashReports</Value>
+      <Value Profile="(Default)">F:\InvalidCrashReports</Value>
     </Setting>
     <Setting Name="DepotIndex" Type="System.String" Scope="Application">
       <Value Profile="(Default)">crash-01</Value>
     </Setting>
+    <Setting Name="SyncBinariesFromDepot" Type="System.String" Scope="Application">
+      <Value Profile="(Default)">Stream\UE4\Main\Engine\Binaries\Win64\MinidumpDiagnostics*</Value>
+    </Setting>
+    <Setting Name="SyncConfigFromDepot" Type="System.String" Scope="Application">
+      <Value Profile="(Default)">Stream\UE4\Main\Engine\Config\...</Value>
+    </Setting>
+    <Setting Name="SyncThirdPartyFromDepot" Type="System.String" Scope="Application">
+      <Value Profile="(Default)">Stream\UE4\Main\Engine\Binaries\ThirdParty\OpenSSL\...</Value>
+    </Setting>
+    <Setting Name="MDDBinariesFolderInDepot" Type="System.String" Scope="Application">
+      <Value Profile="(Default)">Stream\UE4\Main\Engine\Binaries\Win64</Value>
+    </Setting>
   </Settings>
 </SettingsFile>
\ No newline at end of file
diff --git a/Engine/Source/Programs/CrashReporter/CrashReportProcess/ReportProcessor.cs b/Engine/Source/Programs/CrashReporter/CrashReportProcess/ReportProcessor.cs
--- a/Engine/Source/Programs/CrashReporter/CrashReportProcess/ReportProcessor.cs
+++ b/Engine/Source/Programs/CrashReporter/CrashReportProcess/ReportProcessor.cs
@@ -82,11 +82,8 @@ namespace Tools.CrashReporter.CrashReportProcess
 		/// <remarks>All exceptions are caught and written to the event log.</remarks>
 		private void ProcessNewReports()
 		{
-			// Use the latest MinidumpDiagnostics from the main branch.
-			string Win64BinariesDirectory = Path.Combine( Properties.Settings.Default.DepotRoot, "UE4", "Engine", "Binaries", "Win64" );
 #if !DEBUG
-			string MinidumpDiagnosticsName = Path.Combine(Win64BinariesDirectory, "MinidumpDiagnostics*");
-			if (!SyncRequiredFiles(MinidumpDiagnosticsName))
+			if (!SyncRequiredFiles())
 			{
 				return;
 			}
@@ -317,7 +314,8 @@ namespace Tools.CrashReporter.CrashReportProcess
 				NewCrash.SourceContext = NewContext.PrimaryCrashProperties.GetSourceContext();
 				NewCrash.ErrorMessage = NewContext.PrimaryCrashProperties.GetErrorMessage();
 				NewCrash.UserDescription = NewContext.PrimaryCrashProperties.GetUserDescription();
-				
+				NewCrash.UserActivityHint = NewContext.PrimaryCrashProperties.UserActivityHint;
+
 				// Iterate through all files and find a file with the earliest date.
 				DateTime TimeOfCrash = DateTime.UtcNow;
 				foreach( var File in DirInfo.GetFiles() )
@@ -433,6 +431,8 @@ namespace Tools.CrashReporter.CrashReportProcess
 				if (ReportID <= 0)
 				{
 					CrashReporterProcessServicer.WriteFailure( "! NoUpload: Path=" + NewContext.CrashDirectory );
+					string PayloadFailedFileName = Path.Combine(DirInfo.FullName, "PayloadFailed.xml");
+					File.WriteAllText(PayloadFailedFileName, CrashDetails);
 					return false;
 				}
 
@@ -518,39 +518,37 @@ namespace Tools.CrashReporter.CrashReportProcess
 		/// <summary>
 		/// Sync the MinidumpDiagnostics binary and the engine config files to #head.
 		/// </summary>
-		/// <param name="MinidumpDiagnosticsName">The name of the binary to sync to process the minidump.</param>
 		/// <returns>true if all the files synced without issue, false otherwise.</returns>
 		/// <remarks>As MinidumpDiagnostics is synced to #head, it requires the engine config files that match to run properly.</remarks>
-		private bool SyncRequiredFiles(string MinidumpDiagnosticsName)
+		private bool SyncRequiredFiles()
 		{
-			using( var MDDSyncProc = new LaunchProcess( PerforceExePath, null, CrashReporterProcessServicer.WriteP4, "sync", MinidumpDiagnosticsName ) )
+			// Use the latest MinidumpDiagnostics from the main branch.
+			string SyncBinariesString = Path.Combine(Properties.Settings.Default.DepotRoot, Properties.Settings.Default.SyncBinariesFromDepot);
+			using (var MDDSyncProc = new LaunchProcess(PerforceExePath, null, CrashReporterProcessServicer.WriteP4, "sync", SyncBinariesString))
 			{
 				if( MDDSyncProc.WaitForExit( SyncTimeoutSeconds * 1000 ) == EWaitResult.TimedOut )
 				{
-					CrashReporterProcessServicer.WriteFailure( "Failed to sync MinidumpDiagnostics " + MinidumpDiagnosticsName );
-
+					CrashReporterProcessServicer.WriteFailure("Failed to sync MinidumpDiagnostics " + SyncBinariesString);
 					return false;
 				}
 			}
-			
-			string ConfigFiles = Path.Combine(Properties.Settings.Default.DepotRoot, "UE4", "Engine/Config/...");
-			using (var ConfigSyncProc = new LaunchProcess(PerforceExePath, null, CrashReporterProcessServicer.WriteP4, "sync", ConfigFiles))
+
+			string SyncConfigString = Path.Combine(Properties.Settings.Default.DepotRoot, Properties.Settings.Default.SyncConfigFromDepot);
+			using (var ConfigSyncProc = new LaunchProcess(PerforceExePath, null, CrashReporterProcessServicer.WriteP4, "sync", SyncConfigString))
 			{
 				if (ConfigSyncProc.WaitForExit(SyncTimeoutSeconds * 1000) == EWaitResult.TimedOut)
 				{
-					CrashReporterProcessServicer.WriteFailure( "Failed to sync config files " + ConfigFiles );
-
+					CrashReporterProcessServicer.WriteFailure("Failed to sync config files " + SyncConfigString);
 					return false;
 				}
 			}
 
-			string SSLFiles = Path.Combine(Properties.Settings.Default.DepotRoot, "UE4", "Engine/Binaries/ThirdParty/OpenSSL/..."); // Required by Perforce
-			using (var MiscSyncProc = new LaunchProcess(PerforceExePath, null, CrashReporterProcessServicer.WriteP4, "sync", SSLFiles))
+			string SyncThirdPartyString = Path.Combine(Properties.Settings.Default.DepotRoot, Properties.Settings.Default.SyncThirdPartyFromDepot); // Required by Perforce
+			using (var MiscSyncProc = new LaunchProcess(PerforceExePath, null, CrashReporterProcessServicer.WriteP4, "sync", SyncThirdPartyString))
 			{
 				if (MiscSyncProc.WaitForExit(SyncTimeoutSeconds * 1000) == EWaitResult.TimedOut)
 				{
-					CrashReporterProcessServicer.WriteFailure( "Failed to sync OpenSSL files " + SSLFiles );
-
+					CrashReporterProcessServicer.WriteFailure("Failed to sync OpenSSL files " + SyncThirdPartyString);
 					return false;
 				}
 			}
@@ -628,7 +626,7 @@ namespace Tools.CrashReporter.CrashReportProcess
 						break;
 
 					case ".txt":
-						if( Info.Name == CrashReporterConstants.DiagnosticsFileName )
+						if (string.Compare( Info.Name, CrashReporterConstants.DiagnosticsFileName, true ) == 0)
 						{
 							DiagnosticsFileName = Info.Name;
 							ReadDiagnosticsFile( NewContext );				
@@ -698,11 +696,10 @@ namespace Tools.CrashReporter.CrashReportProcess
 		void ProcessDumpFile( string DiagnosticsPath, FGenericCrashContext NewContext )
 		{
 			// Use the latest MinidumpDiagnostics from the main branch.
-			string Win64BinariesDirectory = Path.Combine(Properties.Settings.Default.DepotRoot, "UE4", "Engine", "Binaries", "Win64");
+			string Win64BinariesDirectory = Path.Combine(Properties.Settings.Default.DepotRoot, Properties.Settings.Default.MDDBinariesFolderInDepot);
 #if DEBUG
 			// Note: the debug executable must be built locally or synced from Perforce manually
 			string MinidumpDiagnosticsName = Path.Combine(Win64BinariesDirectory, "MinidumpDiagnostics-Win64-Debug.exe");
-			//string MinidumpDiagnosticsName = Path.Combine(Win64BinariesDirectory, "MinidumpDiagnostics.exe");
 #else
 			string MinidumpDiagnosticsName = Path.Combine(Win64BinariesDirectory, "MinidumpDiagnostics.exe");
 #endif
diff --git a/Engine/Source/Programs/CrashReporter/CrashReportProcess/ReportQueue.cs b/Engine/Source/Programs/CrashReporter/CrashReportProcess/ReportQueue.cs
--- a/Engine/Source/Programs/CrashReporter/CrashReportProcess/ReportQueue.cs
+++ b/Engine/Source/Programs/CrashReporter/CrashReportProcess/ReportQueue.cs
@@ -401,7 +401,7 @@ namespace Tools.CrashReporter.CrashReportProcess
 				string BranchName = EngineVersion.Branch;
 				if (string.IsNullOrEmpty( BranchName ))
 				{
-					CrashReporterProcessServicer.WriteFailure( "% NoBranch: BuiltFromCL=" + string.Format( "{0,7}", BuiltFromCL ) + " Path=" + NewReportPath );
+					CrashReporterProcessServicer.WriteFailure("% NoBranch: BuiltFromCL=" + string.Format("{0,7}", BuiltFromCL) + " Path=" + NewReportPath + " EngineVersion=" + Context.PrimaryCrashProperties.EngineVersion);
 					ReportProcessor.MoveReportToInvalid( NewReportPath, Context.GetAsFilename() );
 					return;
 				}
diff --git a/Engine/Source/Programs/CrashReporter/CrashReportProcess/Settings.cs b/Engine/Source/Programs/CrashReporter/CrashReportProcess/Settings.cs
new file mode 100644
--- /dev/null
+++ b/Engine/Source/Programs/CrashReporter/CrashReportProcess/Settings.cs
@@ -0,0 +1,28 @@
+﻿namespace Tools.CrashReporter.CrashReportProcess.Properties {
+    
+    
+    // This class allows you to handle specific events on the settings class:
+    //  The SettingChanging event is raised before a setting's value is changed.
+    //  The PropertyChanged event is raised after a setting's value is changed.
+    //  The SettingsLoaded event is raised after the setting values are loaded.
+    //  The SettingsSaving event is raised before the setting values are saved.
+    internal sealed partial class Settings {
+        
+        public Settings() {
+            // // To add event handlers for saving and changing settings, uncomment the lines below:
+            //
+            // this.SettingChanging += this.SettingChangingEventHandler;
+            //
+            // this.SettingsSaving += this.SettingsSavingEventHandler;
+            //
+        }
+        
+        private void SettingChangingEventHandler(object sender, System.Configuration.SettingChangingEventArgs e) {
+            // Add code to handle the SettingChangingEvent event here.
+        }
+        
+        private void SettingsSavingEventHandler(object sender, System.ComponentModel.CancelEventArgs e) {
+            // Add code to handle the SettingsSaving event here.
+        }
+    }
+}
diff --git a/Engine/Source/Programs/CrashReporter/CrashReportReceiver/CrashReportReceiver.csproj b/Engine/Source/Programs/CrashReporter/CrashReportReceiver/CrashReportReceiver.csproj
--- a/Engine/Source/Programs/CrashReporter/CrashReportReceiver/CrashReportReceiver.csproj
+++ b/Engine/Source/Programs/CrashReporter/CrashReportReceiver/CrashReportReceiver.csproj
@@ -39,6 +39,18 @@
     <CodeAnalysisRuleSet>MinimumRecommendedRules.ruleset</CodeAnalysisRuleSet>
     <RunCodeAnalysis>true</RunCodeAnalysis>
   </PropertyGroup>
+  <PropertyGroup Condition="'$(Configuration)|$(Platform)' == 'Release|AnyCPU'">
+    <OutputPath>bin\Release\</OutputPath>
+    <DefineConstants>CODE_ANALYSIS;TRACE</DefineConstants>
+    <DocumentationFile>..\..\..\..\Binaries\DotNET\CrashReportReceiver.xml</DocumentationFile>
+    <Optimize>true</Optimize>
+    <DebugType>pdbonly</DebugType>
+    <PlatformTarget>AnyCPU</PlatformTarget>
+    <RunCodeAnalysis>true</RunCodeAnalysis>
+    <ErrorReport>prompt</ErrorReport>
+    <CodeAnalysisRuleSet>MinimumRecommendedRules.ruleset</CodeAnalysisRuleSet>
+    <Prefer32Bit>true</Prefer32Bit>
+  </PropertyGroup>
   <ItemGroup>
     <Reference Include="System" />
     <Reference Include="System.Configuration.Install" />
diff --git a/Engine/Source/Programs/CrashReporter/CrashReportWebSite/CrashReportWebSite.csproj b/Engine/Source/Programs/CrashReporter/CrashReportWebSite/CrashReportWebSite.csproj
--- a/Engine/Source/Programs/CrashReporter/CrashReportWebSite/CrashReportWebSite.csproj
+++ b/Engine/Source/Programs/CrashReporter/CrashReportWebSite/CrashReportWebSite.csproj
@@ -261,7 +261,7 @@
       <Project>{5d7d66e8-8c76-4af9-b3ec-2ef03421d730}</Project>
       <Name>DotNETUtilities</Name>
     </ProjectReference>
-    <ProjectReference Include="..\..\NotForLicensees\CrashReportServer\CrashReportCommon\CrashReportCommon.csproj">
+    <ProjectReference Include="..\CrashReportCommon\CrashReportCommon.csproj">
       <Project>{b83de8c9-9429-4765-86bd-8eeea89b138f}</Project>
       <Name>CrashReportCommon</Name>
     </ProjectReference>
diff --git a/Engine/Source/Programs/CrashReporter/CrashReportWebSite/Models/CrashDataFormater.cs b/Engine/Source/Programs/CrashReporter/CrashReportWebSite/Models/CrashDataFormater.cs
--- a/Engine/Source/Programs/CrashReporter/CrashReportWebSite/Models/CrashDataFormater.cs
+++ b/Engine/Source/Programs/CrashReporter/CrashReportWebSite/Models/CrashDataFormater.cs
@@ -350,16 +350,19 @@ namespace Tools.CrashReporter.CrashReportWebSite.Models
 						FuncName += "()";
 
 						// Grab the source file
-						if( OpenBracketOffset > CloseFuncSymbol && CloseBracketOffset > OpenBracketOffset && bWindows )
+						if( OpenBracketOffset > CloseFuncSymbol && CloseBracketOffset > OpenBracketOffset && (bWindows || bLinux) )
 						{
 							string FileLinePath = CurrentLine.Substring( OpenBracketOffset + 1, CloseBracketOffset - OpenBracketOffset - 1 ).Trim();
 
 							FilePath = FileLinePath.TrimEnd( "0123456789:".ToCharArray() );
-							int SourceLine = 0;
-							Debug.Assert( FilePath.Length < FileLinePath.Length, "WRONG SIZE" );
-							if( int.TryParse( FileLinePath.Substring( FilePath.Length + 1 ), out SourceLine ) )
+
+							if (FileLinePath.Length > FilePath.Length + 1)
 							{
-								LineNumber = SourceLine;
+								int SourceLine = 0;
+								if (int.TryParse(FileLinePath.Substring(FilePath.Length + 1), out SourceLine))
+								{
+									LineNumber = SourceLine;
+								}
 							}
 						}
 					}
diff --git a/Engine/Source/Programs/CrashReporter/CrashReportWebSite/Models/CrashRepository.cs b/Engine/Source/Programs/CrashReporter/CrashReportWebSite/Models/CrashRepository.cs
--- a/Engine/Source/Programs/CrashReporter/CrashReportWebSite/Models/CrashRepository.cs
+++ b/Engine/Source/Programs/CrashReporter/CrashReportWebSite/Models/CrashRepository.cs
@@ -725,6 +725,8 @@ namespace Tools.CrashReporter.CrashReportWebSite.Models
 			//NewCrash.HasDiagnosticsFile = true;		
 			//NewCrash.HasMetaData = true;
 
+			NewCrash.UserActivityHint = NewCrashInfo.UserActivityHint;
+
 			// Add the crash to the database
 			Context.Crashes.InsertOnSubmit( NewCrash );
 			SubmitChanges();
diff --git a/Engine/Source/Programs/CrashReporter/CrashReportWebSite/Views/Crashes/Show.aspx b/Engine/Source/Programs/CrashReporter/CrashReportWebSite/Views/Crashes/Show.aspx
--- a/Engine/Source/Programs/CrashReporter/CrashReportWebSite/Views/Crashes/Show.aspx
+++ b/Engine/Source/Programs/CrashReporter/CrashReportWebSite/Views/Crashes/Show.aspx
@@ -223,6 +223,9 @@
 
 			<dt>Module</dt>
 				<dd ><%=Html.DisplayFor(m => Model.Crash.Module) %></dd>
+			
+			<dt>User Activity</dt>
+				<dd ><%=Html.DisplayFor(m => Model.Crash.UserActivityHint) %></dd>
 		</dl>
 	</div>
 
diff --git a/Engine/Source/Programs/DotNETCommon/DotNETUtilities/DotNETUtilities.csproj b/Engine/Source/Programs/DotNETCommon/DotNETUtilities/DotNETUtilities.csproj
--- a/Engine/Source/Programs/DotNETCommon/DotNETUtilities/DotNETUtilities.csproj
+++ b/Engine/Source/Programs/DotNETCommon/DotNETUtilities/DotNETUtilities.csproj
@@ -35,6 +35,15 @@
     <CodeAnalysisRuleSet>MinimumRecommendedRules.ruleset</CodeAnalysisRuleSet>
     <RunCodeAnalysis>false</RunCodeAnalysis>
   </PropertyGroup>
+  <PropertyGroup Condition="'$(Configuration)|$(Platform)' == 'Release|AnyCPU'">
+    <OutputPath>bin\Release\</OutputPath>
+    <DefineConstants>TRACE</DefineConstants>
+    <Optimize>true</Optimize>
+    <DebugType>pdbonly</DebugType>
+    <PlatformTarget>AnyCPU</PlatformTarget>
+    <ErrorReport>prompt</ErrorReport>
+    <CodeAnalysisRuleSet>MinimumRecommendedRules.ruleset</CodeAnalysisRuleSet>
+  </PropertyGroup>
   <ItemGroup>
     <Reference Include="System" />
     <Reference Include="System.Core" />
diff --git a/Engine/Source/Programs/UnrealBuildTool/Configuration/BuildConfiguration.cs b/Engine/Source/Programs/UnrealBuildTool/Configuration/BuildConfiguration.cs
--- a/Engine/Source/Programs/UnrealBuildTool/Configuration/BuildConfiguration.cs
+++ b/Engine/Source/Programs/UnrealBuildTool/Configuration/BuildConfiguration.cs
@@ -448,7 +448,7 @@ namespace UnrealBuildTool
 		public static bool bUseUBTMakefiles;
 
         /// <summary>
-        /// Enables support for fast UHT parsing by caching results of previous UHT runs. If a module or UnrealHeaderTool.manifest
+        /// Enables support for fast UHT parsing by caching results of previous UHT runs. If a module or *.uhtmanifest
         /// gets changed, all modules up to first changed one get loaded from makefile and the rest is parsed regularly.
         /// </summary>
         [XmlConfig]
diff --git a/Engine/Source/Programs/UnrealBuildTool/Configuration/UEBuildModule.cs b/Engine/Source/Programs/UnrealBuildTool/Configuration/UEBuildModule.cs
--- a/Engine/Source/Programs/UnrealBuildTool/Configuration/UEBuildModule.cs
+++ b/Engine/Source/Programs/UnrealBuildTool/Configuration/UEBuildModule.cs
@@ -13,45 +13,71 @@ using Tools.DotNETCommon.CaselessDictionary;
 namespace UnrealBuildTool
 {
 	/// <summary>
-	/// Type of module. Mirrored in UHT as EBuildModuleType
+	/// Type of module. Mirrored in UHT as EBuildModuleType.
+	/// This should be sorted by the order in which we expect modules to be built.
 	/// </summary>
 	public enum UEBuildModuleType
 	{
-		Unknown,
-		Runtime,
-		Developer,
-		Editor,
-		ThirdParty,
 		Program,
-		Game,
+		EngineRuntime,
+		EngineDeveloper,
+		EngineEditor,
+		EngineThirdParty,
+		GameRuntime,
+		GameDeveloper,
+		GameEditor,
+		GameThirdParty,
 	}
 	public static class UEBuildModuleTypeExtensions
 	{
+		public static bool IsProgramModule(this UEBuildModuleType ModuleType)
+		{
+			return ModuleType == UEBuildModuleType.Program;
+		}
 		public static bool IsEngineModule(this UEBuildModuleType ModuleType)
 		{
-			return ModuleType != UEBuildModuleType.Game;
+			return ModuleType == UEBuildModuleType.EngineRuntime || ModuleType == UEBuildModuleType.EngineDeveloper || ModuleType == UEBuildModuleType.EngineEditor || ModuleType == UEBuildModuleType.EngineThirdParty;
 		}
-		public static UEBuildModuleType FromHostType(ModuleHostType ModuleType)
+		public static bool IsGameModule(this UEBuildModuleType ModuleType)
+		{
+			return ModuleType == UEBuildModuleType.GameRuntime || ModuleType == UEBuildModuleType.GameDeveloper || ModuleType == UEBuildModuleType.GameEditor || ModuleType == UEBuildModuleType.GameThirdParty;
+		}
+		public static UEBuildModuleType? EngineModuleTypeFromHostType(ModuleHostType ModuleType)
 		{
 			switch (ModuleType)
 			{
 				case ModuleHostType.Runtime:
 				case ModuleHostType.RuntimeNoCommandlet:
 				case ModuleHostType.RuntimeAndProgram:
-                case ModuleHostType.ServerOnly:
-					return UEBuildModuleType.Runtime;
+				case ModuleHostType.ServerOnly:
+					return UEBuildModuleType.EngineRuntime;
 				case ModuleHostType.Developer:
-					return UEBuildModuleType.Developer;
+					return UEBuildModuleType.EngineDeveloper;
 				case ModuleHostType.Editor:
 				case ModuleHostType.EditorNoCommandlet:
-					return UEBuildModuleType.Editor;
-				case ModuleHostType.Program:
-					return UEBuildModuleType.Program;
+					return UEBuildModuleType.EngineEditor;
 				default:
-					return UEBuildModuleType.Unknown;
+					return null;
+			}
+		}
+		public static UEBuildModuleType? GameModuleTypeFromHostType(ModuleHostType ModuleType)
+		{
+			switch (ModuleType)
+			{
+				case ModuleHostType.Runtime:
+				case ModuleHostType.RuntimeNoCommandlet:
+				case ModuleHostType.RuntimeAndProgram:
+				case ModuleHostType.ServerOnly:
+					return UEBuildModuleType.GameRuntime;
+				case ModuleHostType.Developer:
+					return UEBuildModuleType.GameDeveloper;
+				case ModuleHostType.Editor:
+				case ModuleHostType.EditorNoCommandlet:
+					return UEBuildModuleType.GameEditor;
+				default:
+					return null;
 			}
 		}
-
 	}
 	/// <summary>
 	/// Distribution level of module.
@@ -87,50 +113,71 @@ namespace UnrealBuildTool
 		public static readonly string ProgramsFolder = String.Format("{0}Programs{0}", Path.DirectorySeparatorChar);
 		public static readonly string ThirdPartyFolder = String.Format("{0}ThirdParty{0}", Path.DirectorySeparatorChar);
 
-		public static UEBuildModuleType GetModuleTypeFromDescriptor(ModuleDescriptor Module)
+		public static UEBuildModuleType GetEngineModuleTypeFromDescriptor(ModuleDescriptor Module)
 		{
 			switch (Module.Type)
 			{
 				case ModuleHostType.Developer:
-					return UEBuildModuleType.Developer;
+					return UEBuildModuleType.EngineDeveloper;
 				case ModuleHostType.Editor:
 				case ModuleHostType.EditorNoCommandlet:
-					return UEBuildModuleType.Editor;
-				case ModuleHostType.Program:
-					return UEBuildModuleType.Program;
+					return UEBuildModuleType.EngineEditor;
 				case ModuleHostType.Runtime:
 				case ModuleHostType.RuntimeNoCommandlet:
 				case ModuleHostType.RuntimeAndProgram:
-					return UEBuildModuleType.Runtime;
+					return UEBuildModuleType.EngineRuntime;
 				default:
-					throw new BuildException("Unhandled module type {0}", Module.Type.ToString());
+					throw new BuildException("Unhandled engine module type {0}", Module.Type.ToString());
 			}
 		}
 
-		public static UEBuildModuleType GetEngineModuleTypeBasedOnLocation(string ModuleName, UEBuildModuleType ModuleType, FileReference ModuleFileName)
+		public static UEBuildModuleType GetGameModuleTypeFromDescriptor(ModuleDescriptor Module)
+		{
+			switch (Module.Type)
+			{
+				case ModuleHostType.Developer:
+					return UEBuildModuleType.GameDeveloper;
+				case ModuleHostType.Editor:
+				case ModuleHostType.EditorNoCommandlet:
+					return UEBuildModuleType.GameEditor;
+				case ModuleHostType.Runtime:
+				case ModuleHostType.RuntimeNoCommandlet:
+				case ModuleHostType.RuntimeAndProgram:
+					return UEBuildModuleType.GameRuntime;
+				default:
+					throw new BuildException("Unhandled game module type {0}", Module.Type.ToString());
+			}
+		}
+
+		public static UEBuildModuleType? GetEngineModuleTypeBasedOnLocation(string ModuleName, FileReference ModuleFileName)
 		{
 			string ModuleFileRelativeToEngineDirectory = ModuleFileName.MakeRelativeTo(UnrealBuildTool.EngineDirectory);
 			if (ModuleFileRelativeToEngineDirectory.IndexOf(UEBuildModule.RuntimeFolder, StringComparison.InvariantCultureIgnoreCase) >= 0)
 			{
-				ModuleType = UEBuildModuleType.Runtime;
+				return UEBuildModuleType.EngineRuntime;
 			}
-			else if (ModuleFileRelativeToEngineDirectory.IndexOf(UEBuildModule.DeveloperFolder, StringComparison.InvariantCultureIgnoreCase) >= 0)
+
+			if (ModuleFileRelativeToEngineDirectory.IndexOf(UEBuildModule.DeveloperFolder, StringComparison.InvariantCultureIgnoreCase) >= 0)
 			{
-				ModuleType = UEBuildModuleType.Developer;
+				return UEBuildModuleType.EngineDeveloper;
 			}
-			else if (ModuleFileRelativeToEngineDirectory.IndexOf(UEBuildModule.EditorFolder, StringComparison.InvariantCultureIgnoreCase) >= 0)
+
+			if (ModuleFileRelativeToEngineDirectory.IndexOf(UEBuildModule.EditorFolder, StringComparison.InvariantCultureIgnoreCase) >= 0)
 			{
-				ModuleType = UEBuildModuleType.Editor;
+				return UEBuildModuleType.EngineEditor;
 			}
-			else if (ModuleFileRelativeToEngineDirectory.IndexOf(UEBuildModule.ProgramsFolder, StringComparison.InvariantCultureIgnoreCase) >= 0)
+	
+			if (ModuleFileRelativeToEngineDirectory.IndexOf(UEBuildModule.ProgramsFolder, StringComparison.InvariantCultureIgnoreCase) >= 0)
 			{
-				ModuleType = UEBuildModuleType.Program;
+				return UEBuildModuleType.Program;
 			}
-			else if (ModuleFileRelativeToEngineDirectory.IndexOf(UEBuildModule.ThirdPartyFolder, StringComparison.InvariantCultureIgnoreCase) >= 0)
+
+			if (ModuleFileRelativeToEngineDirectory.IndexOf(UEBuildModule.ThirdPartyFolder, StringComparison.InvariantCultureIgnoreCase) >= 0)
 			{
-				ModuleType = UEBuildModuleType.ThirdParty;
+				return UEBuildModuleType.EngineThirdParty;
 			}
-			return ModuleType;
+
+			return null;
 		}
 
 		/// <summary>
@@ -243,7 +290,7 @@ namespace UnrealBuildTool
 		{
 			return IsRedistributableOverride.HasValue
 				? IsRedistributableOverride.Value
-				: (Type != UEBuildModuleType.Developer && Type != UEBuildModuleType.Editor);
+				: (Type != UEBuildModuleType.EngineDeveloper && Type != UEBuildModuleType.EngineEditor);
 		}
 
 		/// <summary>
@@ -313,6 +360,15 @@ namespace UnrealBuildTool
 		/// </summary>
 		public List<RuntimeDependency> RuntimeDependencies;
 
+		/// <summary>
+		/// Returns a list of this module's immediate dependencies.
+		/// </summary>
+		/// <returns>An enumerable containing the dependencies of the module.</returns>
+		public IEnumerable<UEBuildModule> GetAllDependencyModules()
+		{
+			return PublicDependencyModules.Concat(PrivateDependencyModules).Concat(DynamicallyLoadedModules).Concat(PlatformSpecificDynamicallyLoadedModules);
+		}
+
 		public UEBuildModule(
 			UEBuildTarget InTarget,
 			string InName,
@@ -515,7 +571,7 @@ namespace UnrealBuildTool
 		{
 			HashSet<UEBuildModule> VisitedModules = new HashSet<UEBuildModule>();
 
-			if (this.Type == UEBuildModuleType.Game)
+			if (this.Type.IsGameModule())
 			{
 				Definitions.Add("DEPRECATED_FORGAME=DEPRECATED");
 			}
diff --git a/Engine/Source/Programs/UnrealBuildTool/Configuration/UEBuildTarget.cs b/Engine/Source/Programs/UnrealBuildTool/Configuration/UEBuildTarget.cs
--- a/Engine/Source/Programs/UnrealBuildTool/Configuration/UEBuildTarget.cs
+++ b/Engine/Source/Programs/UnrealBuildTool/Configuration/UEBuildTarget.cs
@@ -1991,6 +1991,130 @@ namespace UnrealBuildTool
 			}
 		}
 
+		/// <summary>
+		/// Holds a cache of dependency tests between nodes.
+		/// </summary>
+		/// <typeparam name="NodeType">The type of node to be tested for dependencies.</typeparam>
+		class NodeDependencyCache<NodeType>
+		{
+			public NodeDependencyCache(Func<NodeType, IEnumerable<NodeType>> InNodeDependencies)
+			{
+				NodeDependencies = InNodeDependencies;
+				PreviousResults  = new Dictionary<NodeType, Dictionary<NodeType, bool>>();
+			}
+
+			/// <summary>
+			/// Tests if NodeA is dependent on NodeB.
+			/// </summary>
+			/// <param name="NodeA">The dependent to test.</param>
+			/// <param name="NodeB">The dependency to test.</param>
+			/// <returns>true if NodeA is dependent on NodeB, false otherwise.</returns>
+			public bool DependsOn(NodeType NodeA, NodeType NodeB)
+			{
+				// Return any previous result, if there is one.
+				bool Result;
+				Dictionary<NodeType, bool> InnerDictionary;
+				if (PreviousResults.TryGetValue(NodeA, out InnerDictionary))
+				{
+					if (InnerDictionary.TryGetValue(NodeB, out Result))
+					{
+						return Result;
+					}
+				}
+				else
+				{
+					InnerDictionary = new Dictionary<NodeType, bool>();
+					PreviousResults.Add(NodeA, InnerDictionary);
+				}
+
+				Result = DependsOnRecursive(NodeA, NodeB, new HashSet<NodeType>(), InnerDictionary);
+				InnerDictionary.Add(NodeB, Result);
+				return Result;
+			}
+
+			private bool DependsOnRecursive(NodeType NodeA, NodeType NodeB, HashSet<NodeType> Visited, Dictionary<NodeType, bool> InnerDictionary)
+			{
+				if (Visited.Contains(NodeA))
+				{
+					// Nodes are not dependent on themselves
+					return false;
+				}
+
+				Visited.Add(NodeA);
+
+				IEnumerable<NodeType> Deps = NodeDependencies(NodeA);
+
+				if (Deps.Contains(NodeB))
+				{
+					// NodeB is an immediate dependent of NodeA
+					return true;
+				}
+
+				foreach (var i in Deps)
+				{
+					bool Result;
+					if (InnerDictionary.TryGetValue(NodeB, out Result) && Result)
+					{
+						// We've calculated this result before
+						return true;
+					}
+
+					if (DependsOnRecursive(i, NodeB, Visited, InnerDictionary))
+					{
+						return true;
+					}
+				}
+
+				// Didn't find any dependency
+				return false;
+			}
+
+			private Func<NodeType, IEnumerable<NodeType>>            NodeDependencies;
+			private Dictionary<NodeType, Dictionary<NodeType, bool>> PreviousResults;
+		}
+
+		/// <summary>
+		/// Returns a copy of Nodes sorted by dependency.  Independent or circularly-dependent nodes should
+		/// remain in their same relative order within the original Nodes sequence.
+		/// </summary>
+		/// <typeparam name="NodeType">The type of node to sort.</typeparam>
+		/// <param name="Nodes">The sequence of nodes to sort.</param>
+		/// <param name="NodeDependencies">A function which returns a node's immediate dependencies.</param>
+		/// <returns>A copy of Nodes sorted by NodeDependencies, with circular references remaining in the same relative order.</returns>
+		public static IEnumerable<NodeType> StableTopologicalSort<NodeType>(IEnumerable<NodeType> Nodes, Func<NodeType, IEnumerable<NodeType>> NodeDependencies)
+		{
+			List<NodeType> NodeList  = Nodes.ToList();
+			int            NodeCount = NodeList.Count;
+
+			var Cache = new NodeDependencyCache<NodeType>(NodeDependencies);
+			for (int Index1 = 0; Index1 != NodeCount; ++Index1)
+			{
+				var Node1 = NodeList[Index1];
+
+				for (int Index2 = 0; Index2 != Index1; ++Index2)
+				{
+					var Node2 = NodeList[Index2];
+
+					if (Cache.DependsOn(Node2, Node1) && !Cache.DependsOn(Node1, Node2))
+					{
+						// Rotate element at Index1 into position at Index2
+						for (int Index3 = Index1; Index3 != Index2; )
+						{
+							--Index3;
+							NodeList[Index3 + 1] = NodeList[Index3];
+						}
+						NodeList[Index2] = Node1;
+
+						// Break out of this loop, because this iteration must have covered all existing cases
+						// involving the node formerly at position Index1
+						break;
+					}
+				}
+			}
+
+			return NodeList;
+		}
+
 		/// <summary>
 		/// Gathers dependency modules for given binaries list.
 		/// </summary>
@@ -2012,7 +2136,10 @@ namespace UnrealBuildTool
 				}
 			}
 
-			return Output;
+			// Sort modules by type, then by dependency
+			var ModulesSortedByType        = Output.OrderBy(c => c.Type).ToList();
+			var ModulesTopologicallySorted = StableTopologicalSort(ModulesSortedByType, x => x.GetAllDependencyModules().OfType<UEBuildModuleCPP>());
+			return new HashSet<UEBuildModuleCPP>(ModulesTopologicallySorted.AsEnumerable());
 		}
 
 		/// <summary>
@@ -2189,7 +2316,7 @@ namespace UnrealBuildTool
 				if (UObjectModules.Count > 0)
 				{
 					// Execute the header tool
-					FileReference ModuleInfoFileName = FileReference.Combine(ProjectIntermediateDirectory, "UnrealHeaderTool.manifest");
+					FileReference ModuleInfoFileName = FileReference.Combine(ProjectIntermediateDirectory, GetTargetName() + ".uhtmanifest");
 					ECompilationResult UHTResult = ECompilationResult.OtherCompilationError;
 					if (!ExternalExecution.ExecuteHeaderToolIfNecessary(TargetToolChain, this, GlobalCompileEnvironment, UObjectModules, ModuleInfoFileName, ref UHTResult))
 					{
@@ -2728,7 +2855,7 @@ namespace UnrealBuildTool
 			return new UEBuildModuleCPP(
 				InTarget: this,
 				InName: Name,
-				InType: UEBuildModuleType.Game,
+				InType: UEBuildModuleType.GameRuntime,
 				InModuleDirectory: Directory,
 				InGeneratedCodeDirectory: null,
 				InIntelliSenseGatherer: null,
@@ -3725,67 +3852,81 @@ namespace UnrealBuildTool
 				ModuleRules RulesObject = CreateModuleRulesAndSetDefaults(ModuleName, out ModuleFileName);
 				DirectoryReference ModuleDirectory = ModuleFileName.Directory;
 
+				// Get the type of module we're creating
+				UEBuildModuleType? ModuleType = null;
+
 				// Get the plugin for this module
 				PluginInfo Plugin;
 				RulesAssembly.TryGetPluginForModule(ModuleFileName, out Plugin);
 
-				// Get the type of module we're creating
-				var ModuleType = UEBuildModuleType.Unknown;
-
-				// see if it's external
-				if (RulesObject.Type == ModuleRules.ModuleType.External)
+				// Get the module descriptor for this module if it's a plugin
+				ModuleDescriptor PluginModuleDesc = null;
+				if (Plugin != null)
 				{
-					ModuleType = UEBuildModuleType.ThirdParty;
+					PluginModuleDesc = Plugin.Descriptor.Modules.FirstOrDefault(x => x.Name == ModuleName);
+					if (PluginModuleDesc != null && PluginModuleDesc.Type == ModuleHostType.Program)
+					{
+						ModuleType = UEBuildModuleType.Program;
+					}
 				}
-				else
+
+				if (ModuleFileName.IsUnderDirectory(UnrealBuildTool.EngineDirectory))
 				{
-					// Check if it's a plugin
-					if (Plugin != null)
+					if (RulesObject.Type == ModuleRules.ModuleType.External)
 					{
-						ModuleDescriptor Descriptor = Plugin.Descriptor.Modules.FirstOrDefault(x => x.Name == ModuleName);
-						if (Descriptor != null)
+						ModuleType = UEBuildModuleType.EngineThirdParty;
+					}
+					else
+					{
+						if (!ModuleType.HasValue && PluginModuleDesc != null)
+						{
+							ModuleType = UEBuildModule.GetEngineModuleTypeFromDescriptor(PluginModuleDesc);
+						}
+
+						if (!ModuleType.HasValue)
 						{
-							ModuleType = UEBuildModule.GetModuleTypeFromDescriptor(Descriptor);
+							ModuleType = UEBuildModule.GetEngineModuleTypeBasedOnLocation(ModuleName, ModuleFileName);
 						}
 					}
-					if (ModuleType == UEBuildModuleType.Unknown)
+				}
+				else
+				{
+					if (RulesObject.Type == ModuleRules.ModuleType.External)
 					{
-						// not a plugin, see if it is a game module 
-						if (!ModuleFileName.IsUnderDirectory(UnrealBuildTool.EngineSourceDirectory))
+						ModuleType = UEBuildModuleType.GameThirdParty;
+					}
+					else
+					{
+						if (!ModuleType.HasValue && PluginModuleDesc != null)
+						{
+							ModuleType = UEBuildModule.GetGameModuleTypeFromDescriptor(PluginModuleDesc);
+						}
+
+						if (!ModuleType.HasValue)
 						{
 							if (ProjectDescriptor != null && ProjectDescriptor.Modules != null)
 							{
-								foreach (ModuleDescriptor ProjectModule in ProjectDescriptor.Modules)
+								ModuleDescriptor ProjectModule = ProjectDescriptor.Modules.FirstOrDefault(x => x.Name == ModuleName);
+								if (ProjectModule != null)
 								{
-									if (ProjectModule.Name == ModuleName)
-									{
-										ModuleType = UEBuildModuleTypeExtensions.FromHostType(ProjectModule.Type);
-										if (ModuleType == UEBuildModuleType.Runtime)
-										{
-											// We expect game runtime modules to be marked as Game module type.
-											ModuleType = UEBuildModuleType.Game;
-										}
-									}
+									ModuleType = UEBuildModuleTypeExtensions.GameModuleTypeFromHostType(ProjectModule.Type);
+								}
+								else
+								{
+									// No descriptor file or module was not on the list
+									ModuleType = UEBuildModuleType.GameRuntime;
 								}
-							}
-							if (ModuleType == UEBuildModuleType.Unknown)
-							{
-								// No descriptor file or module was not on the list
-								ModuleType = UEBuildModuleType.Game;
-							}						
-						}
-						else
-						{
-							ModuleType = UEBuildModule.GetEngineModuleTypeBasedOnLocation(ModuleName, ModuleType, ModuleFileName);
-							if (ModuleType == UEBuildModuleType.Unknown)
-							{
-								throw new BuildException("Unable to determine module type for {0}", ModuleFileName);
 							}
 						}
 					}
 				}
 
-				var IsGameModule = !ModuleType.IsEngineModule();
+				if (!ModuleType.HasValue)
+				{
+					throw new BuildException("Unable to determine module type for {0}", ModuleFileName);
+				}
+
+				bool bIsGameModuleOrProgram = ModuleType.Value.IsGameModule() || ModuleType.Value.IsProgramModule();
 
 				// Get the base directory for paths referenced by the module. If the module's under the UProject source directory use that, otherwise leave it relative to the Engine source directory.
 				if (ProjectFile != null)
@@ -3824,7 +3965,7 @@ namespace UnrealBuildTool
 				if (RulesObject.Type != ModuleRules.ModuleType.External && ModuleName != "Core")
 				{
 					// Add the default include paths to the module rules, if they exist.
-					AddDefaultIncludePathsToModuleRules(ModuleFileName, IsGameModule, Plugin, RulesObject);
+					AddDefaultIncludePathsToModuleRules(ModuleFileName, bIsGameModuleOrProgram, Plugin, RulesObject);
 
 					// Add the path to the generated headers 
 					if (GeneratedCodeDirectory != null)
@@ -3885,7 +4026,7 @@ namespace UnrealBuildTool
 				UEBuildPlatform.PlatformModifyHostModuleRules(ModuleName, RulesObject, TargetInfo, Only);
 
 				// Now, go ahead and create the module builder instance
-				Module = InstantiateModule(RulesObject, ModuleName, ModuleType, ModuleDirectory, GeneratedCodeDirectory, IntelliSenseGatherer, FoundSourceFiles, bBuildFiles, ModuleFileName);
+				Module = InstantiateModule(RulesObject, ModuleName, ModuleType.Value, ModuleDirectory, GeneratedCodeDirectory, IntelliSenseGatherer, FoundSourceFiles, bBuildFiles, ModuleFileName);
 			}
 			return Module;
 		}
diff --git a/Engine/Source/Programs/UnrealBuildTool/System/ExternalExecution.cs b/Engine/Source/Programs/UnrealBuildTool/System/ExternalExecution.cs
--- a/Engine/Source/Programs/UnrealBuildTool/System/ExternalExecution.cs
+++ b/Engine/Source/Programs/UnrealBuildTool/System/ExternalExecution.cs
@@ -773,7 +773,7 @@ namespace UnrealBuildTool
 					System.IO.File.WriteAllText(ModuleInfoFileName.FullName, fastJSON.JSON.Instance.ToJSON(Manifest, new fastJSON.JSONParameters { UseExtensions = false }));
 
 					string CmdLine = (Target.ProjectFile != null) ? "\"" + Target.ProjectFile.FullName + "\"" : Target.GetTargetName();
-					CmdLine += " \"" + ModuleInfoFileName + "\" -LogCmds=\"loginit warning, logexit warning, logdatabase error\" -Unattended";
+					CmdLine += " \"" + ModuleInfoFileName + "\" -LogCmds=\"loginit warning, logexit warning, logdatabase error\" -Unattended -WarningsAsErrors";
 					if (UnrealBuildTool.IsEngineInstalled())
 					{
 						CmdLine += " -installed";
diff --git a/Engine/Source/Programs/UnrealBuildTool/System/UnrealBuildTool.cs b/Engine/Source/Programs/UnrealBuildTool/System/UnrealBuildTool.cs
--- a/Engine/Source/Programs/UnrealBuildTool/System/UnrealBuildTool.cs
+++ b/Engine/Source/Programs/UnrealBuildTool/System/UnrealBuildTool.cs
@@ -1958,7 +1958,7 @@ namespace UnrealBuildTool
 									if (TargetUObjectModules.Count > 0)
 									{
 										// Execute the header tool
-										FileReference ModuleInfoFileName = FileReference.Combine(Target.ProjectIntermediateDirectory, "UnrealHeaderTool.manifest");
+										FileReference ModuleInfoFileName = FileReference.Combine(Target.ProjectIntermediateDirectory, Target.GetTargetName() + ".uhtmanifest");
 										ECompilationResult UHTResult = ECompilationResult.OtherCompilationError;
 										if (!ExternalExecution.ExecuteHeaderToolIfNecessary(ToolChain, Target, GlobalCompileEnvironment: null, UObjectModules: TargetUObjectModules, ModuleInfoFileName: ModuleInfoFileName, UHTResult: ref UHTResult))
 										{
diff --git a/Engine/Source/Programs/UnrealHeaderTool/Private/CodeGenerator.cpp b/Engine/Source/Programs/UnrealHeaderTool/Private/CodeGenerator.cpp
--- a/Engine/Source/Programs/UnrealHeaderTool/Private/CodeGenerator.cpp
+++ b/Engine/Source/Programs/UnrealHeaderTool/Private/CodeGenerator.cpp
@@ -5138,7 +5138,7 @@ bool FNativeClassHeaderGenerator::SaveHeaderIfChanged(const TCHAR* HeaderPath, c
 		IFileManager::Get().Delete( *TmpHeaderFilename, false, true );
 		if ( !FFileHelper::SaveStringToFile(NewHeaderContents, *TmpHeaderFilename) )
 		{
-			UE_LOG(LogCompile, Warning, TEXT("Failed to save header export preview: '%s'"), *TmpHeaderFilename);
+			UE_LOG_WARNING_UHT(TEXT("Failed to save header export preview: '%s'"), *TmpHeaderFilename);
 		}
 
 		TempHeaderPaths.Add(TmpHeaderFilename);
@@ -5539,12 +5539,12 @@ ECompilationResult::Type PreparseModules(FUHTMakefile& UHTMakefile, const FStrin
 		//       want to make sure our flags get set
 		Package->SetPackageFlags(PKG_ContainsScript | PKG_Compiling);
 		Package->ClearPackageFlags(PKG_ClientOptional | PKG_ServerSideOnly);
-		if (Module.ModuleType == EBuildModuleType::Editor)
+		if (Module.ModuleType == EBuildModuleType::GameEditor || Module.ModuleType == EBuildModuleType::EngineEditor)
 		{
 			Package->SetPackageFlags(PKG_EditorOnly);
 		}
 
-		if (Module.ModuleType == EBuildModuleType::Developer)
+		if (Module.ModuleType == EBuildModuleType::GameDeveloper || Module.ModuleType == EBuildModuleType::EngineDeveloper)
 		{
 			Package->SetPackageFlags(Package->GetPackageFlags() | PKG_Developer);
 		}
@@ -5715,11 +5715,6 @@ ECompilationResult::Type UnrealHeaderTool_Main(const FString& ModuleInfoFilename
 	check(GIsUCCMakeStandaloneHeaderGenerator);
 	ECompilationResult::Type Result = ECompilationResult::Succeeded;
 	
-	if ( !FParse::Param( FCommandLine::Get(), TEXT("IgnoreWarnings")) )
-	{
-		GWarn->TreatWarningsAsErrors = true;
-	}
-
 	FString ModuleInfoPath = FPaths::GetPath(ModuleInfoFilename);
 
 	// Load the manifest file, giving a list of all modules to be processed, pre-sorted by dependency ordering
@@ -5842,8 +5837,13 @@ ECompilationResult::Type UnrealHeaderTool_Main(const FString& ModuleInfoFilename
 						bool bUseVTableConstructors;
 					} UseVTableConstructorsCache;
 
-					Result = FHeaderParser::ParseAllHeadersInside(AllClasses, GWarn, Package, Module, ScriptPlugins,
-						Module.ModuleType != EBuildModuleType::Game || UseVTableConstructorsCache.bUseVTableConstructors, UHTMakefile);
+					bool bModuleIsGame =
+						   Module.ModuleType == EBuildModuleType::GameDeveloper
+						|| Module.ModuleType == EBuildModuleType::GameEditor
+						|| Module.ModuleType == EBuildModuleType::GameRuntime
+						|| Module.ModuleType == EBuildModuleType::GameThirdParty;
+
+					Result = FHeaderParser::ParseAllHeadersInside(AllClasses, GWarn, Package, Module, ScriptPlugins, bModuleIsGame || UseVTableConstructorsCache.bUseVTableConstructors, UHTMakefile);
 #else // WITH_HOT_RELOAD_CTORS
 					Result = FHeaderParser::ParseAllHeadersInside(AllClasses, GWarn, Package, Module, ScriptPlugins, UHTMakefile);
 #endif // WITH_HOT_RELOAD_CTORS
@@ -5971,7 +5971,7 @@ UClass* ProcessParsedClass(bool bClassIsAnInterface, TArray<FHeaderProvider> &De
 			UClass* ConflictingClass = FindObject<UClass>(ANY_PACKAGE, *ClassNameStripped, true);
 			if (ConflictingClass != nullptr)
 			{
-				UE_LOG(LogCompile, Warning, TEXT("Duplicate class name: %s also exists in file %s"), *ClassName, *ConflictingClass->GetOutermost()->GetName());
+				UE_LOG_WARNING_UHT(TEXT("Duplicate class name: %s also exists in file %s"), *ClassName, *ConflictingClass->GetOutermost()->GetName());
 			}
 		}
 
diff --git a/Engine/Source/Programs/UnrealHeaderTool/Private/HeaderParser.cpp b/Engine/Source/Programs/UnrealHeaderTool/Private/HeaderParser.cpp
--- a/Engine/Source/Programs/UnrealHeaderTool/Private/HeaderParser.cpp
+++ b/Engine/Source/Programs/UnrealHeaderTool/Private/HeaderParser.cpp
@@ -849,7 +849,7 @@ namespace
 
 			if (NewKey != NAME_None)
 			{
-				UE_LOG(LogCompile, Warning, TEXT("Remapping old metadata key '%s' to new key '%s', please update the declaration."), *CurrentKey.ToString(), *NewKey.ToString());
+				UE_LOG_WARNING_UHT(TEXT("Remapping old metadata key '%s' to new key '%s', please update the declaration."), *CurrentKey.ToString(), *NewKey.ToString());
 				CurrentKey = NewKey;
 			}
 
@@ -1766,7 +1766,7 @@ UScriptStruct* FHeaderParser::CompileStructDeclaration(FClasses& AllClasses)
 
 			case EStructSpecifier::NoExport:
 			{
-				//UE_LOG(LogCompile, Warning, TEXT("Struct named %s in %s is still marked noexport"), *EffectiveStructName, *(Class->GetName()));//@TODO: UCREMOVAL: Debug printing
+				//UE_LOG_WARNING_UHT(TEXT("Struct named %s in %s is still marked noexport"), *EffectiveStructName, *(Class->GetName()));//@TODO: UCREMOVAL: Debug printing
 				StructFlags &= ~STRUCT_Native;
 				StructFlags |= STRUCT_NoExport;
 			}
@@ -2925,7 +2925,7 @@ FIndexRange*                    ParsedVarIndexRange
 
 				case EVariableSpecifier::NonPIETransient:
 				{
-					UE_LOG(LogCompile, Warning, TEXT("NonPIETransient is deprecated - NonPIEDuplicateTransient should be used instead"));
+					UE_LOG_WARNING_UHT(TEXT("NonPIETransient is deprecated - NonPIEDuplicateTransient should be used instead"));
 					Flags |= CPF_NonPIEDuplicateTransient;
 				}
 				break;
@@ -3112,11 +3112,11 @@ FIndexRange*                    ParsedVarIndexRange
 		{
 			if (0 != (CPF_DisableEditOnInstance & Flags))
 			{
-				UE_LOG(LogCompile, Warning, TEXT("Property cannot have 'DisableEditOnInstance' or 'BlueprintReadOnly' and 'ExposeOnSpawn' flags"));
+				UE_LOG_WARNING_UHT(TEXT("Property cannot have 'DisableEditOnInstance' or 'BlueprintReadOnly' and 'ExposeOnSpawn' flags"));
 			}
 			if (0 == (CPF_BlueprintVisible & Flags))
 			{
-				UE_LOG(LogCompile, Warning, TEXT("Property cannot have 'ExposeOnSpawn' with 'BlueprintVisible' flag."));
+				UE_LOG_WARNING_UHT(TEXT("Property cannot have 'ExposeOnSpawn' with 'BlueprintVisible' flag."));
 			}
 			Flags |= CPF_ExposeOnSpawn;
 		}
@@ -4059,7 +4059,7 @@ UProperty* FHeaderParser::GetVarNameAndDim
 			// Warn if a deprecated property is visible
 			if (VarProperty.PropertyFlags & (CPF_Edit | CPF_EditConst | CPF_BlueprintVisible | CPF_BlueprintReadOnly) && !(VarProperty.ImpliedPropertyFlags & CPF_BlueprintReadOnly))
 			{
-				UE_LOG(LogCompile, Warning, TEXT("%s: Deprecated property '%s' should not be marked as visible or editable"), HintText, *VarName);
+				UE_LOG_WARNING_UHT(TEXT("%s: Deprecated property '%s' should not be marked as visible or editable"), HintText, *VarName);
 			}
 
 			VarProperty.PropertyFlags |= CPF_Deprecated;
@@ -4549,16 +4549,16 @@ bool FHeaderParser::CompileDeclaration(FClasses& AllClasses, FToken& Token)
 			PopNest(TopNest->NestType, TEXT("'Interface'"));
 			PostPopNestInterface(AllClasses, CurrentClass);
 
-			// Ensure the UINTERFACE classes have a GENERATED_UINTERFACE_BODY declaration
+			// Ensure the UINTERFACE classes have a GENERATED_BODY declaration
 			if (bHaveSeenUClass && !bClassHasGeneratedUInterfaceBody)
 			{
-				FError::Throwf(TEXT("Expected a GENERATED_UINTERFACE_BODY() at the start of class"));
+				FError::Throwf(TEXT("Expected a GENERATED_BODY() at the start of class"));
 			}
 
-			// Ensure the non-UINTERFACE interface classes have a GENERATED_IINTERFACE_BODY declaration
+			// Ensure the non-UINTERFACE interface classes have a GENERATED_BODY declaration
 			if (!bHaveSeenUClass && !bClassHasGeneratedIInterfaceBody)
 			{
-				FError::Throwf(TEXT("Expected a GENERATED_IINTERFACE_BODY() at the start of class"));
+				FError::Throwf(TEXT("Expected a GENERATED_BODY() at the start of class"));
 			}
 		}
 		else
@@ -4566,10 +4566,10 @@ bool FHeaderParser::CompileDeclaration(FClasses& AllClasses, FToken& Token)
 			PopNest(NEST_Class, TEXT("'Class'"));
 			PostPopNestClass(CurrentClass);
 
-			// Ensure classes have a GENERATED_UCLASS_BODY declaration
+			// Ensure classes have a GENERATED_BODY declaration
 			if (bHaveSeenUClass && !bClassHasGeneratedBody)
 			{
-				FError::Throwf(TEXT("Expected a GENERATED_UCLASS_BODY() at the start of class"));
+				FError::Throwf(TEXT("Expected a GENERATED_BODY() at the start of class"));
 			}
 		}
 
@@ -5758,7 +5758,7 @@ void FHeaderParser::CompileFunctionDeclaration(FClasses& AllClasses)
 
 			else
 			{
-				UE_LOG(LogCompile, Warning, TEXT("BlueprintImplementableEvents should not be virtual. Use BlueprintNativeEvent instead."));
+				UE_LOG_WARNING_UHT(TEXT("BlueprintImplementableEvents should not be virtual. Use BlueprintNativeEvent instead."));
 			}
 		}
 	}
@@ -6584,7 +6584,7 @@ void FHeaderParser::FinalizeScriptExposedFunctions(UClass* Class)
 		}
 		else
 		{
-			UE_LOG(LogCompile, Warning, TEXT("Unknown and unexpected child named %s of type %s in %s\n"), *ChildStruct->GetName(), *ChildStruct->GetClass()->GetName(), *Class->GetName());
+			UE_LOG_WARNING_UHT(TEXT("Unknown and unexpected child named %s of type %s in %s\n"), *ChildStruct->GetName(), *ChildStruct->GetClass()->GetName(), *Class->GetName());
 			check(false);
 		}
 	}
@@ -6948,7 +6948,7 @@ void FHeaderParser::ExportNativeHeaders(
 			}
 			else
 			{
-				UE_LOG(LogCompile, Warning, TEXT("Exporting native class declarations for %s"), *CurrentPackage->GetName());
+				UE_LOG_WARNING_UHT(TEXT("Exporting native class declarations for %s"), *CurrentPackage->GetName());
 			}
 		}
 		else
@@ -6959,7 +6959,7 @@ void FHeaderParser::ExportNativeHeaders(
 			}
 			else
 			{
-				UE_LOG(LogCompile, Warning, TEXT("Exporting native class declarations"));
+				UE_LOG_WARNING_UHT(TEXT("Exporting native class declarations"));
 			}
 		}
 		UHTMakefile.StartExporting();
diff --git a/Engine/Source/Programs/UnrealHeaderTool/Private/IScriptGeneratorPluginInterface.cpp b/Engine/Source/Programs/UnrealHeaderTool/Private/IScriptGeneratorPluginInterface.cpp
new file mode 100644
--- /dev/null
+++ b/Engine/Source/Programs/UnrealHeaderTool/Private/IScriptGeneratorPluginInterface.cpp
@@ -0,0 +1,40 @@
+// Copyright 1998-2015 Epic Games, Inc. All Rights Reserved.
+
+#include "UnrealHeaderTool.h"
+#include "IScriptGeneratorPluginInterface.h"
+#include "Algo/FindSortedStringCaseInsensitive.h"
+
+EBuildModuleType::Type EBuildModuleType::Parse(const TCHAR* Value)
+{
+	static const TCHAR* AlphabetizedTypes[] = {
+		TEXT("EngineDeveloper"),
+		TEXT("EngineEditor"),
+		TEXT("EngineRuntime"),
+		TEXT("EngineThirdParty"),
+		TEXT("GameDeveloper"),
+		TEXT("GameEditor"),
+		TEXT("GameRuntime"),
+		TEXT("GameThirdParty"),
+		TEXT("Program")
+	};
+
+	int32 TypeIndex = Algo::FindSortedStringCaseInsensitive(Value, AlphabetizedTypes);
+	if (TypeIndex < 0)
+	{
+		FError::Throwf(TEXT("Unrecognized EBuildModuleType name: %s"), Value);
+	}
+
+	static EBuildModuleType::Type AlphabetizedValues[] = {
+		EngineDeveloper,
+		EngineEditor,
+		EngineRuntime,
+		EngineThirdParty,
+		GameDeveloper,
+		GameEditor,
+		GameRuntime,
+		GameThirdParty,
+		Program
+	};
+
+	return AlphabetizedValues[TypeIndex];
+}
diff --git a/Engine/Source/Programs/UnrealHeaderTool/Private/Manifest.h b/Engine/Source/Programs/UnrealHeaderTool/Private/Manifest.h
--- a/Engine/Source/Programs/UnrealHeaderTool/Private/Manifest.h
+++ b/Engine/Source/Programs/UnrealHeaderTool/Private/Manifest.h
@@ -97,7 +97,7 @@ struct FManifest
 	TArray<FManifestModule> Modules;
 
 	/**
-	 * Loads an UnrealHeaderTool.manifest from the specified filename.
+	 * Loads a *.uhtmanifest from the specified filename.
 	 *
 	 * @param Filename The filename of the manifest to load.
 	 * @return The loaded module info.
diff --git a/Engine/Source/Programs/UnrealHeaderTool/Private/UnrealHeaderToolMain.cpp b/Engine/Source/Programs/UnrealHeaderTool/Private/UnrealHeaderToolMain.cpp
--- a/Engine/Source/Programs/UnrealHeaderTool/Private/UnrealHeaderToolMain.cpp
+++ b/Engine/Source/Programs/UnrealHeaderTool/Private/UnrealHeaderToolMain.cpp
@@ -9,6 +9,8 @@ IMPLEMENT_APPLICATION(UnrealHeaderTool, "UnrealHeaderTool");
 
 DEFINE_LOG_CATEGORY(LogCompile);
 
+bool GUHTWarningLogged = false;
+
 /**
  * Application entry point
  *
@@ -85,6 +87,11 @@ INT32_MAIN_INT32_ARGC_TCHAR_ARGV()
 	FEngineLoop::AppPreExit();
 	FEngineLoop::AppExit();
 
+	if (Result == ECompilationResult::Succeeded && GUHTWarningLogged && GWarn->TreatWarningsAsErrors)
+	{
+		Result = ECompilationResult::OtherCompilationError;
+	}
+
 	return Result;
 }
 
diff --git a/Engine/Source/Programs/UnrealHeaderTool/Public/IScriptGeneratorPluginInterface.h b/Engine/Source/Programs/UnrealHeaderTool/Public/IScriptGeneratorPluginInterface.h
--- a/Engine/Source/Programs/UnrealHeaderTool/Public/IScriptGeneratorPluginInterface.h
+++ b/Engine/Source/Programs/UnrealHeaderTool/Public/IScriptGeneratorPluginInterface.h
@@ -3,18 +3,20 @@
 
 #include "ModuleManager.h"
 
-/** Build module type, mirrored in UEBuildModule.cs, enum UEBUildModuletype */
+/** Build module type, mirrored in UEBuildModule.cs, enum UEBuildModuletype */
 struct EBuildModuleType
 {
 	enum Type
 	{
-		Unknown,
-		Runtime,
-		Developer,
-		Editor,
-		ThirdParty,
 		Program,
-		Game,
+		EngineRuntime,
+		EngineDeveloper,
+		EngineEditor,
+		EngineThirdParty,
+		GameRuntime,
+		GameDeveloper,
+		GameEditor,
+		GameThirdParty,
 		// NOTE: If you add a new value, make sure to update the ToString() method below!
 
 		Max
@@ -42,42 +44,7 @@ struct EBuildModuleType
 	*
 	* @return	The enum value corresponding to the name
 	*/
-	inline static const EBuildModuleType::Type Parse(const TCHAR* Value)
-	{
-		if (FCString::Stricmp(Value, TEXT("Unknown")) == 0)
-		{
-			return Unknown;
-		}
-		else if (FCString::Stricmp(Value, TEXT("Runtime")) == 0)
-		{
-			return Runtime;
-		}
-		else if (FCString::Stricmp(Value, TEXT("Developer")) == 0)
-		{
-			return Developer;
-		}
-		else if (FCString::Stricmp(Value, TEXT("Editor")) == 0)
-		{
-			return Editor;
-		}
-		else if (FCString::Stricmp(Value, TEXT("ThirdParty")) == 0)
-		{
-			return ThirdParty;
-		}
-		else if (FCString::Stricmp(Value, TEXT("Program")) == 0)
-		{
-			return Program;
-		}
-		else if (FCString::Stricmp(Value, TEXT("Game")) == 0)
-		{
-			return Game;
-		}
-		else
-		{
-			FError::Throwf(TEXT("Unrecognized EBuildModuleType name: %s"), Value);
-			return Unknown;
-		}
-	}
+	static EBuildModuleType::Type Parse(const TCHAR* Value);
 };
 
 /**
diff --git a/Engine/Source/Programs/UnrealHeaderTool/Public/UnrealHeaderTool.h b/Engine/Source/Programs/UnrealHeaderTool/Public/UnrealHeaderTool.h
--- a/Engine/Source/Programs/UnrealHeaderTool/Public/UnrealHeaderTool.h
+++ b/Engine/Source/Programs/UnrealHeaderTool/Public/UnrealHeaderTool.h
@@ -11,4 +11,8 @@
 
 DECLARE_LOG_CATEGORY_EXTERN(LogCompile, Log, All);
 
+extern bool GUHTWarningLogged;
+
+#define UE_LOG_WARNING_UHT(Format, ...) { GUHTWarningLogged = true; UE_LOG(LogCompile, Warning, Format, ##__VA_ARGS__); }
+
 #endif		// __UnrealHeaderTool_h__
\ No newline at end of file
diff --git a/Engine/Source/Programs/UnrealHeaderTool/Resources/UHTDebugging.manifest b/Engine/Source/Programs/UnrealHeaderTool/Resources/UHTDebugging.manifest
--- a/Engine/Source/Programs/UnrealHeaderTool/Resources/UHTDebugging.manifest
+++ b/Engine/Source/Programs/UnrealHeaderTool/Resources/UHTDebugging.manifest
@@ -5,7 +5,7 @@
 	"TargetName": "UHTDebugging",
 	"Modules":[{
 			"Name": "CoreUObject",
-			"ModuleType": "Runtime",
+			"ModuleType": "EngineRuntime",
 			"BaseDirectory": "..\\..\\..\\Runtime\\CoreUObject",
 			"IncludeBase": "..\\..\\..\\Runtime\\CoreUObject",
 			"OutputDirectory": "..\\..\\..\\..\\Intermediate/Build/Inc\\CoreUObject\\",
@@ -22,7 +22,7 @@
 			"UHTGeneratedCodeVersion":"None"
 		}, {
 			"Name": "UHTDebugging",
-			"ModuleType": "Game",
+			"ModuleType": "GameRuntime",
 			"BaseDirectory": ".\\UHTDebugging",
 			"IncludeBase": ".\\UHTDebugging",
 			"OutputDirectory": "..\\..\\..\\..\\Intermediate/Build/Inc\\UHTDebugging\\",
diff --git a/Engine/Source/Runtime/Core/Private/GenericPlatform/GenericPlatformCrashContext.cpp b/Engine/Source/Runtime/Core/Private/GenericPlatform/GenericPlatformCrashContext.cpp
--- a/Engine/Source/Runtime/Core/Private/GenericPlatform/GenericPlatformCrashContext.cpp
+++ b/Engine/Source/Runtime/Core/Private/GenericPlatform/GenericPlatformCrashContext.cpp
@@ -168,7 +168,7 @@ void FGenericCrashContext::SerializeContentToBuffer()
 	AddCrashProperty( TEXT( "Misc.OSVersionMinor" ), *NCachedCrashContextProperties::OsSubVersion );
 
 
-	// #YRX_Crash: 2015-07-21 Move to the crash report client.
+	// #CrashReport: 2015-07-21 Move to the crash report client.
 	/*{
 		uint64 AppDiskTotalNumberOfBytes = 0;
 		uint64 AppDiskNumberOfFreeBytes = 0;
