integrate from //UE4/Dev-Platform to //UE4/Main

diff --git a/Engine/Source/Programs/UnrealBuildTool/HTML5/HTML5SDKInfo.cs b/Engine/Source/Programs/UnrealBuildTool/HTML5/HTML5SDKInfo.cs
--- a/Engine/Source/Programs/UnrealBuildTool/HTML5/HTML5SDKInfo.cs
+++ b/Engine/Source/Programs/UnrealBuildTool/HTML5/HTML5SDKInfo.cs
@@ -141,6 +141,14 @@ namespace UnrealBuildTool
 				process.WaitForExit();
 //				Console.WriteLine("ExitCode: {0}", process.ExitCode);
 				process.Close();
+				// uncomment OPTIMIZER (GUBP on build machines needs this)
+				// and PYTHON (reduce warnings on EMCC_DEBUG=1)
+				string pyth = Regex.Replace(PYTHON, @"\\", @"\\");
+				string optz = Regex.Replace(Path.Combine(LLVM_ROOT, "optimizer") + PLATFORM_EXE, @"\\", @"\\");
+				File.WriteAllText(DOT_EMSCRIPTEN, Regex.Replace(
+							Regex.Replace(File.ReadAllText(DOT_EMSCRIPTEN), "#(PYTHON).*", "$1 = '" + pyth + "'"),
+							"# (EMSCRIPTEN_NATIVE_OPTIMIZER).*", "$1 = '" + optz + "'"));
+			// --------------------------------------------------
 			// --------------------------------------------------
 			// restore a few things
 			Environment.SetEnvironmentVariable("USERPROFILE", HOME_SAVE);
