Fixing additional instances of linking against automation tests that are now being compiled out in different build configurations.  Adding a way for games to force certain types of automations to be compiled in.

diff --git a/Engine/Source/Programs/UnrealBuildTool/Configuration/UEBuildConfiguration.cs b/Engine/Source/Programs/UnrealBuildTool/Configuration/UEBuildConfiguration.cs
--- a/Engine/Source/Programs/UnrealBuildTool/Configuration/UEBuildConfiguration.cs
+++ b/Engine/Source/Programs/UnrealBuildTool/Configuration/UEBuildConfiguration.cs
@@ -304,6 +304,18 @@ namespace UnrealBuildTool
 		[XmlConfig]
 		public static bool bEditorDependsOnShaderCompileWorker;
 
+        /// <summary>
+        /// Whether to compile development automation tests.
+        /// </summary>
+        [XmlConfig]
+        public static bool bForceCompileDevelopmentAutomationTests;
+
+        /// <summary>
+        /// Whether to compile performance automation tests.
+        /// </summary>
+        [XmlConfig]
+        public static bool bForceCompilePerformanceAutomationTests;
+
 		/// <summary>
 		/// Sets the configuration back to defaults.
 		/// </summary>
@@ -347,6 +359,8 @@ namespace UnrealBuildTool
 			bCompileCEF3 = true;
 			PreferredSubPlatform = "";
 			bEditorDependsOnShaderCompileWorker = true;
+            bForceCompileDevelopmentAutomationTests = false;
+            bForceCompilePerformanceAutomationTests = false;
 		}
 
 		/// <summary>
diff --git a/Engine/Source/Programs/UnrealBuildTool/Configuration/UEBuildTarget.cs b/Engine/Source/Programs/UnrealBuildTool/Configuration/UEBuildTarget.cs
--- a/Engine/Source/Programs/UnrealBuildTool/Configuration/UEBuildTarget.cs
+++ b/Engine/Source/Programs/UnrealBuildTool/Configuration/UEBuildTarget.cs
@@ -3589,21 +3589,41 @@ namespace UnrealBuildTool
 				}
 			}
 
-            if(Configuration == UnrealTargetConfiguration.Shipping)
+            if (UEBuildConfiguration.bForceCompileDevelopmentAutomationTests)
             {
-                GlobalCompileEnvironment.Config.Definitions.Add("WITH_DEV_AUTOMATION_TESTS=0");
-                GlobalCompileEnvironment.Config.Definitions.Add("WITH_PERF_AUTOMATION_TESTS=0");
+                GlobalCompileEnvironment.Config.Definitions.Add("WITH_DEV_AUTOMATION_TESTS=1");
+            }
+            else
+            {
+                switch(Configuration)
+                {
+                    case UnrealTargetConfiguration.Test:
+                    case UnrealTargetConfiguration.Shipping:
+                        GlobalCompileEnvironment.Config.Definitions.Add("WITH_DEV_AUTOMATION_TESTS=0");
+                        break;
+                    default:
+                        GlobalCompileEnvironment.Config.Definitions.Add("WITH_DEV_AUTOMATION_TESTS=1");
+                        break;
+                }
             }
-            else if (Configuration == UnrealTargetConfiguration.Test)
+
+            if (UEBuildConfiguration.bForceCompilePerformanceAutomationTests)
             {
-                GlobalCompileEnvironment.Config.Definitions.Add("WITH_DEV_AUTOMATION_TESTS=0");
                 GlobalCompileEnvironment.Config.Definitions.Add("WITH_PERF_AUTOMATION_TESTS=1");
             }
             else
             {
-                GlobalCompileEnvironment.Config.Definitions.Add("WITH_DEV_AUTOMATION_TESTS=1");
-                GlobalCompileEnvironment.Config.Definitions.Add("WITH_PERF_AUTOMATION_TESTS=1");
+                switch (Configuration)
+                {
+                    case UnrealTargetConfiguration.Shipping:
+                        GlobalCompileEnvironment.Config.Definitions.Add("WITH_PERF_AUTOMATION_TESTS=0");
+                        break;
+                    default:
+                        GlobalCompileEnvironment.Config.Definitions.Add("WITH_PERF_AUTOMATION_TESTS=1");
+                        break;
+                }
             }
+
 			// By default, shadow source files for this target in the root OutputDirectory
 			GlobalCompileEnvironment.Config.LocalShadowDirectory = GlobalCompileEnvironment.Config.OutputDirectory;
 
diff --git a/Engine/Source/Runtime/Engine/Private/Tests/AutomationCommon.cpp b/Engine/Source/Runtime/Engine/Private/Tests/AutomationCommon.cpp
--- a/Engine/Source/Runtime/Engine/Private/Tests/AutomationCommon.cpp
+++ b/Engine/Source/Runtime/Engine/Private/Tests/AutomationCommon.cpp
@@ -11,7 +11,7 @@
 
 #include "Matinee/MatineeActor.h"
 
-#if WITH_DEV_AUTOMATION_TESTS
+#if (WITH_DEV_AUTOMATION_TESTS || WITH_PERF_AUTOMATION_TESTS)
 
 DEFINE_LOG_CATEGORY_STATIC(LogEngineAutomationLatentCommand, Log, All);
 DEFINE_LOG_CATEGORY(LogEditorAutomationTests);
