potential fix for failure to build IOS on build machines

diff --git a/Engine/Source/Programs/UnrealBuildTool/System/ProjectFileGenerator.cs b/Engine/Source/Programs/UnrealBuildTool/System/ProjectFileGenerator.cs
--- a/Engine/Source/Programs/UnrealBuildTool/System/ProjectFileGenerator.cs
+++ b/Engine/Source/Programs/UnrealBuildTool/System/ProjectFileGenerator.cs
@@ -305,7 +305,7 @@ namespace UnrealBuildTool
 			bool IncludeAllPlatforms = true;
 			ConfigureProjectFileGeneration( Arguments, ref IncludeAllPlatforms);
 
-			if (OnlyGameProject != null)
+			if (bGeneratingGameProjectFiles || UnrealBuildTool.IsEngineInstalled())
 			{
 				Log.TraceInformation("Discovering modules, targets and source code for project...");
 
@@ -762,7 +762,7 @@ namespace UnrealBuildTool
 				}
 			}
 
-			if( OnlyGameProject != null )
+			if( bGeneratingGameProjectFiles || UnrealBuildTool.IsEngineInstalled() )
 			{
 				GameProjectName = OnlyGameProject.GetFileNameWithoutExtension();
 				if (String.IsNullOrEmpty(GameProjectName))
