Copying //UE4/Dev-Build to //UE4/Main

==========================
MAJOR FEATURES + CHANGES
==========================

Change 2857832 on 2016/02/05 by [REDACTED_PERSON]

	GitHub #1959 - Prevent UELinkerFixups file from using shared PCHs; it doesn't necessarily have the right compile environment to be able to use them.

Change 2858189 on 2016/02/06 by [REDACTED_PERSON]

	Allow runtime dependencies to be optional, so we can allow games to be staged without requiring PhysX/Apex PDBs.

Change 2858968 on 2016/02/08 by [REDACTED_PERSON]

	Allow plugins to be enabled/disabled depending on different target types. Contributed by [REDACTED_PERSON]/Hi-Rez Studios.

Change 2859325 on 2016/02/08 by [REDACTED_PERSON]

	Stop adding a debug output device if it goes to the same place as the regular log output device. Prevent output messages printing to the log twice on Mac.

Change 2859995 on 2016/02/09 by [REDACTED_PERSON]

	Changed UnrealFileServer to be non-monolithic so that it doesn't include code that is not public

Change 2860393 on 2016/02/09 by [REDACTED_PERSON]

	Add support for pre and post build steps for projects and plugin files. Steps are specified as a list of shell commands per-target platform, and are written to a shell script by UBT. Whenever the target is built, the pre-build steps are executed before UHT, and post-build steps are executed at the end.

	The following target properties are expanded with UBT properties in the commands: $(EngineDir), $(ProjectDir), $(TargetName), $(TargetPlatform), $(TargetConfiguration), $(TargetType), $(ProjectFile), $(PluginDir).

	Example JSON fragment:

		"PreBuildSteps": {
			"Win64": [
				"echo EngineDir=$(EngineDir)",
				"echo ProjectDir=$(ProjectDir)",
				"echo TargetName=$(TargetName)",
				"echo TargetPlatform=$(TargetPlatform)",
				"echo TargetConfiguration=$(TargetConfiguration)",
				"echo TargetType=$(TargetType)",
				"echo ProjectFile=$(ProjectFile)",
				"echo PluginDir=$(PluginDir)"
			]
		},

Change 2860504 on 2016/02/09 by [REDACTED_PERSON]

	Remove nodes from the grid view that we don't have any data for.

Change 2860592 on 2016/02/09 by [REDACTED_PERSON]

	Add script to copy all UAT/UBT build products to another directory. To be used to avoid compiling multiple times on build system.

Change 2861087 on 2016/02/09 by [REDACTED_PERSON]

	Allow UAT binaries compiled on the initial agent to be copied to all other agents in a build. Saves a lot of startup time on Mac, in particular.

Change 2861788 on 2016/02/10 by [REDACTED_PERSON]

	Fix warnings and mangled notification email if triggers finish before all their jobsteps have been created.

diff --git a/Engine/Source/Programs/AutomationTool/AutomationUtils/P4Utils.cs b/Engine/Source/Programs/AutomationTool/AutomationUtils/P4Utils.cs
--- a/Engine/Source/Programs/AutomationTool/AutomationUtils/P4Utils.cs
+++ b/Engine/Source/Programs/AutomationTool/AutomationUtils/P4Utils.cs
@@ -2731,9 +2731,15 @@ namespace AutomationTool
 					{
 						return Line.Substring(CopyPrefix.Length, Line.LastIndexOf('#') - CopyPrefix.Length);
 					}
+
+					const string EditPrefix = "... ... edit from ";
+					if (Line.StartsWith(EditPrefix))
+					{
+						return Line.Substring(EditPrefix.Length, Line.LastIndexOf('#') - EditPrefix.Length);
+					}
 				}
 			}
-			throw new AutomationException("Failed to get integration source for {0}", DepotPath);
+			return null;
 		}
 
 		#region Utilities
diff --git a/Engine/Source/Programs/AutomationTool/AutomationUtils/ProjectUtils.cs b/Engine/Source/Programs/AutomationTool/AutomationUtils/ProjectUtils.cs
--- a/Engine/Source/Programs/AutomationTool/AutomationUtils/ProjectUtils.cs
+++ b/Engine/Source/Programs/AutomationTool/AutomationUtils/ProjectUtils.cs
@@ -233,7 +233,7 @@ namespace AutomationTool
 				// find if there are any plugins enabled or disabled which differ from the default
 				foreach(PluginInfo Plugin in AvailablePlugins)
 				{
-					bool bPluginEnabledForProject = UProjectInfo.IsPluginEnabledForProject(Plugin, Project, TargetPlatformType);
+					bool bPluginEnabledForProject = UProjectInfo.IsPluginEnabledForProject(Plugin, Project, TargetPlatformType, TargetRules.TargetType.Game);
 					if ((bPluginEnabledForProject && !Plugin.Descriptor.bEnabledByDefault) || (bPluginEnabledForProject && Plugin.Descriptor.bInstalled))
 					{
 						if(Plugin.Descriptor.Modules.Any(Module => Module.IsCompiledInConfiguration(TargetPlatformType, TargetRules.TargetType.Game, bBuildDeveloperTools: false, bBuildEditor: false)))
diff --git a/Engine/Source/Programs/AutomationTool/AutomationUtils/UE4Build.cs b/Engine/Source/Programs/AutomationTool/AutomationUtils/UE4Build.cs
--- a/Engine/Source/Programs/AutomationTool/AutomationUtils/UE4Build.cs
+++ b/Engine/Source/Programs/AutomationTool/AutomationUtils/UE4Build.cs
@@ -1532,6 +1532,30 @@ namespace AutomationTool
 			}
 		}
 
+		/// <summary>
+		/// Copy the UAT files to their precompiled location, and add them as build products
+		/// </summary>
+		public void AddUATLauncherFilesToBuildProducts()
+		{
+            var DotNetOutputLocation = CombinePaths(CmdEnv.LocalRoot, "Engine", "Binaries", "DotNET");
+
+			var UATFiles = new List<string>(new string[] 
+					{
+						"AutomationToolLauncher.exe",
+						"AutomationToolLauncher.exe.config",
+					});
+
+			foreach (var UATFile in UATFiles)
+			{
+				var OutputFile = CombinePaths(DotNetOutputLocation, UATFile);
+				if (!FileExists_NoExceptions(OutputFile))
+				{
+					throw new AutomationException("Cannot add UAT to the build products because {0} does not exist.", OutputFile);
+				}
+				AddBuildProduct(OutputFile);
+			}
+		}
+
 		/// <summary>
 		/// Copy the UAT files to their precompiled location, and add them as build products
 		/// </summary>
diff --git a/Engine/Source/Programs/AutomationTool/Scripts/AutomationScripts.Automation.csproj b/Engine/Source/Programs/AutomationTool/Scripts/AutomationScripts.Automation.csproj
--- a/Engine/Source/Programs/AutomationTool/Scripts/AutomationScripts.Automation.csproj
+++ b/Engine/Source/Programs/AutomationTool/Scripts/AutomationScripts.Automation.csproj
@@ -70,6 +70,7 @@
     <Compile Include="CodeSurgery.Automation.cs" />
     <Compile Include="CookCommand.Automation.cs" />
     <Compile Include="CopyBuildToStagingDirectory.Automation.cs" />
+    <Compile Include="CopyUAT.Automation.cs" />
     <Compile Include="DeployCommand.Automation.cs" />
     <Compile Include="FixupRedirects.Automation.cs" />
     <Compile Include="GenerateDSYM.Automation.cs" />
diff --git a/Engine/Source/Programs/AutomationTool/Scripts/CopyBuildToStagingDirectory.Automation.cs b/Engine/Source/Programs/AutomationTool/Scripts/CopyBuildToStagingDirectory.Automation.cs
--- a/Engine/Source/Programs/AutomationTool/Scripts/CopyBuildToStagingDirectory.Automation.cs
+++ b/Engine/Source/Programs/AutomationTool/Scripts/CopyBuildToStagingDirectory.Automation.cs
@@ -396,7 +396,7 @@ public partial class Project : CommandUtils
 			foreach (PluginInfo Plugin in AvailablePlugins)
 			{
 				LogLog("Considering Plugin for Stage: " + Plugin.File.FullName);
-				if (UProjectInfo.IsPluginEnabledForProject(Plugin, Project, SC.StageTargetPlatform.PlatformType))
+				if (UProjectInfo.IsPluginEnabledForProject(Plugin, Project, SC.StageTargetPlatform.PlatformType, TargetRules.TargetType.Game))
 				{
 					LogLog("EnabledPlugin: " + Plugin.File.FullName);
 					SC.StageFiles(StagedFileType.UFS, Plugin.Directory.FullName, "*.uplugin", false, null, null, true, !Params.UsePak(SC.StageTargetPlatform), null, true, false);
diff --git a/Engine/Source/Programs/AutomationTool/Scripts/CopyUAT.Automation.cs b/Engine/Source/Programs/AutomationTool/Scripts/CopyUAT.Automation.cs
new file mode 100644
--- /dev/null
+++ b/Engine/Source/Programs/AutomationTool/Scripts/CopyUAT.Automation.cs
@@ -0,0 +1,55 @@
+﻿// Copyright 1998-2016 Epic Games, Inc. All Rights Reserved.
+
+using System;
+using System.Collections.Generic;
+using System.Linq;
+using System.Text;
+using System.IO;
+using AutomationTool;
+using UnrealBuildTool;
+
+/// <summary>
+/// Copies all UAT and UBT build products to a directory
+/// </summary>
+public class CopyUAT : BuildCommand
+{
+	public override void ExecuteBuild()
+	{
+		// Get the output directory
+		string TargetDir = ParseParamValue("TargetDir");
+		if(TargetDir == null)
+		{
+			throw new AutomationException("Missing -Target argument to CopyUAT");
+		}
+
+		// Construct a dummy UE4Build object to get a list of the UAT and UBT build products
+		UE4Build Build = new UE4Build(this);
+		Build.AddUATFilesToBuildProducts();
+		if(ParseParam("WithLauncher"))
+		{
+			Build.AddUATLauncherFilesToBuildProducts();
+		}
+		Build.AddUBTFilesToBuildProducts();
+
+		// Get a list of all the input files
+		List<string> FileNames = new List<string>(Build.BuildProductFiles);
+		foreach(string FileName in Build.BuildProductFiles)
+		{
+			string SymbolFileName = Path.ChangeExtension(FileName, ".pdb");
+			if(File.Exists(SymbolFileName))
+			{
+				FileNames.Add(SymbolFileName);
+			}
+		}
+
+		// Copy all the files over
+		foreach(string FileName in FileNames)
+		{
+			string TargetFileName = Utils.MakeRerootedFilePath(FileName, CommandUtils.CmdEnv.LocalRoot, TargetDir);
+			Directory.CreateDirectory(Path.GetDirectoryName(TargetFileName));
+			File.Copy(FileName, TargetFileName);
+		}
+
+		Log("Copied {0} files to {1}", FileNames.Count, TargetDir);
+	}
+}
diff --git a/Engine/Source/Programs/UnrealBuildTool/Configuration/UEBuildTarget.cs b/Engine/Source/Programs/UnrealBuildTool/Configuration/UEBuildTarget.cs
--- a/Engine/Source/Programs/UnrealBuildTool/Configuration/UEBuildTarget.cs
+++ b/Engine/Source/Programs/UnrealBuildTool/Configuration/UEBuildTarget.cs
@@ -1002,6 +1002,16 @@ namespace UnrealBuildTool
 		[NonSerialized]
 		UEBuildPlatformContext PlatformContext;
 
+		/// <summary>
+		/// List of scripts to run before building
+		/// </summary>
+		FileReference[] PreBuildStepScripts;
+
+		/// <summary>
+		/// List of scripts to run after building
+		/// </summary>
+		FileReference[] PostBuildStepScripts;
+
 		/// <summary>
 		/// A list of the module filenames which were used to build this target.
 		/// </summary>
@@ -1060,6 +1070,8 @@ namespace UnrealBuildTool
 			FileReferenceToVersionManifestPairs = (KeyValuePair<FileReference, VersionManifest>[])Info.GetValue("vm", typeof(KeyValuePair<FileReference, VersionManifest>[]));
 			TargetCsFilenameField = (FileReference)Info.GetValue("tc", typeof(FileReference));
 			PlatformContext = UEBuildPlatform.GetBuildPlatform(Platform).CreateContext(ProjectFile);
+			PreBuildStepScripts = (FileReference[])Info.GetValue("pr", typeof(FileReference[]));
+			PostBuildStepScripts = (FileReference[])Info.GetValue("po", typeof(FileReference[]));
 		}
 
 		public void GetObjectData(SerializationInfo Info, StreamingContext Context)
@@ -1088,6 +1100,8 @@ namespace UnrealBuildTool
 			Info.AddValue("rf", ReceiptFileName);
 			Info.AddValue("vm", FileReferenceToVersionManifestPairs);
 			Info.AddValue("tc", TargetCsFilenameField);
+			Info.AddValue("pr", PreBuildStepScripts);
+			Info.AddValue("po", PostBuildStepScripts);
 		}
 
 		/// <summary>
@@ -2192,6 +2206,12 @@ namespace UnrealBuildTool
 				? CheckForEULAViolation()
 				: null;
 
+			// Execute the pre-build steps
+			if(!ExecuteCustomPreBuildSteps())
+			{
+				return ECompilationResult.OtherCompilationError;
+			}
+
 			// If we're compiling monolithic, make sure the executable knows about all referenced modules
 			if (ShouldCompileMonolithic())
 			{
@@ -2448,6 +2468,9 @@ namespace UnrealBuildTool
 			// Setup the target's plugins
 			SetupPlugins();
 
+			// Setup the custom build steps for this target
+			SetupCustomBuildSteps();
+
 			// Create all the modules for each binary
 			foreach (UEBuildBinary Binary in AppBinaries)
 			{
@@ -2673,6 +2696,148 @@ namespace UnrealBuildTool
 			}
 		}
 
+		/// <summary>
+		/// Writes scripts for all the custom build steps
+		/// </summary>
+		private void SetupCustomBuildSteps()
+		{
+			// Make sure the intermediate directory exists
+			DirectoryReference ScriptDirectory = ProjectIntermediateDirectory;
+			if(!ScriptDirectory.Exists())
+			{
+				ScriptDirectory.CreateDirectory();
+			}
+
+			// Find all the pre-build steps
+			List<Tuple<CustomBuildSteps, PluginInfo>> PreBuildSteps = new List<Tuple<CustomBuildSteps,PluginInfo>>();
+			if(ProjectDescriptor != null && ProjectDescriptor.PreBuildSteps != null)
+			{
+				PreBuildSteps.Add(Tuple.Create(ProjectDescriptor.PreBuildSteps, (PluginInfo)null));
+			}
+			foreach(PluginInfo EnabledPlugin in EnabledPlugins.Where(x => x.Descriptor.PreBuildSteps != null))
+			{
+				PreBuildSteps.Add(Tuple.Create(EnabledPlugin.Descriptor.PreBuildSteps, EnabledPlugin));
+			}
+			PreBuildStepScripts = WriteCustomBuildStepScripts(BuildHostPlatform.Current.Platform, ScriptDirectory, "PreBuild", PreBuildSteps);
+
+			// Find all the post-build steps
+			List<Tuple<CustomBuildSteps, PluginInfo>> PostBuildSteps = new List<Tuple<CustomBuildSteps,PluginInfo>>();
+			if(ProjectDescriptor != null && ProjectDescriptor.PostBuildSteps != null)
+			{
+				PostBuildSteps.Add(Tuple.Create(ProjectDescriptor.PostBuildSteps, (PluginInfo)null));
+			}
+			foreach(PluginInfo EnabledPlugin in EnabledPlugins.Where(x => x.Descriptor.PostBuildSteps != null))
+			{
+				PostBuildSteps.Add(Tuple.Create(EnabledPlugin.Descriptor.PostBuildSteps, EnabledPlugin));
+			}
+			PostBuildStepScripts = WriteCustomBuildStepScripts(BuildHostPlatform.Current.Platform, ScriptDirectory, "PostBuild", PostBuildSteps);
+		}
+
+		/// <summary>
+		/// Write scripts containing the custom build steps for the given host platform
+		/// </summary>
+		/// <param name="HostPlatform">The current host platform</param>
+		/// <param name="Directory">The output directory for the scripts</param>
+		/// <param name="FilePrefix">Bare prefix for all the created script files</param>
+		/// <param name="BuildStepsAndPluginInfo">List of custom build steps, and their matching PluginInfo (if appropriate)</param>
+		/// <returns>List of created script files</returns>
+		private FileReference[] WriteCustomBuildStepScripts(UnrealTargetPlatform HostPlatform, DirectoryReference Directory, string FilePrefix, List<Tuple<CustomBuildSteps, PluginInfo>> BuildStepsAndPluginInfo)
+		{
+			List<FileReference> ScriptFiles = new List<FileReference>();
+			foreach(Tuple<CustomBuildSteps, PluginInfo> Pair in BuildStepsAndPluginInfo)
+			{
+				CustomBuildSteps BuildSteps = Pair.Item1;
+				if(BuildSteps.HasHostPlatform(HostPlatform))
+				{
+					// Find all the standard variables
+					Dictionary<string, string> Variables = new Dictionary<string,string>();
+					Variables.Add("EngineDir", UnrealBuildTool.EngineDirectory.FullName);
+					Variables.Add("ProjectDir", ProjectDirectory.FullName);
+					Variables.Add("TargetName", TargetName);
+					Variables.Add("TargetPlatform", Platform.ToString());
+					Variables.Add("TargetConfiguration", Configuration.ToString());
+					Variables.Add("TargetType", TargetType.ToString());
+					if(ProjectFile != null)
+					{
+						Variables.Add("ProjectFile", ProjectFile.FullName);
+					}
+					if(Pair.Item2 != null)
+					{
+						Variables.Add("PluginDir", Pair.Item2.Directory.FullName);
+					}
+
+					// Get the commands to execute
+					string[] Commands;
+					if(BuildSteps.TryGetCommands(HostPlatform, Variables, out Commands))
+					{
+						// Get the output path to the script
+						string ScriptExtension = (HostPlatform == UnrealTargetPlatform.Win64)? ".bat" : ".sh";
+						FileReference ScriptFile = FileReference.Combine(Directory, String.Format("{0}-{1}{2}", FilePrefix, ScriptFiles.Count + 1, ScriptExtension));
+
+						// Write it to disk
+						List<string> AllCommands = new List<string>(Commands);
+						if(HostPlatform == UnrealTargetPlatform.Win64)
+						{
+							AllCommands.Insert(0, "@echo off");
+						}
+						File.WriteAllLines(ScriptFile.FullName, AllCommands);
+
+						// Add the output file to the list of generated scripts
+						ScriptFiles.Add(ScriptFile);
+					}
+				}
+			}
+			return ScriptFiles.ToArray();
+		}
+
+		/// <summary>
+		/// Executes the custom pre-build steps
+		/// </summary>
+		public bool ExecuteCustomPreBuildSteps()
+		{
+			return ExecuteCustomBuildSteps(PreBuildStepScripts);
+		}
+
+		/// <summary>
+		/// Executes the custom post-build steps
+		/// </summary>
+		public bool ExecuteCustomPostBuildSteps()
+		{
+			return ExecuteCustomBuildSteps(PostBuildStepScripts);
+		}
+
+		/// <summary>
+		/// Executes a list of custom build step scripts
+		/// </summary>
+		/// <param name="ScriptFiles">List of script files to execute</param>
+		/// <returns>True if the steps succeeded, false otherwise</returns>
+		private bool ExecuteCustomBuildSteps(FileReference[] ScriptFiles)
+		{
+			UnrealTargetPlatform HostPlatform = BuildHostPlatform.Current.Platform;
+			foreach(FileReference ScriptFile in ScriptFiles)
+			{
+				ProcessStartInfo StartInfo = new ProcessStartInfo();
+				if(HostPlatform == UnrealTargetPlatform.Win64)
+				{
+					StartInfo.FileName = "cmd.exe";
+					StartInfo.Arguments = String.Format("/C \"{0}\"", ScriptFile.FullName);
+				}
+				else
+				{
+					StartInfo.FileName = "/bin/sh";
+					StartInfo.Arguments = String.Format("\"{0}\"", ScriptFile.FullName);
+				}
+
+				int ReturnCode = Utils.RunLocalProcessAndLogOutput(StartInfo);
+				if(ReturnCode != 0)
+				{
+					Log.TraceError("Custom build step terminated with exit code {0}", ReturnCode);
+					return false;
+				}
+			}
+			return true;
+		}
+
 		private static FileReference AddModuleFilenameSuffix(string ModuleName, FileReference FilePath, string Suffix)
 		{
 			var MatchPos = FilePath.FullName.LastIndexOf(ModuleName, StringComparison.InvariantCultureIgnoreCase);
@@ -2779,6 +2944,7 @@ namespace UnrealBuildTool
 
 				// Now bind this new module to the executable binary so it will link the plugin libs correctly
 				NewModule.bSkipDefinitionsForCompileEnvironment = true;
+				NewModule.Rules.PCHUsage = ModuleRules.PCHUsageMode.NoSharedPCHs;
 				NewModule.RecursivelyCreateModules();
 				BindArtificialModuleToBinary(NewModule, ExecutableBinary);
 
@@ -3426,7 +3592,7 @@ namespace UnrealBuildTool
 				ProjectDescriptor Project = (UEBuildConfiguration.bCompileAgainstEngine && ProjectFile != null) ? ProjectDescriptor.FromFile(ProjectFile.FullName) : null;
 				foreach (PluginInfo ValidPlugin in ValidPlugins)
 				{
-					if (UProjectInfo.IsPluginEnabledForProject(ValidPlugin, Project, Platform))
+					if(UProjectInfo.IsPluginEnabledForProject(ValidPlugin, Project, Platform, TargetType))
 					{
 						if (ValidPlugin.Descriptor.bCanBeUsedWithUnrealHeaderTool)
 						{
@@ -3780,6 +3946,7 @@ namespace UnrealBuildTool
 
 			// tell the compiled code the name of the UBT platform (this affects folder on disk, etc that the game may need to know)
 			GlobalCompileEnvironment.Config.Definitions.Add("UBT_COMPILED_PLATFORM=" + Platform.ToString());
+			GlobalCompileEnvironment.Config.Definitions.Add("UBT_COMPILED_TARGET=" + TargetType.ToString());
 
 			// Initialize the compile and link environments for the platform, configuration, and project.
 			SetUpPlatformEnvironment();
diff --git a/Engine/Source/Programs/UnrealBuildTool/System/CustomBuildSteps.cs b/Engine/Source/Programs/UnrealBuildTool/System/CustomBuildSteps.cs
new file mode 100644
--- /dev/null
+++ b/Engine/Source/Programs/UnrealBuildTool/System/CustomBuildSteps.cs
@@ -0,0 +1,105 @@
+﻿using System;
+using System.Collections.Generic;
+using System.Linq;
+using System.Text;
+using System.Threading.Tasks;
+
+namespace UnrealBuildTool
+{
+	public class CustomBuildSteps
+	{
+		SortedDictionary<UnrealTargetPlatform, string[]> HostPlatformToCommands = new SortedDictionary<UnrealTargetPlatform,string[]>();
+
+		/// <summary>
+		/// Construct a custom build steps object from a Json object.
+		/// </summary>
+		public CustomBuildSteps(JsonObject RawObject)
+		{
+			foreach(string HostPlatformName in RawObject.KeyNames)
+			{
+				UnrealTargetPlatform HostPlatform;
+				if(Enum.TryParse(HostPlatformName, true, out HostPlatform))
+				{
+					HostPlatformToCommands.Add(HostPlatform, RawObject.GetStringArrayField(HostPlatformName));
+				}
+			}
+		}
+
+		/// <summary>
+		/// Reads a list of build steps from a Json project or plugin descriptor
+		/// </summary>
+		/// <param name="RawObject">The json descriptor object</param>
+		/// <param name="FieldName">Name of the field to read</param>
+		/// <param name="OutBuildSteps">Output variable to store the sorted dictionary that was read</param>
+		/// <returns>True if the field was read (and OutBuildSteps is set), false otherwise.</returns>
+		public static bool TryRead(JsonObject RawObject, string FieldName, out CustomBuildSteps OutBuildSteps)
+		{
+			JsonObject BuildStepsObject;
+			if(RawObject.TryGetObjectField(FieldName, out BuildStepsObject))
+			{
+				OutBuildSteps = new CustomBuildSteps(BuildStepsObject);
+				return true;
+			}
+			else
+			{
+				OutBuildSteps = null;
+				return false;
+			}
+		}
+
+		/// <summary>
+		/// Reads a list of build steps from a Json project or plugin descriptor
+		/// </summary>
+		/// <param name="RawObject">The json descriptor object</param>
+		/// <param name="FieldName">Name of the field to read</param>
+		/// <param name="OutBuildSteps">Output variable to store the sorted dictionary that was read</param>
+		/// <returns>True if the field was read (and OutBuildSteps is set), false otherwise.</returns>
+		public void Write(JsonWriter Writer, string FieldName)
+		{
+			Writer.WriteObjectStart(FieldName);
+			foreach(KeyValuePair<UnrealTargetPlatform, string[]> Pair in HostPlatformToCommands)
+			{
+				Writer.WriteArrayStart(Pair.Key.ToString());
+				foreach(string Line in Pair.Value)
+				{
+					Writer.WriteValue(Line);
+				}
+				Writer.WriteArrayEnd();
+			}
+			Writer.WriteObjectEnd();
+		}
+
+		/// <summary>
+		/// Determines whether there are custom build steps for the given host platform
+		/// </summary>
+		/// <param name="HostPlatform">The host platform to check for</param>
+		/// <returns>True if the host platform exists</returns>
+		public bool HasHostPlatform(UnrealTargetPlatform HostPlatform)
+		{
+			string[] Commands;
+			return HostPlatformToCommands.TryGetValue(HostPlatform, out Commands) && Commands.Length > 0;
+		}
+
+		/// <summary>
+		/// Tries to get the commands for a given host platform
+		/// </summary>
+		/// <param name="HostPlatform">The host platform to look for</param>
+		/// <param name="Variables">Lookup of additional environment variables to expand</param>
+		/// <param name="OutCommands">Array of commands</param>
+		/// <returns>True if a list of commands was generated</returns>
+		public bool TryGetCommands(UnrealTargetPlatform HostPlatform, Dictionary<string, string> Variables, out string[] OutCommands)
+		{
+			string[] Commands;
+			if(HostPlatformToCommands.TryGetValue(HostPlatform, out Commands) && Commands.Length > 0)
+			{
+				OutCommands = Commands.Select(x => Utils.ExpandVariables(x, Variables)).ToArray();
+				return true;
+			}
+			else
+			{
+				OutCommands = null;
+				return false;
+			}
+		}
+	}
+}
diff --git a/Engine/Source/Programs/UnrealBuildTool/System/PluginDescriptor.cs b/Engine/Source/Programs/UnrealBuildTool/System/PluginDescriptor.cs
--- a/Engine/Source/Programs/UnrealBuildTool/System/PluginDescriptor.cs
+++ b/Engine/Source/Programs/UnrealBuildTool/System/PluginDescriptor.cs
@@ -24,63 +24,110 @@ namespace UnrealBuildTool
 
 	public class PluginDescriptor
 	{
-		// Descriptor version number
+
+		/// <summary>
+		/// Descriptor version number
+		/// </summary>
 		public int FileVersion;
 
-		// Version number for the plugin.  The version number must increase with every version of the plugin, so that the system 
-		// can determine whether one version of a plugin is newer than another, or to enforce other requirements.  This version
-		// number is not displayed in front-facing UI.  Use the VersionName for that.
+		/// <summary>
+		/// Version number for the plugin.  The version number must increase with every version of the plugin, so that the system 
+		/// can determine whether one version of a plugin is newer than another, or to enforce other requirements.  This version
+		/// number is not displayed in front-facing UI.  Use the VersionName for that.
+		/// </summary>
 		public int Version;
 
-		// Name of the version for this plugin.  This is the front-facing part of the version number.  It doesn't need to match
-		// the version number numerically, but should be updated when the version number is increased accordingly.
+		/// <summary>
+		/// Name of the version for this plugin.  This is the front-facing part of the version number.  It doesn't need to match
+		/// the version number numerically, but should be updated when the version number is increased accordingly.
+		/// </summary>
 		public string VersionName;
 
-		// Friendly name of the plugin
+		/// <summary>
+		/// Friendly name of the plugin
+		/// </summary>
 		public string FriendlyName;
 
-		// Description of the plugin
+		/// <summary>
+		/// Description of the plugin
+		/// </summary>
 		public string Description;
 
-		// The name of the category this plugin
+		/// <summary>
+		/// The name of the category this plugin
+		/// </summary>
 		public string Category;
 
-		// The company or individual who created this plugin.  This is an optional field that may be displayed in the user interface.
+		/// <summary>
+		/// The company or individual who created this plugin.  This is an optional field that may be displayed in the user interface.
+		/// </summary>
 		public string CreatedBy;
 
-		// Hyperlink URL string for the company or individual who created this plugin.  This is optional.
+		/// <summary>
+		/// Hyperlink URL string for the company or individual who created this plugin.  This is optional.
+		/// </summary>
 		public string CreatedByURL;
 
-		// Documentation URL string.
+		/// <summary>
+		/// Documentation URL string.
+		/// </summary>
 		public string DocsURL;
 
-		// Marketplace URL for this plugin. This URL will be embedded into projects that enable this plugin, so we can redirect to the marketplace if a user doesn't have it installed.
+		/// <summary>
+		/// Marketplace URL for this plugin. This URL will be embedded into projects that enable this plugin, so we can redirect to the marketplace if a user doesn't have it installed.
+		/// </summary>
 		public string MarketplaceURL;
 
-		// Support URL/email for this plugin.
+		/// <summary>
+		/// Support URL/email for this plugin.
+		/// </summary>
 		public string SupportURL;
 
-		// List of all modules associated with this plugin
+		/// <summary>
+		/// List of all modules associated with this plugin
+		/// </summary>
 		public ModuleDescriptor[] Modules;
 
-		// Whether this plugin should be enabled by default for all projects
+		/// <summary>
+		/// Whether this plugin should be enabled by default for all projects
+		/// </summary>
 		public bool bEnabledByDefault;
 
-		// Can this plugin contain content?
+		/// <summary>
+		/// Can this plugin contain content?
+		/// </summary>
 		public bool bCanContainContent;
 
-		// Marks the plugin as beta in the UI
+		/// </summary>
+		/// Marks the plugin as beta in the UI
+		/// </summary>
 		public bool bIsBetaVersion;
 
-		// Whether this plugin can be used by UnrealHeaderTool
+		/// <summary>
+		/// Whether this plugin can be used by UnrealHeaderTool
+		/// </summary>
 		public bool bCanBeUsedWithUnrealHeaderTool;
 
-		// Set for plugins which are installed
+		/// <summary>
+		/// Set for plugins which are installed
+		/// </summary>
 		public bool bInstalled;
 
-		// For plugins that are under a platform folder (eg. /PS4/), determines whether compiling the plugin requires the build platform and/or SDK to be available
+		/// <summary>
+		/// For plugins that are under a platform folder (eg. /PS4/), determines whether compiling the plugin requires the build platform and/or SDK to be available
+		/// </summary>
 		public bool bRequiresBuildPlatform;
 
+		/// <summary>
+		/// Set of pre-build steps to execute, keyed by host platform name.
+		/// </summary>
+		public CustomBuildSteps PreBuildSteps;
+
+		/// <summary>
+		/// Set of post-build steps to execute, keyed by host platform name.
+		/// </summary>
+		public CustomBuildSteps PostBuildSteps;
+
 		/// <summary>
 		/// Private constructor. This object should not be created directly; read it from disk using FromFile() instead.
 		/// </summary>
@@ -155,6 +202,9 @@ namespace UnrealBuildTool
 				RawObject.TryGetBoolField("CanBeUsedWithUnrealHeaderTool", out Descriptor.bCanBeUsedWithUnrealHeaderTool);
 				RawObject.TryGetBoolField("RequiresBuildPlatform", out Descriptor.bRequiresBuildPlatform);
 
+				CustomBuildSteps.TryRead(RawObject, "PreBuildSteps", out Descriptor.PreBuildSteps);
+				CustomBuildSteps.TryRead(RawObject, "PostBuildSteps", out Descriptor.PostBuildSteps);
+
 				return Descriptor;
 			}
 			catch (JsonParseException ParseException)
@@ -192,6 +242,16 @@ namespace UnrealBuildTool
 
 				ModuleDescriptor.WriteArray(Writer, "Modules", Modules);
 
+				if(PreBuildSteps != null)
+				{
+					PreBuildSteps.Write(Writer, "PreBuildSteps");
+				}
+
+				if(PostBuildSteps != null)
+				{
+					PostBuildSteps.Write(Writer, "PostBuildSteps");
+				}
+
 				Writer.WriteObjectEnd();
 			}
 		}
@@ -221,6 +281,12 @@ namespace UnrealBuildTool
 		// If enabled, list of platforms for which the plugin should be disabled.
 		UnrealTargetPlatform[] BlacklistPlatforms;
 
+		// If enabled, list of targets for which the plugin should be enabled (or all targets if blank).
+		TargetRules.TargetType[] WhitelistTargets;
+
+		// If enabled, list of targets for which the plugin should be disabled.
+		TargetRules.TargetType[] BlacklistTargets;
+
 		/// <summary>
 		/// Constructor
 		/// </summary>
@@ -246,6 +312,8 @@ namespace UnrealBuildTool
 			RawObject.TryGetStringField("MarketplaceURL", out Descriptor.MarketplaceURL);
 			RawObject.TryGetEnumArrayField<UnrealTargetPlatform>("WhitelistPlatforms", out Descriptor.WhitelistPlatforms);
 			RawObject.TryGetEnumArrayField<UnrealTargetPlatform>("BlacklistPlatforms", out Descriptor.BlacklistPlatforms);
+			RawObject.TryGetEnumArrayField<TargetRules.TargetType>("WhitelistTargets", out Descriptor.WhitelistTargets);
+			RawObject.TryGetEnumArrayField<TargetRules.TargetType>("BlacklistTargets", out Descriptor.BlacklistTargets);
 			return Descriptor;
 		}
 
@@ -270,5 +338,27 @@ namespace UnrealBuildTool
 			}
 			return true;
 		}
+
+		/// <summary>
+		/// Determines if this reference enables the plugin for a given target
+		/// </summary>
+		/// <param name="Target">The target to check</param>
+		/// <returns>True if the plugin should be enabled</returns>
+		public bool IsEnabledForTarget(TargetRules.TargetType Target)
+		{
+			if (!bEnabled)
+			{
+				return false;
+			}
+			if (WhitelistTargets != null && WhitelistTargets.Length > 0 && !WhitelistTargets.Contains(Target))
+			{
+				return false;
+			}
+			if (BlacklistTargets != null && BlacklistTargets.Contains(Target))
+			{
+				return false;
+			}
+			return true;
+		}
 	}
 }
diff --git a/Engine/Source/Programs/UnrealBuildTool/System/ProjectDescriptor.cs b/Engine/Source/Programs/UnrealBuildTool/System/ProjectDescriptor.cs
--- a/Engine/Source/Programs/UnrealBuildTool/System/ProjectDescriptor.cs
+++ b/Engine/Source/Programs/UnrealBuildTool/System/ProjectDescriptor.cs
@@ -46,6 +46,12 @@ namespace UnrealBuildTool
 		// A hash that is used to determine if the project was forked from a sample
 		public uint EpicSampleNameHash;
 
+		// Steps to execute before building targets in this project
+		public CustomBuildSteps PreBuildSteps;
+
+		// Steps to execute before building targets in this project
+		public CustomBuildSteps PostBuildSteps;
+
 		/// <summary>
 		/// Constructor.
 		/// </summary>
@@ -106,6 +112,10 @@ namespace UnrealBuildTool
 				// Get the sample name hash
 				RawObject.TryGetUnsignedIntegerField("EpicSampleNameHash", out Descriptor.EpicSampleNameHash);
 
+				// Read the pre and post-build steps
+				CustomBuildSteps.TryRead(RawObject, "PreBuildSteps", out Descriptor.PreBuildSteps);
+				CustomBuildSteps.TryRead(RawObject, "PostBuildSteps", out Descriptor.PostBuildSteps);
+
 				return Descriptor;
 			}
 			catch (JsonParseException ParseException)
diff --git a/Engine/Source/Programs/UnrealBuildTool/System/UProjectInfo.cs b/Engine/Source/Programs/UnrealBuildTool/System/UProjectInfo.cs
--- a/Engine/Source/Programs/UnrealBuildTool/System/UProjectInfo.cs
+++ b/Engine/Source/Programs/UnrealBuildTool/System/UProjectInfo.cs
@@ -305,7 +305,7 @@ namespace UnrealBuildTool
 		/// <param name="Plugin">Information about the plugin</param>
 		/// <param name="Platform">The target platform</param>
 		/// <returns>True if the plugin should be enabled for this project</returns>
-		public static bool IsPluginEnabledForProject(PluginInfo Plugin, ProjectDescriptor Project, UnrealTargetPlatform Platform)
+		public static bool IsPluginEnabledForProject(PluginInfo Plugin, ProjectDescriptor Project, UnrealTargetPlatform Platform, TargetRules.TargetType Target)
 		{
 			bool bEnabled = Plugin.Descriptor.bEnabledByDefault || Plugin.LoadedFrom == PluginLoadedFrom.GameProject;
 			if (Project != null && Project.Plugins != null)
@@ -314,7 +314,7 @@ namespace UnrealBuildTool
 				{
 					if (String.Compare(PluginReference.Name, Plugin.Name, true) == 0)
 					{
-						bEnabled = PluginReference.IsEnabledForPlatform(Platform);
+						bEnabled = PluginReference.IsEnabledForPlatform(Platform) && PluginReference.IsEnabledForTarget(Target);
 					}
 				}
 			}
diff --git a/Engine/Source/Programs/UnrealBuildTool/System/UnrealBuildTool.cs b/Engine/Source/Programs/UnrealBuildTool/System/UnrealBuildTool.cs
--- a/Engine/Source/Programs/UnrealBuildTool/System/UnrealBuildTool.cs
+++ b/Engine/Source/Programs/UnrealBuildTool/System/UnrealBuildTool.cs
@@ -1912,7 +1912,7 @@ namespace UnrealBuildTool
 					if (UnrealBuildTool.IsAssemblingBuild)
 					{
 						// If we didn't build the graph in this session, then we'll need to load a cached one
-						if (!UnrealBuildTool.IsGatheringBuild)
+						if (!UnrealBuildTool.IsGatheringBuild && BuildResult.Succeeded())
 						{
 							ActionGraph.AllActions = UBTMakefile.AllActions;
 
@@ -1930,22 +1930,35 @@ namespace UnrealBuildTool
 								Environment.SetEnvironmentVariable(EnvironmentVariable.Item1, EnvironmentVariable.Item2);
 							}
 
+							// Execute all the pre-build steps
+							foreach(UEBuildTarget Target in Targets)
+							{
+								if(!Target.ExecuteCustomPreBuildSteps())
+								{
+									BuildResult = ECompilationResult.OtherCompilationError;
+									break;
+								}
+							}
+
 							// If any of the targets need UHT to be run, we'll go ahead and do that now
-							foreach (var Target in Targets)
+							if(BuildResult.Succeeded())
 							{
-								List<UHTModuleInfo> TargetUObjectModules;
-								if (UBTMakefile.TargetNameToUObjectModules.TryGetValue(Target.GetTargetName(), out TargetUObjectModules))
+								foreach (var Target in Targets)
 								{
-									if (TargetUObjectModules.Count > 0)
+									List<UHTModuleInfo> TargetUObjectModules;
+									if (UBTMakefile.TargetNameToUObjectModules.TryGetValue(Target.GetTargetName(), out TargetUObjectModules))
 									{
-										// Execute the header tool
-										FileReference ModuleInfoFileName = FileReference.Combine(Target.ProjectIntermediateDirectory, Target.GetTargetName() + ".uhtmanifest");
-										ECompilationResult UHTResult = ECompilationResult.OtherCompilationError;
-										if (!ExternalExecution.ExecuteHeaderToolIfNecessary(ToolChain, Target, GlobalCompileEnvironment: null, UObjectModules: TargetUObjectModules, ModuleInfoFileName: ModuleInfoFileName, UHTResult: ref UHTResult))
+										if (TargetUObjectModules.Count > 0)
 										{
-											Log.TraceInformation("UnrealHeaderTool failed for target '" + Target.GetTargetName() + "' (platform: " + Target.Platform.ToString() + ", module info: " + ModuleInfoFileName + ").");
-											BuildResult = UHTResult;
-											break;
+											// Execute the header tool
+											FileReference ModuleInfoFileName = FileReference.Combine(Target.ProjectIntermediateDirectory, Target.GetTargetName() + ".uhtmanifest");
+											ECompilationResult UHTResult = ECompilationResult.OtherCompilationError;
+											if (!ExternalExecution.ExecuteHeaderToolIfNecessary(ToolChain, Target, GlobalCompileEnvironment: null, UObjectModules: TargetUObjectModules, ModuleInfoFileName: ModuleInfoFileName, UHTResult: ref UHTResult))
+											{
+												Log.TraceInformation("UnrealHeaderTool failed for target '" + Target.GetTargetName() + "' (platform: " + Target.Platform.ToString() + ", module info: " + ModuleInfoFileName + ").");
+												BuildResult = UHTResult;
+												break;
+											}
 										}
 									}
 								}
@@ -2018,6 +2031,19 @@ namespace UnrealBuildTool
 								BuildResult = ECompilationResult.OtherCompilationError;
 							}
 						}
+
+						// Execute all the post-build steps
+						if(BuildResult.Succeeded())
+						{
+							foreach(UEBuildTarget Target in Targets)
+							{
+								if(!Target.ExecuteCustomPostBuildSteps())
+								{
+									BuildResult = ECompilationResult.OtherCompilationError;
+									break;
+								}
+							}
+						}
 					}
 				}
 			}
diff --git a/Engine/Source/Programs/UnrealBuildTool/UnrealBuildTool.csproj b/Engine/Source/Programs/UnrealBuildTool/UnrealBuildTool.csproj
--- a/Engine/Source/Programs/UnrealBuildTool/UnrealBuildTool.csproj
+++ b/Engine/Source/Programs/UnrealBuildTool/UnrealBuildTool.csproj
@@ -118,6 +118,7 @@
     <Compile Include="Mac\UEBuildMac.cs" />
     <Compile Include="Mac\MacProjectGenerator.cs" />
     <Compile Include="Mac\UEDeployMac.cs" />
+    <Compile Include="System\CustomBuildSteps.cs" />
     <Compile Include="System\FileSystemReference.cs" />
     <Compile Include="System\InstalledPlatformInfo.cs" />
     <Compile Include="System\TargetReceipt.cs" />
diff --git a/Engine/Source/Programs/UnrealFileServer/UnrealFileServer.Target.cs b/Engine/Source/Programs/UnrealFileServer/UnrealFileServer.Target.cs
--- a/Engine/Source/Programs/UnrealFileServer/UnrealFileServer.Target.cs
+++ b/Engine/Source/Programs/UnrealFileServer/UnrealFileServer.Target.cs
@@ -28,7 +28,8 @@ public class UnrealFileServerTarget : TargetRules
 
 	public override bool ShouldCompileMonolithic(UnrealTargetPlatform InPlatform, UnrealTargetConfiguration InConfiguration)
 	{
-		return true;
+		// Don't build monolithic because it could depend on modules that are not public
+		return false;
 	}
 
 	public override void SetupGlobalEnvironment(
diff --git a/Engine/Source/Programs/UnrealGameSync/UnrealGameSync/Properties/AssemblyInfo.cs b/Engine/Source/Programs/UnrealGameSync/UnrealGameSync/Properties/AssemblyInfo.cs
--- a/Engine/Source/Programs/UnrealGameSync/UnrealGameSync/Properties/AssemblyInfo.cs
+++ b/Engine/Source/Programs/UnrealGameSync/UnrealGameSync/Properties/AssemblyInfo.cs
@@ -34,5 +34,5 @@ using System.Runtime.InteropServices;
 // You can specify all the values or you can default the Build and Revision Numbers 
 // by using the '*' as shown below:
 // [assembly: AssemblyVersion("1.0.*")]
-[assembly: AssemblyVersion("1.65.*")]
-[assembly: AssemblyFileVersion("1.65.0.0")]
+[assembly: AssemblyVersion("1.66.*")]
+[assembly: AssemblyFileVersion("1.66.0.0")]
diff --git a/Engine/Source/Programs/UnrealGameSync/UnrealGameSync/Workspace.cs b/Engine/Source/Programs/UnrealGameSync/UnrealGameSync/Workspace.cs
--- a/Engine/Source/Programs/UnrealGameSync/UnrealGameSync/Workspace.cs
+++ b/Engine/Source/Programs/UnrealGameSync/UnrealGameSync/Workspace.cs
@@ -342,6 +342,7 @@ namespace UnrealGameSync
 						{
 							Dictionary<string, string> BuildVersionStrings = new Dictionary<string,string>();
 							BuildVersionStrings["\"Changelist\":"] = String.Format(" {0},", VersionChangeNumber);
+							BuildVersionStrings["\"BranchName\":"] = String.Format(" \"{0}\"", BranchOrStreamName.Replace('/', '+'));
 							if(!UpdateVersionFile(ClientRootPath + BuildVersionFileName, BuildVersionStrings))
 							{
 								StatusMessage = String.Format("Failed to update {0}.", BuildVersionFileName);
diff --git a/Engine/Source/Runtime/Core/Private/GenericPlatform/GenericPlatformMisc.cpp b/Engine/Source/Runtime/Core/Private/GenericPlatform/GenericPlatformMisc.cpp
--- a/Engine/Source/Runtime/Core/Private/GenericPlatform/GenericPlatformMisc.cpp
+++ b/Engine/Source/Runtime/Core/Private/GenericPlatform/GenericPlatformMisc.cpp
@@ -390,6 +390,11 @@ void FGenericPlatformMisc::LocalPrint( const TCHAR* Str )
 #endif
 }
 
+bool FGenericPlatformMisc::HasSeparateChannelForDebugOutput()
+{
+	return false;
+}
+
 void FGenericPlatformMisc::RequestMinimize()
 {
 }
@@ -864,6 +869,11 @@ const TCHAR* FGenericPlatformMisc::GetUBTPlatform()
 	return TEXT( PREPROCESSOR_TO_STRING(UBT_COMPILED_PLATFORM) );
 }
 
+const TCHAR* FGenericPlatformMisc::GetUBTTarget()
+{
+    return TEXT(PREPROCESSOR_TO_STRING(UBT_COMPILED_TARGET));
+}
+
 const TCHAR* FGenericPlatformMisc::GetDefaultDeviceProfileName()
 {
 	return TEXT("Default");
diff --git a/Engine/Source/Runtime/Core/Private/GenericPlatform/GenericPlatformOutputDevices.cpp b/Engine/Source/Runtime/Core/Private/GenericPlatform/GenericPlatformOutputDevices.cpp
--- a/Engine/Source/Runtime/Core/Private/GenericPlatform/GenericPlatformOutputDevices.cpp
+++ b/Engine/Source/Runtime/Core/Private/GenericPlatform/GenericPlatformOutputDevices.cpp
@@ -21,7 +21,11 @@ void FGenericPlatformOutputDevices::SetupOutputDevices()
 	// A shipping build with logging explicitly enabled will fail the IsDebuggerPresent() check, but we still need to add the debug output device for logging purposes
 	if (!FPlatformProperties::SupportsWindowedMode() || FPlatformMisc::IsDebuggerPresent() || (UE_BUILD_SHIPPING && !NO_LOGGING) || GIsBuildMachine)
 	{
-		GLog->AddOutputDevice(new FOutputDeviceDebug());
+		// Only need to do this if it's actually going to go to a different place than GLogConsole
+		if(FPlatformMisc::HasSeparateChannelForDebugOutput())
+		{
+			GLog->AddOutputDevice(new FOutputDeviceDebug());
+		}
 	}
 
 	GLog->AddOutputDevice(FPlatformOutputDevices::GetEventLog());
diff --git a/Engine/Source/Runtime/Core/Private/Windows/WindowsPlatformMisc.cpp b/Engine/Source/Runtime/Core/Private/Windows/WindowsPlatformMisc.cpp
--- a/Engine/Source/Runtime/Core/Private/Windows/WindowsPlatformMisc.cpp
+++ b/Engine/Source/Runtime/Core/Private/Windows/WindowsPlatformMisc.cpp
@@ -1062,6 +1062,11 @@ void FWindowsPlatformMisc::LocalPrint( const TCHAR *Message )
 	OutputDebugString(Message);
 }
 
+bool FWindowsPlatformMisc::HasSeparateChannelForDebugOutput()
+{
+	return true;
+}
+
 void FWindowsPlatformMisc::RequestExit( bool Force )
 {
 	UE_LOG(LogWindows, Log,  TEXT("FPlatformMisc::RequestExit(%i)"), Force );
diff --git a/Engine/Source/Runtime/Core/Public/GenericPlatform/GenericPlatformMisc.h b/Engine/Source/Runtime/Core/Public/GenericPlatform/GenericPlatformMisc.h
--- a/Engine/Source/Runtime/Core/Public/GenericPlatform/GenericPlatformMisc.h
+++ b/Engine/Source/Runtime/Core/Public/GenericPlatform/GenericPlatformMisc.h
@@ -461,6 +461,11 @@ public:
 	/** Prints string to the default output */
 	static void LocalPrint( const TCHAR* Str );
 
+	/** 
+	 * Whether the platform has a separate debug channel to stdout (eg. OutputDebugString on Windows). Used to suppress messages being output twice 
+	 * if both go to the same place.
+	 */
+	static bool HasSeparateChannelForDebugOutput();
 
 	/** Request application to minimize (goto background). **/
 	static void RequestMinimize();
@@ -665,6 +670,8 @@ public:
 
 	static const TCHAR* GetUBTPlatform();
 
+	static const TCHAR* GetUBTTarget();
+
 	/** 
 	 * Determines the shader format for the plarform
 	 *
diff --git a/Engine/Source/Runtime/Core/Public/Windows/WindowsPlatformMisc.h b/Engine/Source/Runtime/Core/Public/Windows/WindowsPlatformMisc.h
--- a/Engine/Source/Runtime/Core/Public/Windows/WindowsPlatformMisc.h
+++ b/Engine/Source/Runtime/Core/Public/Windows/WindowsPlatformMisc.h
@@ -85,6 +85,7 @@ struct CORE_API FWindowsPlatformMisc
 	static uint32 GetCharKeyMap(uint32* KeyCodes, FString* KeyNames, uint32 MaxMappings);
 	static void SetUTF8Output();
 	static void LocalPrint(const TCHAR *Message);
+	static bool HasSeparateChannelForDebugOutput();
 	static void RequestExit(bool Force);
 	static void RequestMinimize();
 	static const TCHAR* GetSystemErrorMessage(TCHAR* OutBuffer, int32 BufferCount, int32 Error);
diff --git a/Engine/Source/Runtime/Projects/Private/CustomBuildSteps.cpp b/Engine/Source/Runtime/Projects/Private/CustomBuildSteps.cpp
new file mode 100644
--- /dev/null
+++ b/Engine/Source/Runtime/Projects/Private/CustomBuildSteps.cpp
@@ -0,0 +1,53 @@
+// Copyright 1998-2016 Epic Games, Inc. All Rights Reserved.
+
+#include "ProjectsPrivatePCH.h"
+#include "CustomBuildSteps.h"
+
+#define LOCTEXT_NAMESPACE "ModuleDescriptor"
+
+bool FCustomBuildSteps::IsEmpty() const
+{
+	return HostPlatformToCommands.Num() == 0;
+}
+
+void FCustomBuildSteps::Read(const FJsonObject& Object, const FString& FieldName)
+{
+	TSharedPtr<FJsonValue> StepsValue = Object.TryGetField(FieldName);
+	if(StepsValue.IsValid() && StepsValue->Type == EJson::Object)
+	{
+		const TSharedPtr<FJsonObject>& StepsObject = StepsValue->AsObject();
+		for(const TPair<FString, TSharedPtr<FJsonValue>>& HostPlatformAndSteps : StepsObject->Values)
+		{
+			TArray<FString>& Commands = HostPlatformToCommands.FindOrAdd(HostPlatformAndSteps.Key);
+			if(HostPlatformAndSteps.Value.IsValid() && HostPlatformAndSteps.Value->Type == EJson::Array)
+			{
+				const TArray<TSharedPtr<FJsonValue>>& CommandsArray = HostPlatformAndSteps.Value->AsArray();
+				for(const TSharedPtr<FJsonValue>& CommandValue: CommandsArray)
+				{
+					if(CommandValue->Type == EJson::String)
+					{
+						Commands.Add(CommandValue->AsString());
+					}
+				}
+			}
+		}
+	}
+}
+
+void FCustomBuildSteps::Write(TJsonWriter<>& Writer, const FString& FieldName) const
+{
+	Writer.WriteObjectStart(FieldName);
+	for(const TPair<FString, TArray<FString>>& HostPlatformAndCommands: HostPlatformToCommands)
+	{
+		Writer.WriteArrayStart(HostPlatformAndCommands.Key);
+		for(const FString& Command : HostPlatformAndCommands.Value)
+		{
+			Writer.WriteValue(*Command);
+		}
+		Writer.WriteArrayEnd();
+	}
+	Writer.WriteObjectEnd();
+}
+
+#undef LOCTEXT_NAMESPACE
+
diff --git a/Engine/Source/Runtime/Projects/Private/PluginDescriptor.cpp b/Engine/Source/Runtime/Projects/Private/PluginDescriptor.cpp
--- a/Engine/Source/Runtime/Projects/Private/PluginDescriptor.cpp
+++ b/Engine/Source/Runtime/Projects/Private/PluginDescriptor.cpp
@@ -108,6 +108,9 @@ bool FPluginDescriptor::Read(const FString& Text, FText& OutFailReason)
 		bRequiresBuildPlatform = true;
 	}
 
+	PreBuildSteps.Read(Object, TEXT("PreBuildSteps"));
+	PostBuildSteps.Read(Object, TEXT("PostBuildSteps"));
+
 	return true;
 }
 
@@ -156,6 +159,16 @@ FString FPluginDescriptor::ToString() const
 		Writer.WriteValue(TEXT("RequiresBuildPlatform"), bRequiresBuildPlatform);
 	}
 
+	if(!PreBuildSteps.IsEmpty())
+	{
+		PreBuildSteps.Write(Writer, TEXT("PreBuildSteps"));
+	}
+
+	if(!PostBuildSteps.IsEmpty())
+	{
+		PostBuildSteps.Write(Writer, TEXT("PostBuildSteps"));
+	}
+
 	Writer.WriteObjectEnd();
 	Writer.Close();
 
@@ -195,6 +208,28 @@ bool FPluginReferenceDescriptor::IsEnabledForPlatform( const FString& Platform )
 	return true;
 }
 
+bool FPluginReferenceDescriptor::IsEnabledForTarget(const FString& Target) const
+{
+    // If it's not enabled at all, return false
+    if (!bEnabled)
+    {
+        return false;
+    }
+
+    // If there is a list of whitelisted platforms, and this isn't one of them, return false
+    if (WhitelistTargets.Num() > 0 && !WhitelistTargets.Contains(Target))
+    {
+        return false;
+    }
+
+    // If this platform is blacklisted, also return false
+    if (BlacklistTargets.Contains(Target))
+    {
+        return false;
+    }
+
+    return true;
+}
 
 bool FPluginReferenceDescriptor::Read( const FJsonObject& Object, FText& OutFailReason )
 {
@@ -223,6 +258,10 @@ bool FPluginReferenceDescriptor::Read( const FJsonObject& Object, FText& OutFail
 	Object.TryGetStringArrayField(TEXT("WhitelistPlatforms"), WhitelistPlatforms);
 	Object.TryGetStringArrayField(TEXT("BlacklistPlatforms"), BlacklistPlatforms);
 
+	// Get the target lists
+	Object.TryGetStringArrayField(TEXT("WhitelistTargets"), WhitelistTargets);
+	Object.TryGetStringArrayField(TEXT("BlacklistTargets"), BlacklistTargets);
+
 	return true;
 }
 
@@ -300,6 +339,30 @@ void FPluginReferenceDescriptor::Write( TJsonWriter<>& Writer ) const
 		Writer.WriteArrayEnd();
 	}
 
+	if (WhitelistTargets.Num() > 0)
+	{
+		Writer.WriteArrayStart(TEXT("WhitelistTargets"));
+
+		for (int Idx = 0; Idx < WhitelistTargets.Num(); Idx++)
+		{
+			Writer.WriteValue(WhitelistTargets[Idx]);
+		}
+
+		Writer.WriteArrayEnd();
+	}
+
+	if (BlacklistTargets.Num() > 0)
+	{
+		Writer.WriteArrayStart(TEXT("BlacklistTargets"));
+
+		for (int Idx = 0; Idx < BlacklistTargets.Num(); Idx++)
+		{
+			Writer.WriteValue(BlacklistTargets[Idx]);
+		}
+
+		Writer.WriteArrayEnd();
+	}
+
 	Writer.WriteObjectEnd();
 }
 
diff --git a/Engine/Source/Runtime/Projects/Private/PluginManager.cpp b/Engine/Source/Runtime/Projects/Private/PluginManager.cpp
--- a/Engine/Source/Runtime/Projects/Private/PluginManager.cpp
+++ b/Engine/Source/Runtime/Projects/Private/PluginManager.cpp
@@ -376,7 +376,8 @@ bool FPluginManager::ConfigureEnabledPlugins()
 			TArray<FPluginReferenceDescriptor> PluginsCopy = Project->Plugins;
 			for(const FPluginReferenceDescriptor& Plugin: PluginsCopy)
 			{
-				if ((Plugin.bEnabled && !FindPluginInstance(Plugin.Name).IsValid() && !Plugin.bOptional) &&
+				bool bShouldBeEnabled = Plugin.bEnabled && Plugin.IsEnabledForPlatform(FPlatformMisc::GetUBTPlatform()) && Plugin.IsEnabledForTarget(FPlatformMisc::GetUBTTarget());
+				if ((bShouldBeEnabled && !FindPluginInstance(Plugin.Name).IsValid() && !Plugin.bOptional) &&
 					 (!IS_PROGRAM || AllEnabledPlugins.Contains(Plugin.Name))) // skip if this is a program and the plugin is not enabled
 				{
 					FText Caption(LOCTEXT("PluginMissingCaption", "Plugin missing"));
diff --git a/Engine/Source/Runtime/Projects/Private/ProjectDescriptor.cpp b/Engine/Source/Runtime/Projects/Private/ProjectDescriptor.cpp
--- a/Engine/Source/Runtime/Projects/Private/ProjectDescriptor.cpp
+++ b/Engine/Source/Runtime/Projects/Private/ProjectDescriptor.cpp
@@ -124,6 +124,11 @@ bool FProjectDescriptor::Read(const FJsonObject& Object, FText& OutFailReason)
 
 	// Get the sample name hash
 	Object.TryGetNumberField(TEXT("EpicSampleNameHash"), EpicSampleNameHash);
+
+	// Read the custom build steps
+	PreBuildSteps.Read(Object, TEXT("PreBuildSteps"));
+	PostBuildSteps.Read(Object, TEXT("PostBuildSteps"));
+
 	return true;
 }
 
@@ -164,7 +169,7 @@ void FProjectDescriptor::Write(TJsonWriter<>& Writer) const
 	FPluginReferenceDescriptor::WriteArray(Writer, TEXT("Plugins"), Plugins);
 
 	// Write the target platforms
-	if ( TargetPlatforms.Num() > 0 )
+	if(TargetPlatforms.Num() > 0)
 	{
 		Writer.WriteArrayStart(TEXT("TargetPlatforms"));
 		for(int Idx = 0; Idx < TargetPlatforms.Num(); Idx++)
@@ -180,6 +185,16 @@ void FProjectDescriptor::Write(TJsonWriter<>& Writer) const
 		Writer.WriteValue(TEXT("EpicSampleNameHash"), FString::Printf(TEXT("%u"), EpicSampleNameHash));
 	}
 
+	// Write the custom build steps
+	if(!PreBuildSteps.IsEmpty())
+	{
+		PreBuildSteps.Write(Writer, TEXT("PreBuildSteps"));
+	}
+	if(!PostBuildSteps.IsEmpty())
+	{
+		PostBuildSteps.Write(Writer, TEXT("PostBuildSteps"));
+	}
+
 	Writer.WriteObjectEnd();
 }
 
diff --git a/Engine/Source/Runtime/Projects/Private/ProjectManager.cpp b/Engine/Source/Runtime/Projects/Private/ProjectManager.cpp
--- a/Engine/Source/Runtime/Projects/Private/ProjectManager.cpp
+++ b/Engine/Source/Runtime/Projects/Private/ProjectManager.cpp
@@ -23,6 +23,10 @@ bool FProjectManager::LoadProjectFile( const FString& InProjectFile )
 	TSharedPtr<FProjectDescriptor> Descriptor = MakeShareable(new FProjectDescriptor());
 	if(Descriptor->Load(InProjectFile, FailureReason))
 	{
+		Descriptor->Save("D:\\P4 Streams\\Templates\\TP_Flying\\Out.json", FailureReason);
+
+
+
 		// Create the project
 		CurrentProject = Descriptor;
 		return true;
@@ -248,7 +252,7 @@ void FProjectManager::GetEnabledPlugins(TArray<FString>& OutPluginNames) const
 	{
 		for(const FPluginReferenceDescriptor& Plugin: Project->Plugins)
 		{
-			if(Plugin.IsEnabledForPlatform(FPlatformMisc::GetUBTPlatform()))
+			if(Plugin.IsEnabledForPlatform(FPlatformMisc::GetUBTPlatform()) && Plugin.IsEnabledForTarget(FPlatformMisc::GetUBTTarget()))
 			{
 				OutPluginNames.AddUnique(Plugin.Name);
 			}
diff --git a/Engine/Source/Runtime/Projects/Public/CustomBuildSteps.h b/Engine/Source/Runtime/Projects/Public/CustomBuildSteps.h
new file mode 100644
--- /dev/null
+++ b/Engine/Source/Runtime/Projects/Public/CustomBuildSteps.h
@@ -0,0 +1,21 @@
+// Copyright 1998-2016 Epic Games, Inc. All Rights Reserved.
+
+#pragma once
+
+/**
+ * Descriptor for projects. Contains all the information contained within a .uproject file.
+ */
+struct PROJECTS_API FCustomBuildSteps
+{
+	/** Mapping from host platform to list of commands */
+	TMap<FString, TArray<FString>> HostPlatformToCommands;
+
+	/** Whether this object is empty */
+	bool IsEmpty() const;
+
+	/** Reads the descriptor from the given JSON object */
+	void Read(const FJsonObject& Object, const FString& FieldName);
+
+	/** Writes the descriptor to the given JSON object */
+	void Write(TJsonWriter<>& Writer, const FString& FieldName) const;
+};
diff --git a/Engine/Source/Runtime/Projects/Public/PluginDescriptor.h b/Engine/Source/Runtime/Projects/Public/PluginDescriptor.h
--- a/Engine/Source/Runtime/Projects/Public/PluginDescriptor.h
+++ b/Engine/Source/Runtime/Projects/Public/PluginDescriptor.h
@@ -3,6 +3,7 @@
 #pragma once
 
 #include "ModuleDescriptor.h"
+#include "CustomBuildSteps.h"
 
 /**
  * Version numbers for plugin descriptors.
@@ -82,6 +83,12 @@ struct PROJECTS_API FPluginDescriptor
 	/** For plugins that are under a platform folder (eg. /PS4/), determines whether compiling the plugin requires the build platform and/or SDK to be available */
 	bool bRequiresBuildPlatform;
 
+	/** Pre-build steps for each host platform */
+	FCustomBuildSteps PreBuildSteps;
+
+	/** Pre-build steps for each host platform */
+	FCustomBuildSteps PostBuildSteps;
+
 	/** Constructor. */
 	FPluginDescriptor();
 
@@ -123,6 +130,12 @@ struct PROJECTS_API FPluginReferenceDescriptor
 
 	/** If enabled, list of platforms for which the plugin should be disabled. */
 	TArray<FString> BlacklistPlatforms;
+ 
+	/** If enabled, list of targets for which the plugin should be enabled (or all targets if blank). */
+	TArray<FString> WhitelistTargets;
+
+	/** If enabled, list of targets for which the plugin should be disabled. */
+	TArray<FString> BlacklistTargets;
 
 	/** Constructor */
 	FPluginReferenceDescriptor(const FString& InName = TEXT(""), bool bInEnabled = false, const FString& InMarketplaceURL = TEXT(""));
@@ -130,6 +143,9 @@ struct PROJECTS_API FPluginReferenceDescriptor
 	/** Determines whether the plugin is enabled for the given platform */
 	bool IsEnabledForPlatform(const FString& Platform) const;
 
+	/** Determines whether the plugin is enabled for the given target */
+	bool IsEnabledForTarget(const FString& Target) const;
+
 	/** Reads the descriptor from the given JSON object */
 	bool Read(const FJsonObject& Object, FText& OutFailReason);
 
diff --git a/Engine/Source/Runtime/Projects/Public/ProjectDescriptor.h b/Engine/Source/Runtime/Projects/Public/ProjectDescriptor.h
--- a/Engine/Source/Runtime/Projects/Public/ProjectDescriptor.h
+++ b/Engine/Source/Runtime/Projects/Public/ProjectDescriptor.h
@@ -4,6 +4,7 @@
 
 #include "ModuleDescriptor.h"
 #include "PluginDescriptor.h"
+#include "CustomBuildSteps.h"
 
 /**
  * Version numbers for project descriptors.
@@ -81,6 +82,12 @@ struct PROJECTS_API FProjectDescriptor
 	/** A hash that is used to determine if the project was forked from a sample */
 	uint32 EpicSampleNameHash;
 
+	/** Custom steps to execute before building targets in this project */
+	FCustomBuildSteps PreBuildSteps;
+
+	/** Custom steps to execute after building targets in this project */
+	FCustomBuildSteps PostBuildSteps;
+
 	/** Constructor. */
 	FProjectDescriptor();
 
