use HOME (osx) and USERPROFILE (windows) on appropriate target platform

diff --git a/Engine/Source/Programs/UnrealBuildTool/HTML5/HTML5SDKInfo.cs b/Engine/Source/Programs/UnrealBuildTool/HTML5/HTML5SDKInfo.cs
--- a/Engine/Source/Programs/UnrealBuildTool/HTML5/HTML5SDKInfo.cs
+++ b/Engine/Source/Programs/UnrealBuildTool/HTML5/HTML5SDKInfo.cs
@@ -106,6 +106,22 @@ namespace UnrealBuildTool
 			return TempPath;
 		}
 
+		static string PLATFORM_USER_HOME
+		{
+			get
+			{
+				switch (BuildHostPlatform.Current.Platform)
+				{
+					case UnrealTargetPlatform.Win64:
+						return "USERPROFILE";
+					case UnrealTargetPlatform.Mac:
+						return "HOME";
+					default:
+						return "error_unknown_platform";
+				}
+			}
+		}
+
         public static string SetUpEmscriptenConfigFile()
 		{
 			// make a fresh .emscripten resource file
@@ -119,12 +135,12 @@ namespace UnrealBuildTool
 			// --------------------------------------------------
 			// save a few things 
 			string PATH_SAVE = Environment.GetEnvironmentVariable("PATH");
-			string HOME_SAVE = Environment.GetEnvironmentVariable("USERPROFILE");
+			string HOME_SAVE = Environment.GetEnvironmentVariable(PLATFORM_USER_HOME);
 			// warm up the .emscripten resource file
 			string NODE_ROOT = Path.GetDirectoryName(NODE_JS);
 			string PYTHON_ROOT = Path.GetDirectoryName(PYTHON);
 			Environment.SetEnvironmentVariable("PATH", NODE_ROOT + Path.PathSeparator + LLVM_ROOT + Path.PathSeparator + PYTHON_ROOT + Path.PathSeparator + EMSCRIPTEN_ROOT + Path.PathSeparator + PATH_SAVE);
-			Environment.SetEnvironmentVariable("USERPROFILE", HTML5Intermediatory);
+			Environment.SetEnvironmentVariable(PLATFORM_USER_HOME, HTML5Intermediatory);
 			// --------------------------------------------------
 				string cmd = "\"" + Path.Combine(EMSCRIPTEN_ROOT, "emcc") + "\"";
 				ProcessStartInfo processInfo = new ProcessStartInfo(PYTHON, cmd + " -v");
@@ -151,7 +167,7 @@ namespace UnrealBuildTool
 			// --------------------------------------------------
 			// --------------------------------------------------
 			// restore a few things
-			Environment.SetEnvironmentVariable("USERPROFILE", HOME_SAVE);
+			Environment.SetEnvironmentVariable(PLATFORM_USER_HOME, HOME_SAVE);
 			Environment.SetEnvironmentVariable("PATH", PATH_SAVE);
 			// --------------------------------------------------
 
