UBT: Hack to allow separable distribution of OVR SDK (don't build disabled Oculus plugins)

diff --git a/Engine/Source/Programs/UnrealBuildTool/Configuration/UEBuildTarget.cs b/Engine/Source/Programs/UnrealBuildTool/Configuration/UEBuildTarget.cs
--- a/Engine/Source/Programs/UnrealBuildTool/Configuration/UEBuildTarget.cs
+++ b/Engine/Source/Programs/UnrealBuildTool/Configuration/UEBuildTarget.cs
@@ -3635,6 +3635,21 @@ namespace UnrealBuildTool
 				BuildPlugins = new List<PluginInfo>(EnabledPlugins);
 			}
 
+			// @todo vreditor urgent: Don't build Oculus plugins that are disabled (for separable distribution)
+			{
+				for( int PluginIndex = 0; PluginIndex < BuildPlugins.Count; ++PluginIndex )
+				{
+					PluginInfo PluginInfo = BuildPlugins[ PluginIndex ];
+					if( PluginInfo.Name.StartsWith( "Oculus", StringComparison.InvariantCultureIgnoreCase ) )
+					{
+						if( !EnabledPlugins.Contains( PluginInfo ) )
+						{
+							BuildPlugins.RemoveAt( PluginIndex-- );
+						}
+					}
+				}
+			}
+
 			// Add any foreign plugins to the list
 			if (ForeignPlugins != null)
 			{
