Merging //UE4/Dev-Main to Dev-VREditor (//UE4/Dev-VREditor) (Up to CL 2942580)

diff --git a/Engine/Source/Programs/UnrealBuildTool/System/ActionGraph.cs b/Engine/Source/Programs/UnrealBuildTool/System/ActionGraph.cs
--- a/Engine/Source/Programs/UnrealBuildTool/System/ActionGraph.cs
+++ b/Engine/Source/Programs/UnrealBuildTool/System/ActionGraph.cs
@@ -366,7 +366,7 @@ namespace UnrealBuildTool
 				HashSet<Action> UnlinkedActions = ActionsToExecute.Where(Action => Action.ActionType == ActionType.Link && !ProducedItems.Overlaps(Action.PrerequisiteItems)).ToHashSet();
 
 				// Remove unlinked items
-				ActionsToExecute = new HashSet<Action>(ActionsToExecute.Except(UnlinkedActions));
+				ActionsToExecute.ExceptWith(UnlinkedActions);
 
 				// Re-add unlinked items which produce things which are dependencies of other actions
 				for (;;)
@@ -390,6 +390,9 @@ namespace UnrealBuildTool
 						break;
 					}
 				}
+
+				// Remove actions that are wholly dependent on unlinked actions
+				ActionsToExecute = ActionsToExecute.Where(Action => Action.PrerequisiteItems.Count == 0 || !Action.PrerequisiteItems.Select(Item => Item.ProducingAction).ToHashSet().IsSubsetOf(UnlinkedActions)).ToHashSet();
 			}
 
 			if (BuildConfiguration.bPrintPerformanceInfo)
