Copying //UE4/Dev-Build to //UE4/Main

==========================
MAJOR FEATURES + CHANGES
==========================

Change 2864843 on 2016/02/12 by [REDACTED_PERSON]

	Add individual 'status', 'outcome', and 'error_code' fields to parsed jobsteps. Should fix grid view not being able to display 'pending' icons.

Change 2865161 on 2016/02/12 by [REDACTED_PERSON]

	Stop storing a reference to UEBuildTarget from UEBuildModule. It creates an awkward cyclic data dependency, and makes it easy for people to write lazy code that just reaches into the internal state of the build.

Change 2865643 on 2016/02/12 by [REDACTED_PERSON]

	Rename UEBuildModuleType to UHTModuleType, and move implementation into ExternalExecution.

Change 2874408 on 2016/02/19 by [REDACTED_PERSON]

	Automatically sort nodes in the dashboard grid view by a weight derived from the node's order in the build graph, summed across all the jobs in which it was present.

Change 2879572 on 2016/02/24 by [REDACTED_PERSON]

	Allow spoofing a Git merge from a given commit, using a changelist description containing the tag "git merge <branch> <changelist>", where <branch> is the name of a branch on Git (eg. master, 4.11, etc..), and <changelist> is the changelist being merged in.

Change 2883216 on 2016/02/26 by [REDACTED_PERSON]

	Prevent Jira tickets being incorrectly updated with 'Main CL' fields which are after the 'Fix CL' fields.

Change 2883755 on 2016/02/26 by [REDACTED_PERSON]

	Fix solution files having a Shipping configuration, even when -NoShippingConfigs is passed on the command line.

Change 2886223 on 2016/02/29 by [REDACTED_PERSON]

	Ignore SignTool errors - we can recover from them.

Change 2887414 on 2016/03/01 by [REDACTED_PERSON]

	Dump all the *.crash files produced while running commandlets, to make it easier to diagnose build system crashes cooking on Mac.

Change 2888235 on 2016/03/01 by [REDACTED_PERSON]

	Add overloads for methods in FileFilter which take FileReference and DirectoryReference objects.

Change 2889602 on 2016/03/02 by [REDACTED_PERSON]

	Treat shaders as code in UGS. Don't sync them as part of content-only syncs, and don't allow syncing past them without updated binaries.

Change 2889610 on 2016/03/02 by [REDACTED_PERSON]

	Fix setting for using incremental builds not being saved. Also hide command to do incremental builds if the 'use incremental builds' option is not checked.

Change 2891866 on 2016/03/03 by [REDACTED_PERSON]

	Removed Rocket specific batch files and made sure installed build won't try to include them
	Removed last use of RocketGenerateProjectFiles.sh by using UBT directly instead

Change 2893349 on 2016/03/03 by [REDACTED_PERSON]

	Add derived ReplicatedBranch to support mirroring the VR editor branch to GitHub.

Change 2894703 on 2016/03/04 by [REDACTED_PERSON]

	Include *.usf when looking for the last code changelist. Also update version to 1.68.

Change 2897991 on 2016/03/07 by [REDACTED_PERSON]

	Copy the changelist number to the clipboard when the user presses Ctrl-C. Update version number to 1.69.

Change 2898005 on 2016/03/07 by [REDACTED_PERSON]

	Minor changes to support BuildGraph:

	* UE4Build now has a static function that can update version files.
	* Adding FileReference/DirectoryReference methods to FileFilter and CommandUtils.
	* FileFilter treats any pattern containing a slash as implictly starting from the root directory, unless it begins with "...".

Change 2898095 on 2016/03/07 by [REDACTED_PERSON]

	UAT - Don't retry builds if we're using local executor; we don't encounter failures due to timeouts.

Change 2898248 on 2016/03/07 by [REDACTED_PERSON]

	UBT - Add the standard game include paths back in to plugin modules. Existing game code relies on this.

Change 2898615 on 2016/03/08 by [REDACTED_PERSON]

	Removed last uses of RunningRocket function
	All seemed to be overly cautious about people using an Installed build to do non standard things, don't see any ill effects in the most common circumstances.

Change 2898681 on 2016/03/08 by [REDACTED_PERSON]

	Removed Automation.RunningRocket function as there are no more uses
	Changed the majority of comments referencing Rocket mode that are now either about the engine being installed or from the Launcher etc.

diff --git a/Engine/Build/BuildFarm/PostpExtensions.pl b/Engine/Build/BuildFarm/PostpExtensions.pl
--- a/Engine/Build/BuildFarm/PostpExtensions.pl
+++ b/Engine/Build/BuildFarm/PostpExtensions.pl
@@ -16,6 +16,10 @@
 	# Warning from Android monolithic builds about the javac version 
 	".*major version 51 is newer than 50.*",
 	
+	# SignTool errors which we can retry and recover from
+	"SignTool Error: ",
+	
+	
 #	".*ERROR: The process.*not found",
 #	".*ERROR: This operation returned because the timeout period expired.*",
 #	".*Sync.VerifyKnownFileInManifest: ERROR:.*",
diff --git a/Engine/Build/InstalledEngineFilters.ini b/Engine/Build/InstalledEngineFilters.ini
--- a/Engine/Build/InstalledEngineFilters.ini
+++ b/Engine/Build/InstalledEngineFilters.ini
@@ -49,7 +49,7 @@
 ; Shaders
 /Engine/Shaders/...
 
-; Exclude all the intermediate Rocket files in the Engine/Saved folder
+; Exclude all the intermediate files in the Engine/Saved folder
 -/Engine/Saved/...
 
 ; Valid templates will be included by name from the build script, but we need to exclude intermediates
@@ -74,7 +74,6 @@
 {Win64} /Engine/Build/BatchFiles/Build.bat
 {Win64} /Engine/Build/BatchFiles/Clean.bat
 {Win64} /Engine/Build/BatchFiles/Rebuild.bat
-{Win64} /Engine/Build/BatchFiles/RocketGenerateProjectFiles.bat
 {Win64} /Engine/Build/BatchFiles/RunUAT.bat
 {Win64} /Engine/Build/BatchFiles/GetVSComnToolsPath.bat
 {Win64} -/Engine/Build/Receipts/PS4MapFile*
@@ -100,8 +99,6 @@
 {Mac} /Engine/Build/BatchFiles/Mac/FixDependencyFiles.sh
 {Mac} /Engine/Build/BatchFiles/Mac/SetupMono.sh
 {Mac} /Engine/Build/BatchFiles/Mac/FixMonoFiles.sh
-{Mac} /Engine/Build/BatchFiles/Mac/RocketBuild.sh
-{Mac} /Engine/Build/BatchFiles/Mac/RocketGenerateProjectFiles.sh
 {Mac} /Engine/Build/BatchFiles/Mac/RunMono.sh
 {Mac} /Engine/Build/BatchFiles/Mac/RunXBuild.sh
 {Mac} /Engine/Build/BatchFiles/Mac/xcodeunlock.sh
diff --git a/Engine/Source/Developer/Android/AndroidPlatformEditor/Private/AndroidTargetSettingsCustomization.cpp b/Engine/Source/Developer/Android/AndroidPlatformEditor/Private/AndroidTargetSettingsCustomization.cpp
--- a/Engine/Source/Developer/Android/AndroidPlatformEditor/Private/AndroidTargetSettingsCustomization.cpp
+++ b/Engine/Source/Developer/Android/AndroidPlatformEditor/Private/AndroidTargetSettingsCustomization.cpp
@@ -215,7 +215,7 @@ void FAndroidTargetSettingsCustomization::BuildAppManifestSection(IDetailLayoutB
 	SETUP_SOURCEONLY_PROP(bBuildForES31, BuildCategory, LOCTEXT("BuildForES31ToolTip", "Enable OpenGL ES31 + AEP (Android Extension Pack) rendering support?"));
 	
 	// @todo android fat binary: Put back in when we expose those
-//	SETUP_NONROCKET_PROP(bSplitIntoSeparateApks, BuildCategory, LOCTEXT("SplitIntoSeparateAPKsToolTip", "If checked, CPU architectures and rendering types will be split into separate .apk files"));
+//	SETUP_SOURCEONLY_PROP(bSplitIntoSeparateApks, BuildCategory, LOCTEXT("SplitIntoSeparateAPKsToolTip", "If checked, CPU architectures and rendering types will be split into separate .apk files"));
 }
 
 void FAndroidTargetSettingsCustomization::BuildIconSection(IDetailLayoutBuilder& DetailLayout)
diff --git a/Engine/Source/Developer/DesktopPlatform/Private/DesktopPlatformBase.cpp b/Engine/Source/Developer/DesktopPlatform/Private/DesktopPlatformBase.cpp
--- a/Engine/Source/Developer/DesktopPlatform/Private/DesktopPlatformBase.cpp
+++ b/Engine/Source/Developer/DesktopPlatform/Private/DesktopPlatformBase.cpp
@@ -688,7 +688,7 @@ bool FDesktopPlatformBase::InvalidateMakefiles(const FString& RootDir, const FSt
 
 bool FDesktopPlatformBase::IsUnrealBuildToolAvailable()
 {
-	// If using Rocket and the Rocket unreal build tool executable exists, then UBT is available. Otherwise check it can be built.
+	// If using installed build and the unreal build tool executable exists, then UBT is available. Otherwise check it can be built.
 	if (FApp::IsEngineInstalled())
 	{
 		return FPaths::FileExists(GetUnrealBuildToolExecutableFilename(FPaths::RootDir()));
@@ -740,11 +740,11 @@ FProcHandle FDesktopPlatformBase::InvokeUnrealBuildToolAsync(const FString& InCm
 	// UnrealBuildTool is currently always located in the Binaries/DotNET folder
 	FString ExecutableFileName = GetUnrealBuildToolExecutableFilename(FPaths::RootDir());
 
-	// Rocket never builds UBT, UnrealBuildTool should already exist
+	// Installed builds never build UBT, UnrealBuildTool should already exist
 	bool bSkipBuild = FApp::IsEngineInstalled() || bSkipBuildUBT;
 	if (!bSkipBuild)
 	{
-		// When not using rocket, we should attempt to build UBT to make sure it is up to date
+		// When not using an installed build, we should attempt to build UBT to make sure it is up to date
 		// Only do this if we have not already successfully done it once during this session.
 		static bool bSuccessfullyBuiltUBTOnce = false;
 		if (!bSuccessfullyBuiltUBTOnce)
diff --git a/Engine/Source/Developer/HotReload/Private/HotReload.cpp b/Engine/Source/Developer/HotReload/Private/HotReload.cpp
--- a/Engine/Source/Developer/HotReload/Private/HotReload.cpp
+++ b/Engine/Source/Developer/HotReload/Private/HotReload.cpp
@@ -1501,7 +1501,7 @@ FString FHotReloadModule::MakeUBTArgumentsForModuleCompiling()
 		// We have to pass FULL paths to UBT
 		FString FullProjectPath = FPaths::ConvertRelativePathToFull(FPaths::GetProjectFilePath());
 
-		// @todo projectdirs: Currently non-Rocket projects that exist under the UE4 root are compiled by UBT with no .uproject file
+		// @todo projectdirs: Currently non-installed projects that exist under the UE4 root are compiled by UBT with no .uproject file
 		//     name passed in (see bIsProjectTarget in VCProject.cs), which causes intermediate libraries to be saved to the Engine
 		//     intermediate folder instead of the project's intermediate folder.  We're emulating this behavior here for module
 		//     recompiling, so that compiled modules will be able to find their import libraries in the original folder they were compiled.
diff --git a/Engine/Source/Developer/LauncherServices/Private/Launcher/LauncherWorker.h b/Engine/Source/Developer/LauncherServices/Private/Launcher/LauncherWorker.h
--- a/Engine/Source/Developer/LauncherServices/Private/Launcher/LauncherWorker.h
+++ b/Engine/Source/Developer/LauncherServices/Private/Launcher/LauncherWorker.h
@@ -24,7 +24,6 @@ public:
 	 *
 	 * @param InProfile - The profile to process.
 	 * @param InDeviceProxyManager - The target device proxy manager to use.
-	 * @param LaunchType - The launch type, i.e. UE4 or Rocket.
 	 */
 	FLauncherWorker( const ITargetDeviceProxyManagerRef& InDeviceProxyManager, const ILauncherProfileRef& InProfile );
 
diff --git a/Engine/Source/Developer/LauncherServices/Private/Profiles/LauncherProfile.h b/Engine/Source/Developer/LauncherServices/Private/Profiles/LauncherProfile.h
--- a/Engine/Source/Developer/LauncherServices/Private/Profiles/LauncherProfile.h
+++ b/Engine/Source/Developer/LauncherServices/Private/Profiles/LauncherProfile.h
@@ -1350,7 +1350,6 @@ public:
 		"runautomationtest", ""
 		"runautomationtests", "true/false"
 		"skipserver", "true/false"
-		"rocket", "true/false"
 		"ue4exe", ""
 		"unattended", "true/false"
 		"deviceuser", ""
diff --git a/Engine/Source/Developer/LauncherServices/Public/Interfaces/ILauncherProfile.h b/Engine/Source/Developer/LauncherServices/Public/Interfaces/ILauncherProfile.h
--- a/Engine/Source/Developer/LauncherServices/Public/Interfaces/ILauncherProfile.h
+++ b/Engine/Source/Developer/LauncherServices/Public/Interfaces/ILauncherProfile.h
@@ -638,7 +638,7 @@ public:
 	 * Gets the full path to the Unreal project to use.
 	 *
 	 * @return The path.
-	 * @see SetRocketProjactPath
+	 * @see SetProjectPath
 	 */
 	virtual FString GetProjectPath( ) const = 0;
 
@@ -1083,10 +1083,10 @@ public:
 	virtual void FallbackProjectUpdated() = 0;
 
 	/**
-	 * Sets the path to the Rocket project to use.
+	 * Sets the path to the Unreal project to use.
 	 *
 	 * @param Path The full path to the project.
-	 * @see GetRocketProjectPath
+	 * @see GetProjectPath
 	 */
 	virtual void SetProjectPath( const FString& Path ) = 0;
 
diff --git a/Engine/Source/Developer/SimplygonMeshReduction/Private/SimplygonMeshReduction.cpp b/Engine/Source/Developer/SimplygonMeshReduction/Private/SimplygonMeshReduction.cpp
--- a/Engine/Source/Developer/SimplygonMeshReduction/Private/SimplygonMeshReduction.cpp
+++ b/Engine/Source/Developer/SimplygonMeshReduction/Private/SimplygonMeshReduction.cpp
@@ -516,7 +516,7 @@ public:
 			DllFilename = FPaths::Combine(*DllPath, TEXT("SimplygonSDKRuntimeReleasex64.dll"));
 		}
 
-		// If the DLL just doesn't exist, fail gracefully. Licensees and Rocket users will not necessarily have Simplygon.
+		// If the DLL just doesn't exist, fail gracefully. Licensees and Subscribers will not necessarily have Simplygon.
 		if( !FPaths::FileExists(DllFilename) )
 		{
 			UE_LOG(LogSimplygon,Warning,TEXT("Simplygon DLL not present - disabling."));
diff --git a/Engine/Source/Editor/MainFrame/Private/Frame/MainFrameActions.cpp b/Engine/Source/Editor/MainFrame/Private/Frame/MainFrameActions.cpp
--- a/Engine/Source/Editor/MainFrame/Private/Frame/MainFrameActions.cpp
+++ b/Engine/Source/Editor/MainFrame/Private/Frame/MainFrameActions.cpp
@@ -287,7 +287,7 @@ void FMainFrameActionCallbacks::AddCodeToProject()
 const TCHAR* GetUATCompilationFlags()
 {
 	// We never want to compile editor targets when invoking UAT in this context.
-	// If we are rocket or don't have a compiler, we must assume we have a precompiled UAT.
+	// If we are installed or don't have a compiler, we must assume we have a precompiled UAT.
 	return (FApp::GetEngineIsPromotedBuild() || FApp::IsEngineInstalled())
 		? TEXT("-nocompile -nocompileeditor")
 		: TEXT("-nocompileeditor");
diff --git a/Engine/Source/Editor/UnrealEd/Private/EditorEngine.cpp b/Engine/Source/Editor/UnrealEd/Private/EditorEngine.cpp
--- a/Engine/Source/Editor/UnrealEd/Private/EditorEngine.cpp
+++ b/Engine/Source/Editor/UnrealEd/Private/EditorEngine.cpp
@@ -6103,8 +6103,8 @@ void UEditorEngine::UpdateAutoLoadProject()
 {
 	// If the recent project file exists and is non-empty, update the contents with the currently loaded .uproject
 	// If the recent project file exists and is empty, recent project files should not be auto-loaded
-	// If the recent project file does not exist, auto-populate it with the currently loaded project in Rocket and auto-populate empty in non-rocket
-	//		In Rocket we default to auto-loading, in non-Rocket we default to opting out of auto loading
+	// If the recent project file does not exist, auto-populate it with the currently loaded project in installed builds and auto-populate empty in non-installed
+	//		In installed builds we default to auto-loading, in non-installed we default to opting out of auto loading
 	const FString& AutoLoadProjectFileName = IProjectManager::Get().GetAutoLoadProjectFileName();
 	FString RecentProjectFileContents;
 	bool bShouldLoadRecentProjects = false;
diff --git a/Engine/Source/Editor/UnrealEd/Private/Lightmass/Lightmass.cpp b/Engine/Source/Editor/UnrealEd/Private/Lightmass/Lightmass.cpp
--- a/Engine/Source/Editor/UnrealEd/Private/Lightmass/Lightmass.cpp
+++ b/Engine/Source/Editor/UnrealEd/Private/Lightmass/Lightmass.cpp
@@ -2562,7 +2562,7 @@ bool FLightmassProcessor::BeginRun()
 #endif
 	const int32 RequiredDependencyPaths64Count = ARRAY_COUNT(RequiredDependencyPaths64);
 
-	// Set up optional dependencies.  These might not exist in Rocket distributions, for example.
+	// Set up optional dependencies.  These might not exist in Launcher distributions, for example.
 	const TCHAR* OptionalDependencyPaths32[] =
 	{
 		TEXT("../Win32/UnrealLightmass.pdb"),
diff --git a/Engine/Source/Programs/AutomationTool/AutomationUtils/Automation.cs b/Engine/Source/Programs/AutomationTool/AutomationUtils/Automation.cs
--- a/Engine/Source/Programs/AutomationTool/AutomationUtils/Automation.cs
+++ b/Engine/Source/Programs/AutomationTool/AutomationUtils/Automation.cs
@@ -115,7 +115,6 @@ namespace AutomationTool
         public static CommandLineArg NoCompileEditor = new CommandLineArg("-NoCompileEditor");
 		public static CommandLineArg Help = new CommandLineArg("-Help");
 		public static CommandLineArg List = new CommandLineArg("-List");
-		public static CommandLineArg Rocket = new CommandLineArg("-Rocket");
 		public static CommandLineArg VS2015 = new CommandLineArg("-2015");
 		public static CommandLineArg NoKill = new CommandLineArg("-NoKill");
 		public static CommandLineArg Installed = new CommandLineArg("-Installed");
@@ -611,23 +610,6 @@ AutomationTool.exe [-verbose] [-compileonly] [-p4] Command0 [-Arg0 -Arg1 -Arg2 
 
 		#region HelperFunctions
 
-		/// <summary>
-		/// Returns true if Automation tool is running with Rocket mode enabled
-		/// </summary>
-		/// <returns>True if running in Rocket mode</returns>
-		static public bool RunningRocket()
-		{
-			if (!bRunningRocket.HasValue)
-			{
-				bRunningRocket = GlobalCommandLine.Rocket;
-				string RocketFile = Path.Combine(CommandUtils.CmdEnv.LocalRoot, "Engine", "Build", "Rocket.txt");
-				bRunningRocket |= File.Exists(RocketFile);
-			}
-
-			return bRunningRocket.Value;
-		}
-		static private bool? bRunningRocket;
-
 		/// <summary>
 		/// Returns true if AutomationTool is running using installed Engine components
 		/// </summary>
@@ -636,9 +618,17 @@ AutomationTool.exe [-verbose] [-compileonly] [-p4] Command0 [-Arg0 -Arg1 -Arg2 
 		{
 			if (!bIsEngineInstalled.HasValue)
 			{
-				bIsEngineInstalled = GlobalCommandLine.Installed || (Automation.RunningRocket() ? !GlobalCommandLine.NotInstalledEngine : GlobalCommandLine.InstalledEngine);
+				bIsEngineInstalled = GlobalCommandLine.Installed;
 				string InstalledBuildFile = Path.Combine(CommandUtils.CmdEnv.LocalRoot, "Engine", "Build", "InstalledBuild.txt");
 				bIsEngineInstalled |= File.Exists(InstalledBuildFile);
+				if (bIsEngineInstalled.Value)
+				{
+					bIsEngineInstalled = !GlobalCommandLine.NotInstalledEngine;
+				}
+				else
+				{
+					bIsEngineInstalled = GlobalCommandLine.InstalledEngine;
+				}
 			}
 
 			return bIsEngineInstalled.Value;
diff --git a/Engine/Source/Programs/AutomationTool/AutomationUtils/CommandUtils.cs b/Engine/Source/Programs/AutomationTool/AutomationUtils/CommandUtils.cs
--- a/Engine/Source/Programs/AutomationTool/AutomationUtils/CommandUtils.cs
+++ b/Engine/Source/Programs/AutomationTool/AutomationUtils/CommandUtils.cs
@@ -1903,8 +1903,8 @@ namespace AutomationTool
 		/// <summary>
 		/// Command line parameters for this command (empty by non-null by default)
 		/// </summary>
-		private object[] CommandLineParams = new object[0];
-		public object[] Params
+		private string[] CommandLineParams = new string[0];
+		public string[] Params
 		{
 			get { return CommandLineParams; }
 			set { CommandLineParams = value; }
@@ -2090,13 +2090,13 @@ namespace AutomationTool
 					FilesList += " \"" + FilteredFile + "\"";
 					if (FilesList.Length > 32000)
 					{
-						CommandUtils.RunAndLog(CommandUtils.CmdEnv, "zip", "-g -q " + ZipFileName + FilesList);
+						CommandUtils.RunAndLog(CommandUtils.CmdEnv, "zip", "-g -q \"" + ZipFileName + "\"" + FilesList);
 						FilesList = "";
 					}
 				}
 				if (FilesList.Length > 0)
 				{
-					CommandUtils.RunAndLog(CommandUtils.CmdEnv, "zip", "-g -q " + ZipFileName + FilesList);
+					CommandUtils.RunAndLog(CommandUtils.CmdEnv, "zip", "-g -q \"" + ZipFileName + "\"" + FilesList);
 				}
 				CommandUtils.PopDir();
 			}
@@ -2113,6 +2113,48 @@ namespace AutomationTool
 			}
 		}
 
+		/// <summary>
+		/// Creates a zip file containing the given input files
+		/// </summary>
+		/// <param name="ZipFile">Filename for the zip</param>
+		/// <param name="BaseDirectory">Base directory to store relative paths in the zip file to</param>
+		/// <param name="Files">Files to include in the archive</param>
+		public static void ZipFiles(FileReference ZipFile, DirectoryReference BaseDirectory, IEnumerable<FileReference> Files)
+		{
+			// Ionic.Zip.Zip64Option.Always option produces broken archives on Mono, so we use system zip tool instead
+			if (Utils.IsRunningOnMono)
+			{
+				CommandUtils.CreateDirectory(ZipFile.Directory.FullName);
+ 				CommandUtils.PushDir(BaseDirectory.FullName);
+ 				string FilesList = "";
+				foreach(FileReference File in Files)
+				{
+					FilesList += " \"" + File.MakeRelativeTo(BaseDirectory) + "\"";
+					if (FilesList.Length > 32000)
+					{
+						CommandUtils.RunAndLog(CommandUtils.CmdEnv, "zip", "-g -q \"" + ZipFile.FullName + "\"" + FilesList);
+						FilesList = "";
+					}
+				}
+				if (FilesList.Length > 0)
+				{
+					CommandUtils.RunAndLog(CommandUtils.CmdEnv, "zip", "-g -q \"" + ZipFile.FullName + "\"" + FilesList);
+				}
+				CommandUtils.PopDir();
+			}
+			else
+			{
+				Ionic.Zip.ZipFile Zip = new Ionic.Zip.ZipFile();
+				Zip.UseZip64WhenSaving = Ionic.Zip.Zip64Option.Always;
+				foreach(FileReference File in Files)
+				{
+					Zip.AddFile(File.FullName, Path.GetDirectoryName(File.MakeRelativeTo(BaseDirectory)));
+				}
+				CommandUtils.CreateDirectory(ZipFile.Directory.FullName);
+				Zip.Save(ZipFile.FullName);
+			}
+		}
+
 		/// <summary>
 		/// Extracts the contents of a zip file
 		/// </summary>
diff --git a/Engine/Source/Programs/AutomationTool/AutomationUtils/CommandletUtils.cs b/Engine/Source/Programs/AutomationTool/AutomationUtils/CommandletUtils.cs
--- a/Engine/Source/Programs/AutomationTool/AutomationUtils/CommandletUtils.cs
+++ b/Engine/Source/Programs/AutomationTool/AutomationUtils/CommandletUtils.cs
@@ -4,6 +4,7 @@ using System.Collections.Generic;
 using System.Linq;
 using System.Text;
 using System.IO;
+using System.Threading;
 using UnrealBuildTool;
 
 
@@ -229,6 +230,8 @@ namespace AutomationTool
 			
 			PushDir(CWD);
 
+			DateTime StartTime = DateTime.UtcNow;
+
 			string LocalLogFile = LogUtils.GetUniqueLogName(CombinePaths(CmdEnv.EngineSavedFolder, Commandlet));
 			Log("Commandlet log file is {0}", LocalLogFile);
 			string Args = String.Format(
@@ -262,6 +265,58 @@ namespace AutomationTool
 				}
 			}
 
+			// If we're running on a Mac, dump all the *.crash files that were generated while the editor was running.
+			if(HostPlatform.Current.HostEditorPlatform == UnrealTargetPlatform.Mac)
+			{
+				// If the exit code indicates the main process crashed, introduce a small delay because the crash report is written asynchronously.
+				// If we exited normally, still check without waiting in case SCW or some other child process crashed.
+				if(RunResult.ExitCode > 128)
+				{
+					CommandUtils.Log("Pausing before checking for crash logs...");
+					Thread.Sleep(10 * 1000);
+				}
+				
+				// Create a list of directories containing crash logs, and add the system log folder
+				List<string> CrashDirs = new List<string>();
+				CrashDirs.Add("/Library/Logs/DiagnosticReports");
+					
+				// Add the user's log directory too
+				string HomeDir = Environment.GetEnvironmentVariable("HOME");
+				if(!String.IsNullOrEmpty(HomeDir))
+				{
+					CrashDirs.Add(Path.Combine(HomeDir, "Library/Logs/DiagnosticReports"));
+				}
+
+				// Check each directory for crash logs
+				List<FileInfo> CrashFileInfos = new List<FileInfo>();
+				foreach(string CrashDir in CrashDirs)
+				{
+					DirectoryInfo CrashDirInfo = new DirectoryInfo(CrashDir);
+					if(CrashDirInfo.Exists)
+					{
+						CrashFileInfos.AddRange(CrashDirInfo.EnumerateFiles("*.crash", SearchOption.TopDirectoryOnly).Where(x => x.LastWriteTimeUtc >= StartTime));
+					}
+				}
+				
+				// Dump them all to the log
+				foreach(FileInfo CrashFileInfo in CrashFileInfos)
+				{
+					CommandUtils.LogWarning("Found crash log - {0}:", CrashFileInfo.FullName);
+					try
+					{
+						string[] Lines = File.ReadAllLines(CrashFileInfo.FullName);
+						foreach(string Line in Lines)
+						{
+							CommandUtils.Log("CrashDump: {0}", Line);
+						}
+					}
+					catch(Exception Ex)
+					{
+						CommandUtils.LogWarning("Failed to read file ({0})", Ex.Message);
+					}
+				}
+			}
+
 			// Copy the local commandlet log to the destination folder.
 			string DestLogFile = LogUtils.GetUniqueLogName(CombinePaths(CmdEnv.LogFolder, Commandlet));
 			if (!CommandUtils.CopyFile_NoExceptions(LocalLogFile, DestLogFile))
@@ -293,7 +348,7 @@ namespace AutomationTool
 				throw new AutomationException("BUILD FAILED: Failed while running {0} for {1}; see log {2}", Commandlet, ProjectName, DestLogFile);
 			}
 		}
-
+		
 		/// <summary>
 		/// Returns the default path of the editor executable to use for running commandlets.
 		/// </summary>
diff --git a/Engine/Source/Programs/AutomationTool/AutomationUtils/FileFilter.cs b/Engine/Source/Programs/AutomationTool/AutomationUtils/FileFilter.cs
--- a/Engine/Source/Programs/AutomationTool/AutomationUtils/FileFilter.cs
+++ b/Engine/Source/Programs/AutomationTool/AutomationUtils/FileFilter.cs
@@ -8,6 +8,7 @@ using System.IO;
 using System.Linq;
 using System.Text;
 using System.Text.RegularExpressions;
+using UnrealBuildTool;
 
 namespace AutomationTool
 {
@@ -136,6 +137,19 @@ namespace AutomationTool
 			}
 		}
 
+		/// <summary>
+		/// Adds several rules to the filter
+		/// </summary>
+		/// <param name="Rules">List of patterns to match.</param>
+		/// <param name="Type">Whether the rules are include or exclude rules</param>
+		public void AddRules(IEnumerable<string> Rules, FileFilterType Type)
+		{
+			foreach(string Rule in Rules)
+			{
+				AddRule(Rule, Type);
+			}
+		}
+
 		/// <summary>
 		/// Adds several rules in the given lines. Rules may be prefixed with conditions of the syntax {Key=Value, Key2=Value2}, which
 		/// will be evaluated using variables in the given dictionary before being added.
@@ -176,6 +190,20 @@ namespace AutomationTool
 			}
 		}
 
+		/// <summary>
+		/// Adds rules for files which match the given names
+		/// </summary>
+		/// <param name="FileName">The filenames to rules for</param>
+		/// <param name="BaseDirectoryName">Base directory for the rules</param>
+		/// <param name="Type">Whether to add an include or exclude rule</param>
+		public void AddRuleForFiles(IEnumerable<FileReference> Files, DirectoryReference BaseDirectory, FileFilterType Type)
+		{
+			foreach (FileReference File in Files)
+			{
+				AddRule("/" + File.MakeRelativeTo(BaseDirectory), Type);
+			}
+		}
+
 		/// <summary>
 		/// Reads a configuration file split into sections
 		/// </summary>
@@ -253,12 +281,15 @@ namespace AutomationTool
 		{
 			string NormalizedPattern = Pattern.Replace('\\', '/');
 
-			// We don't want a slash at the start, but if there was not one specified, it's not anchored to the root of the tree.
+			// Remove the slash from the start of the pattern, interpreting it in the following manner:
+			// * Any pattern starting with a path separator is matched relative to the root directory
+			// * Any pattern containing a path separator is matched relative to the root directory, unless it starts with a wildcard indicating that it can be matched anywhere
+			// * Anything else is matched anywhere
 			if (NormalizedPattern.StartsWith("/"))
 			{
 				NormalizedPattern = NormalizedPattern.Substring(1);
 			}
-			else if(!NormalizedPattern.StartsWith("..."))
+			else if(!NormalizedPattern.Contains("/") && !NormalizedPattern.StartsWith("..."))
 			{
 				NormalizedPattern = ".../" + NormalizedPattern;
 			}
@@ -385,10 +416,9 @@ namespace AutomationTool
 		/// </summary>
 		/// <param name="FileNames">List of filenames</param>
 		/// <returns>List of filenames which match the filter</returns>
-		public IEnumerable<string> ApplyTo(string BaseDirectory, IEnumerable<string> FileNames)
+		public IEnumerable<FileReference> ApplyTo(DirectoryReference BaseDirectory, IEnumerable<FileReference> FileNames)
 		{
-			string FullBaseDirectory = Path.GetFullPath(BaseDirectory);
-			return FileNames.Where(x => Matches(AutomationTool.CommandUtils.StripBaseDirectory(Path.GetFullPath(x), FullBaseDirectory)));
+			return FileNames.Where(x => Matches(x.MakeRelativeTo(BaseDirectory)));
 		}
 
 		/// <summary>
@@ -403,6 +433,18 @@ namespace AutomationTool
 			return MatchingFileNames;
 		}
 
+		/// <summary>
+		/// Finds a list of files within a given directory which match the filter.
+		/// </summary>
+		/// <param name="DirectoryName">File to match</param>
+		/// <returns>List of files that pass the filter</returns>
+		public List<FileReference> ApplyToDirectory(DirectoryReference DirectoryName, bool bIgnoreSymlinks)
+		{
+			List<FileReference> MatchingFileNames = new List<FileReference>();
+			FindMatchesFromDirectory(new DirectoryInfo(DirectoryName.FullName), "", bIgnoreSymlinks, MatchingFileNames);
+			return MatchingFileNames;
+		}
+
 		/// <summary>
 		/// Finds a list of files within a given directory which match the filter.
 		/// </summary>
@@ -428,6 +470,31 @@ namespace AutomationTool
 			}
 		}
 
+		/// <summary>
+		/// Finds a list of files within a given directory which match the filter.
+		/// </summary>
+		/// <param name="FolderName">File to match</param>
+		/// <returns>True if the file passes the filter</returns>
+		void FindMatchesFromDirectory(DirectoryInfo CurrentDirectory, string NamePrefix, bool bIgnoreSymlinks, List<FileReference> MatchingFileNames)
+		{
+			foreach (FileInfo NextFile in CurrentDirectory.EnumerateFiles())
+			{
+				string FileName = NamePrefix + NextFile.Name;
+				if (Matches(FileName) && (!bIgnoreSymlinks || !NextFile.Attributes.HasFlag(FileAttributes.ReparsePoint)))
+				{
+					MatchingFileNames.Add(new FileReference(NextFile));
+				}
+			}
+			foreach (DirectoryInfo NextDirectory in CurrentDirectory.EnumerateDirectories())
+			{
+				string NextNamePrefix = NamePrefix + NextDirectory.Name;
+				if (PossiblyMatches(NextNamePrefix))
+				{
+					FindMatchesFromDirectory(NextDirectory, NextNamePrefix + "/", bIgnoreSymlinks, MatchingFileNames);
+				}
+			}
+		}
+
 		/// <summary>
 		/// Finds the node which matches a given list of tokens.
 		/// </summary>
diff --git a/Engine/Source/Programs/AutomationTool/AutomationUtils/ProcessUtils.cs b/Engine/Source/Programs/AutomationTool/AutomationUtils/ProcessUtils.cs
--- a/Engine/Source/Programs/AutomationTool/AutomationUtils/ProcessUtils.cs
+++ b/Engine/Source/Programs/AutomationTool/AutomationUtils/ProcessUtils.cs
@@ -849,7 +849,8 @@ namespace AutomationTool
 		/// <param name="CommandLine">Commandline to pass on to the executable</param>
 		public static string RunUAT(CommandEnvironment Env, string CommandLine)
 		{
-			// this doesn't do much, but it does need to make sure log folders are reasonable and don't collide
+			// We want to redirect the output from recursive UAT calls into our normal log folder, but prefix everything with a unique identifier. To do so, we set the EnvVarNames.LogFolder environment
+			// variable to a subfolder of it, then copy its contents into the main folder with a prefix after it's finished. Start by finding a base name we can use to identify the output of this run.
 			string BaseLogSubdir = "Recur";
 			if (!String.IsNullOrEmpty(CommandLine))
 			{
@@ -868,11 +869,11 @@ namespace AutomationTool
                     BaseLogSubdir = BaseLogSubdir + "_" + CommandLine;
                 }
 			}
-			int Index = 0;
-
 			BaseLogSubdir = BaseLogSubdir.Trim();
-			string DirOnlyName = BaseLogSubdir;
 
+			// Check if there are already log files which start with this prefix, and try to uniquify it if until there aren't.
+			int Index = 0;
+			string DirOnlyName = BaseLogSubdir;
 			string LogSubdir = CombinePaths(CmdEnv.LogFolder, DirOnlyName, "");
             while (true)
 			{
@@ -889,10 +890,13 @@ namespace AutomationTool
 				DirOnlyName = String.Format("{0}_{1}_", BaseLogSubdir, Index);
 				LogSubdir = CombinePaths(CmdEnv.LogFolder, DirOnlyName, "");
 			}
+
+			// Get the stdout log file for this run, and create the subdirectory for all the other log output
 			string LogFile = CombinePaths(CmdEnv.LogFolder, DirOnlyName + ".log");
 			LogVerbose("Recursive UAT Run, in log folder {0}, main log file {1}", LogSubdir, LogFile);
 			CreateDirectory(LogSubdir);
 
+			// Run UAT with the log folder redirected through the environment
 			string App = CmdEnv.UATExe;
 
 			Log("Running {0} {1}", App, CommandLine);
@@ -915,6 +919,7 @@ namespace AutomationTool
 				WriteToFile(LogFile, "[None!, no output produced]");
 			}
 
+			// Copy everything into the main log folder, using the prefix we decided on earlier.
 			LogVerbose("Flattening log folder {0}", LogSubdir);
 
 			var Files = FindFiles("*", true, LogSubdir);
diff --git a/Engine/Source/Programs/AutomationTool/AutomationUtils/ProjectParams.cs b/Engine/Source/Programs/AutomationTool/AutomationUtils/ProjectParams.cs
--- a/Engine/Source/Programs/AutomationTool/AutomationUtils/ProjectParams.cs
+++ b/Engine/Source/Programs/AutomationTool/AutomationUtils/ProjectParams.cs
@@ -1559,50 +1559,7 @@ namespace AutomationTool
 			var ProgramTarget = String.Empty;
 			var ProjectType = TargetRules.TargetType.Game;
 
-			if (Automation.RunningRocket())
-			{
-				if (!CommandUtils.CmdEnv.HasCapabilityToCompile || !bIsCodeBasedProject)
-				{
-					if (bIsCodeBasedProject)
-					{
-						var ShortName = ProjectUtils.GetShortProjectName(RawProjectPath);
-						GameTarget = ShortName;
-						EditorTarget = ShortName + "Editor";
-						ServerTarget = ShortName + "Server";
-					}
-					else
-					{
-						GameTarget = "UE4Game";
-						EditorTarget = "UE4Editor";
-						//ServerTarget = "RocketServer";
-
-						Build = false;
-					}
-				}
-				else
-				{
-					SingleTargetProperties TargetData;
-					if (DetectedTargets.TryGetValue(TargetRules.TargetType.Editor, out TargetData))
-					{
-						EditorTarget = TargetData.TargetName;
-					}
-
-					if (DetectedTargets.TryGetValue(TargetRules.TargetType.Program, out TargetData))
-					{
-						ProgramTarget = TargetData.TargetName;
-					}
-					//DetectedTargets.TryGetValue(TargetRules.TargetType.Server, out ServerTarget);
-
-					if (string.IsNullOrEmpty(GameTarget))
-					{
-						if (DetectedTargets.TryGetValue(TargetRules.TargetType.Game, out TargetData))
-						{
-							GameTarget = TargetData.TargetName;
-						}
-					}
-				}
-			}
-			else if (!bIsCodeBasedProject)
+			if (!bIsCodeBasedProject)
 			{
 				GameTarget = "UE4Game";
 				EditorTarget = "UE4Editor";
@@ -1694,7 +1651,7 @@ namespace AutomationTool
 
 			IsProgramTarget = ProjectType == TargetRules.TargetType.Program;
 
-			if (String.IsNullOrEmpty(EditorTarget) && ProjectType != TargetRules.TargetType.Program && CommandUtils.IsNullOrEmpty(EditorTargetsList) && !Automation.RunningRocket())
+			if (String.IsNullOrEmpty(EditorTarget) && ProjectType != TargetRules.TargetType.Program && CommandUtils.IsNullOrEmpty(EditorTargetsList))
 			{
 				if (Properties.bWasGenerated)
 				{
@@ -1705,7 +1662,7 @@ namespace AutomationTool
 					throw new AutomationException("Editor target not found!");
 				}
 			}
-			if (String.IsNullOrEmpty(GameTarget) && Run && !NoClient && (Cook || CookOnTheFly) && CommandUtils.IsNullOrEmpty(ClientCookedTargetsList) && !Automation.RunningRocket())
+			if (String.IsNullOrEmpty(GameTarget) && Run && !NoClient && (Cook || CookOnTheFly) && CommandUtils.IsNullOrEmpty(ClientCookedTargetsList))
 			{
 				throw new AutomationException("Game target not found. Game target is required with -cook or -cookonthefly");
 			}
@@ -2284,7 +2241,6 @@ namespace AutomationTool
 				CommandUtils.LogLog("Prereqs={0}", Prereqs);
 				CommandUtils.LogLog("NoBootstrapExe={0}", NoBootstrapExe);
 				CommandUtils.LogLog("RawProjectPath={0}", RawProjectPath);
-				CommandUtils.LogLog("Rocket={0}", Automation.RunningRocket());
 				CommandUtils.LogLog("Run={0}", Run);
 				CommandUtils.LogLog("ServerConfigsToBuild={0}", string.Join(",", ServerConfigsToBuild));
 				CommandUtils.LogLog("ServerCookedTargets={0}", ServerCookedTargets.ToString());
diff --git a/Engine/Source/Programs/AutomationTool/AutomationUtils/UE4Build.cs b/Engine/Source/Programs/AutomationTool/AutomationUtils/UE4Build.cs
--- a/Engine/Source/Programs/AutomationTool/AutomationUtils/UE4Build.cs
+++ b/Engine/Source/Programs/AutomationTool/AutomationUtils/UE4Build.cs
@@ -448,15 +448,20 @@ namespace AutomationTool
 			bool bIsLicenseeVersion = ParseParam("Licensee");
 			bool bDoUpdateVersionFiles = CommandUtils.P4Enabled && ActuallyUpdateVersionFiles;		
 			int ChangelistNumber = 0;
-			string ChangelistString = String.Empty;
 			if (bDoUpdateVersionFiles)
 			{
 				ChangelistNumber = ChangelistNumberOverride.HasValue? ChangelistNumberOverride.Value : P4Env.Changelist;
-				ChangelistString = ChangelistNumber.ToString();
 			}
 
-			var Result = new List<String>();
 			string Branch = P4Enabled ? P4Env.BuildRootEscaped : "";
+			return StaticUpdateVersionFiles(ChangelistNumber, Branch, bIsLicenseeVersion, bDoUpdateVersionFiles);
+		}
+
+		public static List<string> StaticUpdateVersionFiles(int ChangelistNumber, string Branch, bool bIsLicenseeVersion, bool bDoUpdateVersionFiles)
+		{
+			string ChangelistString = (ChangelistNumber != 0 && bDoUpdateVersionFiles)? ChangelistNumber.ToString() : String.Empty;
+
+			var Result = new List<String>();
 			{
 				string VerFile = CombinePaths(CmdEnv.LocalRoot, "Engine", "Build", "Build.version");
 				if (bDoUpdateVersionFiles)
@@ -820,7 +825,7 @@ namespace AutomationTool
 									bOutputContainsProject = Output.Contains("------Project:");
 								}
 								PopDir();
-								if (ConnectionRetries > 0 && (SuccesCode == 4 || SuccesCode == 2) && !bOutputContainsProject)
+								if (ConnectionRetries > 0 && (SuccesCode == 4 || SuccesCode == 2) && !bOutputContainsProject && XGETool != null)
 								{
 									LogWarning(String.Format("{0} failure on the local connection timeout", XGETool));
 									if (ConnectionRetries < 2)
@@ -1605,10 +1610,8 @@ namespace AutomationTool
 
 		string GetUBTManifest(FileReference UProjectPath, string InAddArgs)
 		{
-			// Can't write to Engine directory on 
-            bool bForceRocket = InAddArgs.ToLowerInvariant().Contains(" -rocket"); //awful
-
-            if ((Automation.IsEngineInstalled() || bForceRocket) && UProjectPath != null)
+			// Can't write to Engine directory on installed builds
+			if (Automation.IsEngineInstalled() && UProjectPath != null)
 			{
 				return Path.Combine(Path.GetDirectoryName(UProjectPath.FullName), "Intermediate/Build/Manifest.xml");				
 			}
diff --git a/Engine/Source/Programs/AutomationTool/AutomationUtils/Utils.cs b/Engine/Source/Programs/AutomationTool/AutomationUtils/Utils.cs
--- a/Engine/Source/Programs/AutomationTool/AutomationUtils/Utils.cs
+++ b/Engine/Source/Programs/AutomationTool/AutomationUtils/Utils.cs
@@ -979,7 +979,7 @@ namespace AutomationTool
 				}
 				else
 				{
-					// This is the standard Rocket versioning scheme, e.g. "4.5.0-12345678+++depot+UE4"
+					// This is the standard Launcher versioning scheme, e.g. "4.5.0-12345678+++depot+UE4"
 					var dotSplit = versionString.Split(new[] { '.' }, 3);
 					var dashSplit = dotSplit[2].Split(new[] { '-' }, 2);
 					var plusSplit = dashSplit[1].Split(new[] { '+' }, 2);
diff --git a/Engine/Source/Programs/AutomationTool/GUBP/LegacyNodes.cs b/Engine/Source/Programs/AutomationTool/GUBP/LegacyNodes.cs
--- a/Engine/Source/Programs/AutomationTool/GUBP/LegacyNodes.cs
+++ b/Engine/Source/Programs/AutomationTool/GUBP/LegacyNodes.cs
@@ -1058,6 +1058,10 @@ partial class GUBP
 				{
 					UnrealPakExe = CombinePaths(CmdEnv.LocalRoot, "Engine/Binaries/Linux/UnrealPak");
 				}
+                else if (HostPlatform == UnrealTargetPlatform.Mac)
+                {
+                    UnrealPakExe = CombinePaths(CmdEnv.LocalRoot, "Engine/Binaries/Mac/UnrealPak");
+                }
 				else
 				{
 					throw new AutomationException("Unknown path to UnrealPak for host platform ({0})", HostPlatform);
diff --git a/Engine/Source/Programs/AutomationTool/IOS/IOSPlatform.Automation.cs b/Engine/Source/Programs/AutomationTool/IOS/IOSPlatform.Automation.cs
--- a/Engine/Source/Programs/AutomationTool/IOS/IOSPlatform.Automation.cs
+++ b/Engine/Source/Programs/AutomationTool/IOS/IOSPlatform.Automation.cs
@@ -437,16 +437,25 @@ public class IOSPlatform : Platform
 			// project.xcodeproj doesn't exist, so generate temp project
 			string Arguments = "-project=\"" + RawProjectPath + "\"";
 			Arguments += " -platforms=" + PlatformName + " -game -nointellisense -" + PlatformName + "deployonly -ignorejunk";
-			string Script = CombinePaths(CmdEnv.LocalRoot, "Engine/Build/BatchFiles/Mac/GenerateProjectFiles.sh");
-			if (Automation.RunningRocket())
+
+			// If engine is installed then UBT doesn't need to be built
+			if (Automation.IsEngineInstalled())
 			{
-				Script = CombinePaths(CmdEnv.LocalRoot, "Engine/Build/BatchFiles/Mac/RocketGenerateProjectFiles.sh");
+				// Get the path to UBT
+				string InstalledUBT = CombinePaths(CmdEnv.LocalRoot, "Engine/Binaries/DotNET/UnrealBuildTool.exe");
+				Arguments = "-XcodeProjectFile " + Arguments;
+				RunUBT(CmdEnv, InstalledUBT, Arguments);
 			}
-			string CWD = Directory.GetCurrentDirectory ();
-			Directory.SetCurrentDirectory (Path.GetDirectoryName (Script));
-			Run (Script, Arguments, null, ERunOptions.Default);
+			else
+			{
+				string Script = CombinePaths(CmdEnv.LocalRoot, "Engine/Build/BatchFiles/Mac/GenerateProjectFiles.sh");
+				string CWD = Directory.GetCurrentDirectory();
+				Directory.SetCurrentDirectory(Path.GetDirectoryName(Script));
+				Run(Script, Arguments, null, ERunOptions.Default);
+				Directory.SetCurrentDirectory(CWD);
+			}
+
 			bWasGenerated = true;
-			Directory.SetCurrentDirectory (CWD);
 
 			if (!Directory.Exists (XcodeProj))
 			{
diff --git a/Engine/Source/Programs/AutomationTool/Scripts/BuildProjectCommand.Automation.cs b/Engine/Source/Programs/AutomationTool/Scripts/BuildProjectCommand.Automation.cs
--- a/Engine/Source/Programs/AutomationTool/Scripts/BuildProjectCommand.Automation.cs
+++ b/Engine/Source/Programs/AutomationTool/Scripts/BuildProjectCommand.Automation.cs
@@ -135,7 +135,7 @@ public partial class Project : CommandUtils
 				}
 			}
 		}
-		if (Params.HasProgramTargets && !Automation.RunningRocket() && (TargetMask & ProjectBuildTargets.Programs) == ProjectBuildTargets.Programs)
+		if (Params.HasProgramTargets && (TargetMask & ProjectBuildTargets.Programs) == ProjectBuildTargets.Programs)
 		{
 			foreach (var BuildConfig in Params.ClientConfigsToBuild)
 			{
diff --git a/Engine/Source/Programs/AutomationTool/Scripts/RocketBuild.Automation.cs b/Engine/Source/Programs/AutomationTool/Scripts/RocketBuild.Automation.cs
--- a/Engine/Source/Programs/AutomationTool/Scripts/RocketBuild.Automation.cs
+++ b/Engine/Source/Programs/AutomationTool/Scripts/RocketBuild.Automation.cs
@@ -1152,10 +1152,6 @@ namespace Rocket
 			BuildDerivedDataCacheNode DerivedDataCacheNode = (BuildDerivedDataCacheNode)BranchConfig.FindNode(BuildDerivedDataCacheNode.StaticGetFullName(HostPlatform));
 			CopyManifestFilesToOutput(DerivedDataCacheNode.SavedManifestPath, DerivedDataCacheNode.SavedDir, OutputDir);
 
-			// Write the Rocket.txt file to indicate a Rocket build
-			string RocketFile = CommandUtils.CombinePaths(OutputDir, "Engine/Build/Rocket.txt");
-			CommandUtils.WriteAllText(RocketFile, "");
-
 			// Write InstalledBuild.txt to indicate Engine is installed
 			string InstalledBuildFile = CommandUtils.CombinePaths(OutputDir, "Engine/Build/InstalledBuild.txt");
 			CommandUtils.WriteAllText(InstalledBuildFile, "");
diff --git a/Engine/Source/Programs/CrashReporter/CrashReportClient/CrashReportClient.Target.cs b/Engine/Source/Programs/CrashReporter/CrashReportClient/CrashReportClient.Target.cs
--- a/Engine/Source/Programs/CrashReporter/CrashReportClient/CrashReportClient.Target.cs
+++ b/Engine/Source/Programs/CrashReporter/CrashReportClient/CrashReportClient.Target.cs
@@ -27,7 +27,10 @@ public class CrashReportClientTarget : TargetRules
 	{
 		if( base.GetSupportedConfigurations( ref OutConfigurations, bIncludeTestAndShippingConfigs ) )
 		{
-			OutConfigurations.Add( UnrealTargetConfiguration.Shipping );
+			if( bIncludeTestAndShippingConfigs )
+			{ 
+				OutConfigurations.Add( UnrealTargetConfiguration.Shipping );
+			}
 			OutConfigurations.Add( UnrealTargetConfiguration.Debug );
 			return true;
 		}
diff --git a/Engine/Source/Programs/UnrealBuildTool/Configuration/BuildConfiguration.cs b/Engine/Source/Programs/UnrealBuildTool/Configuration/BuildConfiguration.cs
--- a/Engine/Source/Programs/UnrealBuildTool/Configuration/BuildConfiguration.cs
+++ b/Engine/Source/Programs/UnrealBuildTool/Configuration/BuildConfiguration.cs
@@ -759,7 +759,7 @@ namespace UnrealBuildTool
 		/// <param name="bCreateDebugInfo">True if debug info should be created</param>
 		public static void ValidateConfiguration(CPPTargetConfiguration Configuration, CPPTargetPlatform Platform, bool bCreateDebugInfo, UEBuildPlatformContext PlatformContext)
 		{
-			var BuildPlatform = UEBuildPlatform.GetBuildPlatformForCPPTargetPlatform(Platform);
+			UEBuildPlatform BuildPlatform = UEBuildPlatform.GetBuildPlatformForCPPTargetPlatform(Platform);
 
 			// E&C support.
 			if (bSupportEditAndContinue)
diff --git a/Engine/Source/Programs/UnrealBuildTool/Configuration/EngineConfiguration.cs b/Engine/Source/Programs/UnrealBuildTool/Configuration/EngineConfiguration.cs
--- a/Engine/Source/Programs/UnrealBuildTool/Configuration/EngineConfiguration.cs
+++ b/Engine/Source/Programs/UnrealBuildTool/Configuration/EngineConfiguration.cs
@@ -94,7 +94,7 @@ namespace UnrealBuildTool
 			public IniSection(IniSection Other)
 				: this()
 			{
-				foreach (var Pair in Other)
+				foreach (KeyValuePair<string, IniValues> Pair in Other)
 				{
 					Add(Pair.Key, new IniValues(Pair.Value));
 				}
@@ -191,14 +191,14 @@ namespace UnrealBuildTool
 
 			if (BaseCache != null)
 			{
-				foreach (var Pair in BaseCache.Sections)
+				foreach (KeyValuePair<string, IniSection> Pair in BaseCache.Sections)
 				{
 					Sections.Add(Pair.Key, new IniSection(Pair.Value));
 				}
 			}
 			if (EngineOnly)
 			{
-				foreach (var IniFileName in EnumerateEngineIniFileNames(EngineDirectory, BaseIniName))
+				foreach (FileReference IniFileName in EnumerateEngineIniFileNames(EngineDirectory, BaseIniName))
 				{
 					if (IniFileName.Exists())
 					{
@@ -208,7 +208,7 @@ namespace UnrealBuildTool
 			}
 			else
 			{
-				foreach (var IniFileName in EnumerateCrossPlatformIniFileNames(ProjectDirectory, EngineDirectory, Platform, BaseIniName, BaseCache != null))
+				foreach (FileReference IniFileName in EnumerateCrossPlatformIniFileNames(ProjectDirectory, EngineDirectory, Platform, BaseIniName, BaseCache != null))
 				{
 					if (IniFileName.Exists())
 					{
@@ -236,7 +236,7 @@ namespace UnrealBuildTool
 		private bool GetList(string SectionName, string Key, out IniValues Value)
 		{
 			bool Result = false;
-			var Section = FindSection(SectionName);
+			IniSection Section = FindSection(SectionName);
 			Value = null;
 			if (Section != null)
 			{
@@ -693,7 +693,7 @@ namespace UnrealBuildTool
 						IniValues Value;
 						if (CurrentSection.TryGetValue(Key, out Value))
 						{
-							var ExistingIndex = Value.FindIndex(X => (String.Compare(SingleValue, X, true) == 0));
+							int ExistingIndex = Value.FindIndex(X => (String.Compare(SingleValue, X, true) == 0));
 							if (ExistingIndex >= 0)
 							{
 								Value.RemoveAt(ExistingIndex);
diff --git a/Engine/Source/Programs/UnrealBuildTool/Configuration/UEBuildBinary.cs b/Engine/Source/Programs/UnrealBuildTool/Configuration/UEBuildBinary.cs
--- a/Engine/Source/Programs/UnrealBuildTool/Configuration/UEBuildBinary.cs
+++ b/Engine/Source/Programs/UnrealBuildTool/Configuration/UEBuildBinary.cs
@@ -201,7 +201,7 @@ namespace UnrealBuildTool
 		/// <summary>
 		/// Creates all the modules referenced by this target.
 		/// </summary>
-		public virtual void CreateAllDependentModules() { }
+		public virtual void CreateAllDependentModules(UEBuildTarget Target) { }
 
 		/// <summary>
 		/// Builds the binary.
@@ -210,7 +210,7 @@ namespace UnrealBuildTool
 		/// <param name="CompileEnvironment">The environment to compile the binary in</param>
 		/// <param name="LinkEnvironment">The environment to link the binary in</param>
 		/// <returns></returns>
-		public abstract IEnumerable<FileItem> Build(UEToolChain ToolChain, CPPEnvironment CompileEnvironment, LinkEnvironment LinkEnvironment);
+		public abstract IEnumerable<FileItem> Build(UEBuildTarget Target, UEToolChain ToolChain, CPPEnvironment CompileEnvironment, LinkEnvironment LinkEnvironment);
 
 		/// <summary>
 		/// Called to allow the binary to modify the link environment of a different binary containing 
@@ -349,10 +349,10 @@ namespace UnrealBuildTool
 		public void CheckOutputDistributionLevelAgainstDependencies(Dictionary<UEBuildModule, UEBuildModuleDistribution> ModuleDistributionCache)
 		{
 			// Find maximum distribution level of its direct dependencies
-			var DistributionLevel = UEBuildModuleDistribution.Public;
-			var DependantModules = GetAllDependencyModules(false, false);
+			UEBuildModuleDistribution DistributionLevel = UEBuildModuleDistribution.Public;
+			List<UEBuildModule> DependantModules = GetAllDependencyModules(false, false);
 			List<string>[] DependantModuleNames = new List<string>[Enum.GetNames(typeof(UEBuildModuleDistribution)).Length];
-			foreach (var Module in DependantModules)
+			foreach (UEBuildModule Module in DependantModules)
 			{
 				UEBuildModuleDistribution ModuleDistributionLevel;
 				if (!ModuleDistributionCache.TryGetValue(Module, out ModuleDistributionLevel))
@@ -378,14 +378,14 @@ namespace UnrealBuildTool
 			// Check Output Paths if dependencies shouldn't be distributed to everyone
 			if (DistributionLevel != UEBuildModuleDistribution.Public)
 			{
-				foreach (var OutputFilePath in Config.OutputFilePaths)
+				foreach (FileReference OutputFilePath in Config.OutputFilePaths)
 				{
-					var OutputDistributionLevel = UEBuildModule.GetModuleDistributionLevelBasedOnLocation(OutputFilePath.FullName);
+					UEBuildModuleDistribution OutputDistributionLevel = UEBuildModule.GetModuleDistributionLevelBasedOnLocation(OutputFilePath.FullName);
 
 					// Throw exception if output path is not appropriate
 					if (OutputDistributionLevel < DistributionLevel)
 					{
-						var JoinedModuleNames = String.Join(",", DependantModuleNames[(int)DistributionLevel]);
+						string JoinedModuleNames = String.Join(",", DependantModuleNames[(int)DistributionLevel]);
 						throw new BuildException("Output file \"{0}\" has distribution level of \"{1}\" but has direct dependencies on modules with distribution level of \"{2}\" ({3}).\nEither change to dynamic dependencies, set BinariesSubFolder/ExeBinariesSubFolder to \"{2}\" or set bOutputPubliclyDistributable to true in the target.cs or build.cs file.",
 							OutputFilePath, OutputDistributionLevel.ToString(), DistributionLevel.ToString(), JoinedModuleNames);
 					}
@@ -432,13 +432,13 @@ namespace UnrealBuildTool
 		/// <summary>
 		/// Creates all the modules referenced by this target.
 		/// </summary>
-		public override void CreateAllDependentModules()
+		public override void CreateAllDependentModules(UEBuildTarget Target)
 		{
 			if (Config.bHasModuleRules)
 			{
 				foreach (UEBuildModule Module in Modules)
 				{
-					Module.RecursivelyCreateModules();
+					Module.RecursivelyCreateModules(Target);
 				}
 			}
 		}
@@ -451,20 +451,22 @@ namespace UnrealBuildTool
 		/// <returns>List of all referenced modules</returns>
 		public override List<UEBuildModule> GetAllDependencyModules(bool bIncludeDynamicallyLoaded, bool bForceCircular)
 		{
-			var ReferencedModules = new CaselessDictionary<UEBuildModule.ModuleIndexPair>();
-			foreach (var Module in Modules)
+			List<UEBuildModule> ReferencedModules = new List<UEBuildModule>();
+			HashSet<UEBuildModule> IgnoreReferencedModules = new HashSet<UEBuildModule>();
+
+			foreach (UEBuildModule Module in Modules)
 			{
-				if (!ReferencedModules.ContainsKey(Module.Name))
+				if (!IgnoreReferencedModules.Contains(Module))
 				{
-					ReferencedModules[Module.Name] = null;
+					IgnoreReferencedModules.Add(Module);
 
-					Module.GetAllDependencyModules(ReferencedModules, bIncludeDynamicallyLoaded, bForceCircular, bOnlyDirectDependencies: false);
+					Module.GetAllDependencyModules(ReferencedModules, IgnoreReferencedModules, bIncludeDynamicallyLoaded, bForceCircular, bOnlyDirectDependencies: false);
 
-					ReferencedModules[Module.Name] = new UEBuildModule.ModuleIndexPair { Module = Module, Index = ReferencedModules.Count };
+					ReferencedModules.Add(Module);
 				}
 			}
 
-			return ReferencedModules.Values.OrderBy(M => M.Index).Select(M => M.Module).ToList();
+			return ReferencedModules;
 		}
 
 		/// <summary>
@@ -511,7 +513,7 @@ namespace UnrealBuildTool
 		/// <param name="CompileEnvironment">The environment to compile the binary in</param>
 		/// <param name="LinkEnvironment">The environment to link the binary in</param>
 		/// <returns></returns>
-		public override IEnumerable<FileItem> Build(UEToolChain ToolChain, CPPEnvironment CompileEnvironment, LinkEnvironment LinkEnvironment)
+		public override IEnumerable<FileItem> Build(UEBuildTarget Target, UEToolChain ToolChain, CPPEnvironment CompileEnvironment, LinkEnvironment LinkEnvironment)
 		{
 			// UnrealCodeAnalyzer produces output files only for a specific module.
 			if (BuildConfiguration.bRunUnrealCodeAnalyzer && !(Modules.Any(x => x.Name == BuildConfiguration.UCAModuleToAnalyze)))
@@ -520,7 +522,7 @@ namespace UnrealBuildTool
 			}
 
 			// Setup linking environment.
-			var BinaryLinkEnvironment = SetupBinaryLinkEnvironment(ToolChain, LinkEnvironment, CompileEnvironment);
+			LinkEnvironment BinaryLinkEnvironment = SetupBinaryLinkEnvironment(Target, ToolChain, LinkEnvironment, CompileEnvironment);
 
 			// Return linked files.
 			return SetupOutputFiles(ToolChain, ref BinaryLinkEnvironment);
@@ -568,9 +570,9 @@ namespace UnrealBuildTool
 		/// <returns>The OnlyModule if found, null if not</returns>
 		public override OnlyModule FindOnlyModule(List<OnlyModule> OnlyModules)
 		{
-			foreach (var Module in Modules)
+			foreach (UEBuildModule Module in Modules)
 			{
-				foreach (var OnlyModule in OnlyModules)
+				foreach (OnlyModule OnlyModule in OnlyModules)
 				{
 					if (OnlyModule.OnlyModuleName.ToLower() == Module.Name.ToLower())
 					{
@@ -583,8 +585,8 @@ namespace UnrealBuildTool
 
 		public override List<UEBuildModule> FindGameModules()
 		{
-			var GameModules = new List<UEBuildModule>();
-			foreach (var Module in Modules)
+			List<UEBuildModule> GameModules = new List<UEBuildModule>();
+			foreach (UEBuildModule Module in Modules)
 			{
 				if (!Module.ModuleDirectory.IsUnderDirectory(UnrealBuildTool.EngineDirectory))
 				{
@@ -622,15 +624,15 @@ namespace UnrealBuildTool
 			return Config.OutputFilePath.FullName;
 		}
 
-		private LinkEnvironment SetupBinaryLinkEnvironment(UEToolChain ToolChain, LinkEnvironment LinkEnvironment, CPPEnvironment CompileEnvironment)
+		private LinkEnvironment SetupBinaryLinkEnvironment(UEBuildTarget Target, UEToolChain ToolChain, LinkEnvironment LinkEnvironment, CPPEnvironment CompileEnvironment)
 		{
-			var BinaryLinkEnvironment = LinkEnvironment.DeepCopy();
-			var LinkEnvironmentVisitedModules = new HashSet<UEBuildModule>();
-			var BinaryDependencies = new List<UEBuildBinary>();
+			LinkEnvironment BinaryLinkEnvironment = LinkEnvironment.DeepCopy();
+			HashSet<UEBuildModule> LinkEnvironmentVisitedModules = new HashSet<UEBuildModule>();
+			List<UEBuildBinary> BinaryDependencies = new List<UEBuildBinary>();
 			CompileEnvironment.Config.bIsBuildingDLL = IsBuildingDll(Config.Type);
 			CompileEnvironment.Config.bIsBuildingLibrary = IsBuildingLibrary(Config.Type);
 
-			var BinaryCompileEnvironment = CompileEnvironment.DeepCopy();
+			CPPEnvironment BinaryCompileEnvironment = CompileEnvironment.DeepCopy();
 			// @Hack: This to prevent UHT from listing CoreUObject.generated.cpp as its dependency.
 			// We flag the compile environment when we build UHT so that we don't need to check
 			// this for each file when generating their dependencies.
@@ -638,24 +640,24 @@ namespace UnrealBuildTool
 
 			// @todo: This should be in some Windows code somewhere...
 			// Set the original file name macro; used in PCLaunch.rc to set the binary metadata fields.
-			var OriginalFilename = (Config.OriginalOutputFilePaths != null) ?
+			string OriginalFilename = (Config.OriginalOutputFilePaths != null) ?
 				Config.OriginalOutputFilePaths[0].GetFileName() :
 				Config.OutputFilePaths[0].GetFileName();
 			BinaryCompileEnvironment.Config.Definitions.Add("ORIGINAL_FILE_NAME=\"" + OriginalFilename + "\"");
 
-			foreach (var Module in Modules)
+			foreach (UEBuildModule Module in Modules)
 			{
 				List<FileItem> LinkInputFiles;
 				if (Module.Binary == null || Module.Binary == this)
 				{
 					// Compile each module.
 					Log.TraceVerbose("Compile module: " + Module.Name);
-					LinkInputFiles = Module.Compile(ToolChain, CompileEnvironment, BinaryCompileEnvironment);
+					LinkInputFiles = Module.Compile(Target, ToolChain, CompileEnvironment, BinaryCompileEnvironment);
 
 					// NOTE: Because of 'Shared PCHs', in monolithic builds the same PCH file may appear as a link input
 					// multiple times for a single binary.  We'll check for that here, and only add it once.  This avoids
 					// a linker warning about redundant .obj files. 
-					foreach (var LinkInputFile in LinkInputFiles)
+					foreach (FileItem LinkInputFile in LinkInputFiles)
 					{
 						if (!BinaryLinkEnvironment.InputFiles.Contains(LinkInputFile))
 						{
@@ -677,7 +679,7 @@ namespace UnrealBuildTool
 
 
 			// Allow the binary dependencies to modify the link environment.
-			foreach (var BinaryDependency in BinaryDependencies)
+			foreach (UEBuildBinary BinaryDependency in BinaryDependencies)
 			{
 				BinaryDependency.SetupDependentLinkEnvironment(BinaryLinkEnvironment);
 			}
@@ -740,7 +742,7 @@ namespace UnrealBuildTool
 			//
 			// Regular linking action.
 			//
-			var OutputFiles = new List<FileItem>();
+			List<FileItem> OutputFiles = new List<FileItem>();
 			if (bCreateImportLibrarySeparately)
 			{
 				// Mark the link environment as cross-referenced.
@@ -763,7 +765,7 @@ namespace UnrealBuildTool
 			if (Config.bBuildAdditionalConsoleApp)
 			{
 				// Produce additional binary but link it as a console app
-				var ConsoleAppLinkEvironment = BinaryLinkEnvironment.DeepCopy();
+				LinkEnvironment ConsoleAppLinkEvironment = BinaryLinkEnvironment.DeepCopy();
 				ConsoleAppLinkEvironment.Config.bIsBuildingConsoleApplication = true;
 				ConsoleAppLinkEvironment.Config.WindowsEntryPointOverride = "WinMainCRTStartup";		// For WinMain() instead of "main()" for Launch module
 				ConsoleAppLinkEvironment.Config.OutputFilePaths = ConsoleAppLinkEvironment.Config.OutputFilePaths.Select(Path => GetAdditionalConsoleAppPath(Path)).ToList();
@@ -772,7 +774,7 @@ namespace UnrealBuildTool
 				OutputFiles.AddRange(ToolChain.LinkAllFiles(ConsoleAppLinkEvironment, false));
 			}
 
-			foreach (var Executable in Executables)
+			foreach (FileItem Executable in Executables)
 			{
 				OutputFiles.AddRange(ToolChain.PostBuild(Executable, BinaryLinkEnvironment));
 			}
@@ -782,13 +784,13 @@ namespace UnrealBuildTool
 
 		private List<FileItem> CreateOutputFilesForUCA(LinkEnvironment BinaryLinkEnvironment)
 		{
-			var OutputFiles = new List<FileItem>();
-			var ModuleName = Modules.Select(Module => Module.Name).First(Name => Name.CompareTo(BuildConfiguration.UCAModuleToAnalyze) == 0);
-			var ModuleCPP = (UEBuildModuleCPP)Target.GetModuleByName(ModuleName);
-			var ModulePrivatePCH = ModuleCPP.ProcessedDependencies.UniquePCHHeaderFile;
-			var IntermediatePath = Path.Combine(Target.ProjectIntermediateDirectory.FullName, ModuleName);
-			var OutputFileName = Target.OutputPath;
-			var OutputFile = FileItem.GetItemByFileReference(OutputFileName);
+			List<FileItem> OutputFiles = new List<FileItem>();
+			string ModuleName = Modules.Select(Module => Module.Name).First(Name => Name.CompareTo(BuildConfiguration.UCAModuleToAnalyze) == 0);
+			UEBuildModuleCPP ModuleCPP = (UEBuildModuleCPP)Target.GetModuleByName(ModuleName);
+			FileItem ModulePrivatePCH = ModuleCPP.ProcessedDependencies.UniquePCHHeaderFile;
+			string IntermediatePath = Path.Combine(Target.ProjectIntermediateDirectory.FullName, ModuleName);
+			FileReference OutputFileName = Target.OutputPath;
+			FileItem OutputFile = FileItem.GetItemByFileReference(OutputFileName);
 
 			Action LinkAction = new Action(ActionType.Compile);
 			LinkAction.WorkingDirectory = UnrealBuildTool.EngineSourceDirectory.FullName;
@@ -829,9 +831,9 @@ namespace UnrealBuildTool
 		/// <param name="CompileEnvironment">The environment to compile the binary in</param>
 		/// <param name="LinkEnvironment">The environment to link the binary in</param>
 		/// <returns></returns>
-		public override IEnumerable<FileItem> Build(UEToolChain ToolChain, CPPEnvironment CompileEnvironment, LinkEnvironment LinkEnvironment)
+		public override IEnumerable<FileItem> Build(UEBuildTarget Target, UEToolChain ToolChain, CPPEnvironment CompileEnvironment, LinkEnvironment LinkEnvironment)
 		{
-			var ProjectCSharpEnviroment = new CSharpEnvironment();
+			CSharpEnvironment ProjectCSharpEnviroment = new CSharpEnvironment();
 			if (LinkEnvironment.Config.Target.Configuration == CPPTargetConfiguration.Debug)
 			{
 				ProjectCSharpEnviroment.TargetConfiguration = CSharpTargetConfiguration.Debug;
diff --git a/Engine/Source/Programs/UnrealBuildTool/Configuration/UEBuildModule.cs b/Engine/Source/Programs/UnrealBuildTool/Configuration/UEBuildModule.cs
--- a/Engine/Source/Programs/UnrealBuildTool/Configuration/UEBuildModule.cs
+++ b/Engine/Source/Programs/UnrealBuildTool/Configuration/UEBuildModule.cs
@@ -12,73 +12,6 @@ using Tools.DotNETCommon.CaselessDictionary;
 
 namespace UnrealBuildTool
 {
-	/// <summary>
-	/// Type of module. Mirrored in UHT as EBuildModuleType.
-	/// This should be sorted by the order in which we expect modules to be built.
-	/// </summary>
-	public enum UEBuildModuleType
-	{
-		Program,
-		EngineRuntime,
-		EngineDeveloper,
-		EngineEditor,
-		EngineThirdParty,
-		GameRuntime,
-		GameDeveloper,
-		GameEditor,
-		GameThirdParty,
-	}
-	public static class UEBuildModuleTypeExtensions
-	{
-		public static bool IsProgramModule(this UEBuildModuleType ModuleType)
-		{
-			return ModuleType == UEBuildModuleType.Program;
-		}
-		public static bool IsEngineModule(this UEBuildModuleType ModuleType)
-		{
-			return ModuleType == UEBuildModuleType.EngineRuntime || ModuleType == UEBuildModuleType.EngineDeveloper || ModuleType == UEBuildModuleType.EngineEditor || ModuleType == UEBuildModuleType.EngineThirdParty;
-		}
-		public static bool IsGameModule(this UEBuildModuleType ModuleType)
-		{
-			return ModuleType == UEBuildModuleType.GameRuntime || ModuleType == UEBuildModuleType.GameDeveloper || ModuleType == UEBuildModuleType.GameEditor || ModuleType == UEBuildModuleType.GameThirdParty;
-		}
-		public static UEBuildModuleType? EngineModuleTypeFromHostType(ModuleHostType ModuleType)
-		{
-			switch (ModuleType)
-			{
-				case ModuleHostType.Runtime:
-				case ModuleHostType.RuntimeNoCommandlet:
-				case ModuleHostType.RuntimeAndProgram:
-				case ModuleHostType.ServerOnly:
-					return UEBuildModuleType.EngineRuntime;
-				case ModuleHostType.Developer:
-					return UEBuildModuleType.EngineDeveloper;
-				case ModuleHostType.Editor:
-				case ModuleHostType.EditorNoCommandlet:
-					return UEBuildModuleType.EngineEditor;
-				default:
-					return null;
-			}
-		}
-		public static UEBuildModuleType? GameModuleTypeFromHostType(ModuleHostType ModuleType)
-		{
-			switch (ModuleType)
-			{
-				case ModuleHostType.Runtime:
-				case ModuleHostType.RuntimeNoCommandlet:
-				case ModuleHostType.RuntimeAndProgram:
-				case ModuleHostType.ServerOnly:
-					return UEBuildModuleType.GameRuntime;
-				case ModuleHostType.Developer:
-					return UEBuildModuleType.GameDeveloper;
-				case ModuleHostType.Editor:
-				case ModuleHostType.EditorNoCommandlet:
-					return UEBuildModuleType.GameEditor;
-				default:
-					return null;
-			}
-		}
-	}
 	/// <summary>
 	/// Distribution level of module.
 	/// Note: The name of each entry is used to search for/create folders
@@ -104,82 +37,6 @@ namespace UnrealBuildTool
 	/// </summary>
 	public abstract class UEBuildModule
 	{
-		/// <summary>
-		/// Known module folders
-		/// </summary>
-		public static readonly string RuntimeFolder = String.Format("{0}Runtime{0}", Path.DirectorySeparatorChar);
-		public static readonly string DeveloperFolder = String.Format("{0}Developer{0}", Path.DirectorySeparatorChar);
-		public static readonly string EditorFolder = String.Format("{0}Editor{0}", Path.DirectorySeparatorChar);
-		public static readonly string ProgramsFolder = String.Format("{0}Programs{0}", Path.DirectorySeparatorChar);
-		public static readonly string ThirdPartyFolder = String.Format("{0}ThirdParty{0}", Path.DirectorySeparatorChar);
-
-		public static UEBuildModuleType GetEngineModuleTypeFromDescriptor(ModuleDescriptor Module)
-		{
-			switch (Module.Type)
-			{
-				case ModuleHostType.Developer:
-					return UEBuildModuleType.EngineDeveloper;
-				case ModuleHostType.Editor:
-				case ModuleHostType.EditorNoCommandlet:
-					return UEBuildModuleType.EngineEditor;
-				case ModuleHostType.Runtime:
-				case ModuleHostType.RuntimeNoCommandlet:
-				case ModuleHostType.RuntimeAndProgram:
-					return UEBuildModuleType.EngineRuntime;
-				default:
-					throw new BuildException("Unhandled engine module type {0}", Module.Type.ToString());
-			}
-		}
-
-		public static UEBuildModuleType GetGameModuleTypeFromDescriptor(ModuleDescriptor Module)
-		{
-			switch (Module.Type)
-			{
-				case ModuleHostType.Developer:
-					return UEBuildModuleType.GameDeveloper;
-				case ModuleHostType.Editor:
-				case ModuleHostType.EditorNoCommandlet:
-					return UEBuildModuleType.GameEditor;
-				case ModuleHostType.Runtime:
-				case ModuleHostType.RuntimeNoCommandlet:
-				case ModuleHostType.RuntimeAndProgram:
-					return UEBuildModuleType.GameRuntime;
-				default:
-					throw new BuildException("Unhandled game module type {0}", Module.Type.ToString());
-			}
-		}
-
-		public static UEBuildModuleType? GetEngineModuleTypeBasedOnLocation(string ModuleName, FileReference ModuleFileName)
-		{
-			string ModuleFileRelativeToEngineDirectory = ModuleFileName.MakeRelativeTo(UnrealBuildTool.EngineDirectory);
-			if (ModuleFileRelativeToEngineDirectory.IndexOf(UEBuildModule.RuntimeFolder, StringComparison.InvariantCultureIgnoreCase) >= 0)
-			{
-				return UEBuildModuleType.EngineRuntime;
-			}
-
-			if (ModuleFileRelativeToEngineDirectory.IndexOf(UEBuildModule.DeveloperFolder, StringComparison.InvariantCultureIgnoreCase) >= 0)
-			{
-				return UEBuildModuleType.EngineDeveloper;
-			}
-
-			if (ModuleFileRelativeToEngineDirectory.IndexOf(UEBuildModule.EditorFolder, StringComparison.InvariantCultureIgnoreCase) >= 0)
-			{
-				return UEBuildModuleType.EngineEditor;
-			}
-	
-			if (ModuleFileRelativeToEngineDirectory.IndexOf(UEBuildModule.ProgramsFolder, StringComparison.InvariantCultureIgnoreCase) >= 0)
-			{
-				return UEBuildModuleType.Program;
-			}
-
-			if (ModuleFileRelativeToEngineDirectory.IndexOf(UEBuildModule.ThirdPartyFolder, StringComparison.InvariantCultureIgnoreCase) >= 0)
-			{
-				return UEBuildModuleType.EngineThirdParty;
-			}
-
-			return null;
-		}
-
 		/// <summary>
 		/// Checks what distribution level a particular path should have by checking for key folders anywhere in it
 		/// </summary>
@@ -189,9 +46,9 @@ namespace UnrealBuildTool
 			FilePath = Path.GetFullPath(FilePath);
 
 			// Check from highest to lowest (don't actually need to check 'Everyone')
-			for (var DistributionLevel = UEBuildModuleDistribution.NoRedist; DistributionLevel > UEBuildModuleDistribution.Public; DistributionLevel--)
+			for (UEBuildModuleDistribution DistributionLevel = UEBuildModuleDistribution.NoRedist; DistributionLevel > UEBuildModuleDistribution.Public; DistributionLevel--)
 			{
-				var DistributionFolderName = String.Format("{0}{1}{0}", Path.DirectorySeparatorChar, DistributionLevel.ToString());
+				string DistributionFolderName = String.Format("{0}{1}{0}", Path.DirectorySeparatorChar, DistributionLevel.ToString());
 				if (FilePath.IndexOf(DistributionFolderName, StringComparison.InvariantCultureIgnoreCase) >= 0)
 				{
 					return DistributionLevel;
@@ -200,8 +57,8 @@ namespace UnrealBuildTool
 				if (DistributionLevel == UEBuildModuleDistribution.NotForLicensees)
 				{
 					// Extra checks for PS4 and XboxOne folders, which are equivalent to NotForLicensees
-					var PS4FolderName = String.Format("{0}ps4{0}", Path.DirectorySeparatorChar);
-					var XboxFolderName = String.Format("{0}xboxone{0}", Path.DirectorySeparatorChar);
+					string PS4FolderName = String.Format("{0}ps4{0}", Path.DirectorySeparatorChar);
+					string XboxFolderName = String.Format("{0}xboxone{0}", Path.DirectorySeparatorChar);
 					if (FilePath.IndexOf(PS4FolderName, StringComparison.InvariantCultureIgnoreCase) >= 0
 					|| FilePath.IndexOf(XboxFolderName, StringComparison.InvariantCultureIgnoreCase) >= 0)
 					{
@@ -248,11 +105,6 @@ namespace UnrealBuildTool
 			return InEnumerableStrings == null ? new HashSet<string>() : new HashSet<string>(InEnumerableStrings);
 		}
 
-		/// <summary>
-		/// The target which owns this module.
-		/// </summary>
-		public readonly UEBuildTarget Target;
-
 		/// <summary>
 		/// The name that uniquely identifies the module.
 		/// </summary>
@@ -261,7 +113,7 @@ namespace UnrealBuildTool
 		/// <summary>
 		/// The type of module being built. Used to switch between debug/development and precompiled/source configurations.
 		/// </summary>
-		public UEBuildModuleType Type;
+		public UHTModuleType Type;
 
 		/// <summary>
 		/// The rules for this module
@@ -283,16 +135,6 @@ namespace UnrealBuildTool
 		/// </summary>
 		public FileReference RulesFile;
 
-		/// <summary>
-		/// Tells if this module can be redistributed.
-		/// </summary>
-		public bool IsRedistributable()
-		{
-			return IsRedistributableOverride.HasValue
-				? IsRedistributableOverride.Value
-				: (Type != UEBuildModuleType.EngineDeveloper && Type != UEBuildModuleType.EngineEditor);
-		}
-
 		/// <summary>
 		/// The binary the module will be linked into for the current target.  Only set after UEBuildBinary.BindModules is called.
 		/// </summary>
@@ -364,21 +206,19 @@ namespace UnrealBuildTool
 		/// Returns a list of this module's immediate dependencies.
 		/// </summary>
 		/// <returns>An enumerable containing the dependencies of the module.</returns>
-		public IEnumerable<UEBuildModule> GetAllDependencyModules()
+		public IEnumerable<UEBuildModule> GetDirectDependencyModules()
 		{
 			return PublicDependencyModules.Concat(PrivateDependencyModules).Concat(DynamicallyLoadedModules).Concat(PlatformSpecificDynamicallyLoadedModules);
 		}
 
 		public UEBuildModule(
-			UEBuildTarget InTarget,
 			string InName,
-			UEBuildModuleType InType,
+			UHTModuleType InType,
 			DirectoryReference InModuleDirectory,
 			ModuleRules InRules,
 			FileReference InRulesFile
 			)
 		{
-			Target = InTarget;
 			Name = InName;
 			Type = InType;
 			ModuleDirectory = InModuleDirectory;
@@ -402,8 +242,6 @@ namespace UnrealBuildTool
 			PrivateIncludePaths = HashSetFromOptionalEnumerableStringParameter(InRules.PrivateIncludePaths);
 			RuntimeDependencies = (InRules.RuntimeDependencies == null) ? new List<RuntimeDependency>() : new List<RuntimeDependency>(InRules.RuntimeDependencies);
 			IsRedistributableOverride = InRules.IsRedistributableOverride;
-
-			Target.RegisterModule(this);
 		}
 
 		/// <summary>
@@ -517,7 +355,7 @@ namespace UnrealBuildTool
 				AdditionalFrameworks.AddRange(PublicAdditionalFrameworks);
 
 				// Remember the module so we can refer to it when needed
-				foreach (var Framework in PublicAdditionalFrameworks)
+				foreach (UEBuildFramework Framework in PublicAdditionalFrameworks)
 				{
 					Framework.OwningModule = this;
 				}
@@ -547,9 +385,9 @@ namespace UnrealBuildTool
 			}
 			else
 			{
-				foreach (var Path in PathsToAdd)
+				foreach (string Path in PathsToAdd)
 				{
-					var NormalizedPath = Path.TrimEnd('/');
+					string NormalizedPath = Path.TrimEnd('/');
 					// If path doesn't exist, it may contain VC macro (which is passed directly to and expanded by compiler).
 					if (Directory.Exists(NormalizedPath) || DoesPathContainVCMacro(NormalizedPath))
 					{
@@ -656,7 +494,7 @@ namespace UnrealBuildTool
 				WeakFrameworks.AddRange(PublicWeakFrameworks);
 				AdditionalBundleResources.AddRange(PublicAdditionalBundleResources);
 				// Remember the module so we can refer to it when needed
-				foreach (var Framework in PublicAdditionalFrameworks)
+				foreach (UEBuildFramework Framework in PublicAdditionalFrameworks)
 				{
 					Framework.OwningModule = this;
 				}
@@ -695,7 +533,7 @@ namespace UnrealBuildTool
 		/// <summary>
 		/// Compiles the module, and returns a list of files output by the compiler.
 		/// </summary>
-		public abstract List<FileItem> Compile(UEToolChain ToolChain, CPPEnvironment GlobalCompileEnvironment, CPPEnvironment CompileEnvironment);
+		public abstract List<FileItem> Compile(UEBuildTarget Target, UEToolChain ToolChain, CPPEnvironment GlobalCompileEnvironment, CPPEnvironment CompileEnvironment);
 
 		// Object interface.
 		public override string ToString()
@@ -726,10 +564,11 @@ namespace UnrealBuildTool
 		/// Gets all of the modules referenced by this module
 		/// </summary>
 		/// <param name="ReferencedModules">Hash of all referenced modules with their addition index.</param>
+		/// <param name="IgnoreReferencedModules">Hashset used to ignore modules which are already added to the list</param>
 		/// <param name="bIncludeDynamicallyLoaded">True if dynamically loaded modules (and all of their dependent modules) should be included.</param>
 		/// <param name="bForceCircular">True if circular dependencies should be processed</param>
 		/// <param name="bOnlyDirectDependencies">True to return only this module's direct dependencies</param>
-		public virtual void GetAllDependencyModules(CaselessDictionary<ModuleIndexPair> ReferencedModules, bool bIncludeDynamicallyLoaded, bool bForceCircular, bool bOnlyDirectDependencies)
+		public virtual void GetAllDependencyModules(List<UEBuildModule> ReferencedModules, HashSet<UEBuildModule> IgnoreReferencedModules, bool bIncludeDynamicallyLoaded, bool bForceCircular, bool bOnlyDirectDependencies)
 		{
 		}
 
@@ -744,7 +583,7 @@ namespace UnrealBuildTool
 		/// <summary>
 		/// Creates all the modules required for this target
 		/// </summary>
-		public void RecursivelyCreateModules()
+		public void RecursivelyCreateModules(UEBuildTarget Target)
 		{
 			// Create all the include path modules. These modules may not be added to the target (and we don't process their dependencies), but they need 
 			// to be created to set up their compile environment.
@@ -769,7 +608,7 @@ namespace UnrealBuildTool
 					UEBuildModule Module = Target.FindOrCreateModuleByName(ModuleName);
 					if (!Modules.Contains(Module))
 					{
-						Module.RecursivelyCreateModules();
+						Module.RecursivelyCreateModules(Target);
 						Modules.Add(Module);
 					}
 				}
@@ -798,15 +637,13 @@ namespace UnrealBuildTool
 	class UEBuildExternalModule : UEBuildModule
 	{
 		public UEBuildExternalModule(
-			UEBuildTarget InTarget,
-			UEBuildModuleType InType,
+			UHTModuleType InType,
 			string InName,
 			DirectoryReference InModuleDirectory,
 			ModuleRules InRules,
 			FileReference InRulesFile
 			)
 			: base(
-				InTarget: InTarget,
 				InType: InType,
 				InName: InName,
 				InModuleDirectory: InModuleDirectory,
@@ -817,7 +654,7 @@ namespace UnrealBuildTool
 		}
 
 		// UEBuildModule interface.
-		public override List<FileItem> Compile(UEToolChain ToolChain, CPPEnvironment GlobalCompileEnvironment, CPPEnvironment CompileEnvironment)
+		public override List<FileItem> Compile(UEBuildTarget Target, UEToolChain ToolChain, CPPEnvironment GlobalCompileEnvironment, CPPEnvironment CompileEnvironment)
 		{
 			return new List<FileItem>();
 		}
@@ -969,7 +806,7 @@ namespace UnrealBuildTool
 		/// </summary>
 		private static void CategorizeSourceFiles(IEnumerable<FileItem> InSourceFiles, SourceFilesClass OutSourceFiles)
 		{
-			foreach (var SourceFile in InSourceFiles)
+			foreach (FileItem SourceFile in InSourceFiles)
 			{
 				string Extension = Path.GetExtension(SourceFile.AbsolutePath).ToUpperInvariant();
 				if (!SourceFile.bExists)
@@ -1004,9 +841,8 @@ namespace UnrealBuildTool
 		}
 
 		public UEBuildModuleCPP(
-			UEBuildTarget InTarget,
 			string InName,
-			UEBuildModuleType InType,
+			UHTModuleType InType,
 			DirectoryReference InModuleDirectory,
 			DirectoryReference InGeneratedCodeDirectory,
 			IntelliSenseGatherer InIntelliSenseGatherer,
@@ -1015,7 +851,7 @@ namespace UnrealBuildTool
 			bool bInBuildSourceFiles,
 			FileReference InRulesFile
 			)
-			: base(InTarget,
+			: base(
 					InName,
 					InType,
 					InModuleDirectory,
@@ -1033,25 +869,25 @@ namespace UnrealBuildTool
 			}
 
 			Definitions = HashSetFromOptionalEnumerableStringParameter(InRules.Definitions);
-			foreach (var Def in Definitions)
+			foreach (string Def in Definitions)
 			{
 				Log.TraceVerbose("Compile Env {0}: {1}", Name, Def);
 			}
 		}
 
 		// UEBuildModule interface.
-		public override List<FileItem> Compile(UEToolChain ToolChain, CPPEnvironment GlobalCompileEnvironment, CPPEnvironment CompileEnvironment)
+		public override List<FileItem> Compile(UEBuildTarget Target, UEToolChain ToolChain, CPPEnvironment GlobalCompileEnvironment, CPPEnvironment CompileEnvironment)
 		{
-			var BuildPlatform = UEBuildPlatform.GetBuildPlatformForCPPTargetPlatform(CompileEnvironment.Config.Target.Platform);
+			UEBuildPlatform BuildPlatform = UEBuildPlatform.GetBuildPlatformForCPPTargetPlatform(CompileEnvironment.Config.Target.Platform);
 
-			var LinkInputFiles = new List<FileItem>();
+			List<FileItem> LinkInputFiles = new List<FileItem>();
 			if (ProjectFileGenerator.bGenerateProjectFiles && IntelliSenseGatherer == null)
 			{
 				// Nothing to do for IntelliSense, bail out early
 				return LinkInputFiles;
 			}
 
-			var ModuleCompileEnvironment = CreateModuleCompileEnvironment(CompileEnvironment);
+			CPPEnvironment ModuleCompileEnvironment = CreateModuleCompileEnvironment(Target, CompileEnvironment);
 			IncludeSearchPaths = ModuleCompileEnvironment.Config.CPPIncludeInfo.IncludePaths.ToList();
 			IncludeSearchPaths.AddRange(ModuleCompileEnvironment.Config.CPPIncludeInfo.SystemIncludePaths.ToList());
 
@@ -1102,10 +938,10 @@ namespace UnrealBuildTool
 
 			{
 				// Process all of the header file dependencies for this module
-				this.CachePCHUsageForModuleSourceFiles(ModuleCompileEnvironment);
+				this.CachePCHUsageForModuleSourceFiles(Target, ModuleCompileEnvironment);
 
 				// Make sure our RC files have cached includes.  
-				foreach (var RCFile in SourceFilesToBuild.RCFiles)
+				foreach (FileItem RCFile in SourceFilesToBuild.RCFiles)
 				{
 					RCFile.CachedCPPIncludeInfo = ModuleCompileEnvironment.Config.CPPIncludeInfo;
 				}
@@ -1113,8 +949,8 @@ namespace UnrealBuildTool
 
 
 			// Check to see if this is an Engine module (including program or plugin modules).  That is, the module is located under the "Engine" folder
-			var IsPluginModule = ModuleDirectory.IsUnderDirectory(DirectoryReference.Combine(Target.ProjectDirectory, "Plugins"));
-			var IsGameModule = !IsPluginModule && !ModuleDirectory.IsUnderDirectory(UnrealBuildTool.EngineDirectory);
+			bool IsPluginModule = ModuleDirectory.IsUnderDirectory(DirectoryReference.Combine(Target.ProjectDirectory, "Plugins"));
+			bool IsGameModule = !IsPluginModule && !ModuleDirectory.IsUnderDirectory(UnrealBuildTool.EngineDirectory);
 
 			// Should we force a precompiled header to be generated for this module?  Usually, we only bother with a
 			// precompiled header if there are at least several source files in the module (after combining them for unity
@@ -1181,13 +1017,13 @@ namespace UnrealBuildTool
 			}
 
 			// The environment with which to compile the CPP files
-			var CPPCompileEnvironment = ModuleCompileEnvironment;
+			CPPEnvironment CPPCompileEnvironment = ModuleCompileEnvironment;
 
 			// Precompiled header support.
 			bool bWasModuleCodeCompiled = false;
 			if (BuildPlatform.ShouldUsePCHFiles(CompileEnvironment.Config.Target.Platform, CompileEnvironment.Config.Target.Configuration))
 			{
-				var PCHGenTimerStart = DateTime.UtcNow;
+				DateTime PCHGenTimerStart = DateTime.UtcNow;
 
 				// The code below will figure out whether this module will either use a "unique PCH" (private PCH that will only be included by
 				// this module's code files), or a "shared PCH" (potentially included by many code files in many modules.)  Only one or the other
@@ -1220,18 +1056,18 @@ namespace UnrealBuildTool
 					if (bAllowSharedPCH && (!bIsASharedPCHModule || bCanModuleUseOwnSharedPCH))
 					{
 						// Figure out which shared PCH tier we're in
-						var ReferencedModules = new CaselessDictionary<ModuleIndexPair>();
+						List<UEBuildModule> ReferencedModules = new List<UEBuildModule>();
 						{
-							this.GetAllDependencyModules(ReferencedModules, bIncludeDynamicallyLoaded: false, bForceCircular: false, bOnlyDirectDependencies: true);
+							this.GetAllDependencyModules(ReferencedModules, new HashSet<UEBuildModule>(), bIncludeDynamicallyLoaded: false, bForceCircular: false, bOnlyDirectDependencies: true);
 						}
 
 						int LargestSharedPCHHeaderFileIndex = -1;
-						foreach (var DependencyModule in ReferencedModules.Values.OrderBy(P => P.Index).Select(P => P.Module))
+						foreach (UEBuildModule DependencyModule in ReferencedModules)
 						{
 							// These Shared PCHs are ordered from least complex to most complex.  We'll start at the last one and search backwards.
-							for (var SharedPCHHeaderFileIndex = GlobalCompileEnvironment.SharedPCHHeaderFiles.Count - 1; SharedPCHHeaderFileIndex > LargestSharedPCHHeaderFileIndex; --SharedPCHHeaderFileIndex)
+							for (int SharedPCHHeaderFileIndex = GlobalCompileEnvironment.SharedPCHHeaderFiles.Count - 1; SharedPCHHeaderFileIndex > LargestSharedPCHHeaderFileIndex; --SharedPCHHeaderFileIndex)
 							{
-								var CurSharedPCHHeaderFile = GlobalCompileEnvironment.SharedPCHHeaderFiles[SharedPCHHeaderFileIndex];
+								SharedPCHHeaderInfo CurSharedPCHHeaderFile = GlobalCompileEnvironment.SharedPCHHeaderFiles[SharedPCHHeaderFileIndex];
 
 								if (DependencyModule == CurSharedPCHHeaderFile.Module ||
 									(bIsASharedPCHModule && CurSharedPCHHeaderFile.Module == this))	// If we ourselves are a shared PCH module, always at least use our own module as our shared PCH header if we can't find anything better
@@ -1257,7 +1093,7 @@ namespace UnrealBuildTool
 						}
 
 						// Keep track of how many modules make use of this PCH for performance diagnostics
-						var LargestSharedPCHHeader = GlobalCompileEnvironment.SharedPCHHeaderFiles[LargestSharedPCHHeaderFileIndex];
+						SharedPCHHeaderInfo LargestSharedPCHHeader = GlobalCompileEnvironment.SharedPCHHeaderFiles[LargestSharedPCHHeaderFileIndex];
 						++LargestSharedPCHHeader.NumModulesUsingThisPCH;
 
 					}
@@ -1299,7 +1135,7 @@ namespace UnrealBuildTool
 					if (PCHToUse != null)
 					{
 						// Update all CPPFiles to point to the PCH
-						foreach (var CPPFile in SourceFilesToBuild.CPPFiles)
+						foreach (FileItem CPPFile in SourceFilesToBuild.CPPFiles)
 						{
 							CPPFile.PCHHeaderNameInCode = PCHToUse.AbsolutePath;
 							CPPFile.PrecompiledHeaderIncludeFilename = PCHToUse.Reference;
@@ -1309,14 +1145,14 @@ namespace UnrealBuildTool
 					// A shared PCH was not already set up for us, so set one up.
 					if (ModulePCHEnvironment == null)
 					{
-						var PCHHeaderFile = ProcessedDependencies.UniquePCHHeaderFile;
-						var PCHModuleName = this.Name;
+						FileItem PCHHeaderFile = ProcessedDependencies.UniquePCHHeaderFile;
+						string PCHModuleName = this.Name;
 						if (SharedPCHHeaderFile != null)
 						{
 							PCHHeaderFile = SharedPCHHeaderFile;
 							PCHModuleName = SharedPCHModuleName;
 						}
-						var PCHHeaderNameInCode = SourceFilesToBuild.CPPFiles[0].PCHHeaderNameInCode;
+						string PCHHeaderNameInCode = SourceFilesToBuild.CPPFiles[0].PCHHeaderNameInCode;
 
 						ModulePCHEnvironment = new PrecompileHeaderEnvironment(PCHModuleName, PCHHeaderNameInCode, PCHHeaderFile, ModuleCompileEnvironment.Config.CLRMode, ModuleCompileEnvironment.Config.OptimizeCode);
 
@@ -1343,7 +1179,7 @@ namespace UnrealBuildTool
 					// Setup a new compile environment for this module's source files.  It's pretty much the exact same as the
 					// module's compile environment, except that it will include a PCH file.
 
-					var ModulePCHCompileEnvironment = ModuleCompileEnvironment.DeepCopy();
+					CPPEnvironment ModulePCHCompileEnvironment = ModuleCompileEnvironment.DeepCopy();
 					ModulePCHCompileEnvironment.Config.PrecompiledHeaderAction = PrecompiledHeaderAction.Include;
 					ModulePCHCompileEnvironment.Config.PrecompiledHeaderIncludeFilename = ModulePCHEnvironment.PrecompiledHeaderIncludeFilename.Reference;
 					ModulePCHCompileEnvironment.Config.PCHHeaderNameInCode = ModulePCHEnvironment.PCHHeaderNameInCode;
@@ -1355,7 +1191,7 @@ namespace UnrealBuildTool
 						ModulePCHCompileEnvironment.Config.bForceIncludePrecompiledHeader = true;
 					}
 
-					var CPPFilesToBuild = SourceFilesToBuild.CPPFiles;
+					List<FileItem> CPPFilesToBuild = SourceFilesToBuild.CPPFiles;
 					if (bModuleUsesUnityBuild)
 					{
 						// unity files generated for only the set of files which share the same PCH environment
@@ -1436,14 +1272,14 @@ namespace UnrealBuildTool
 
 				if (BuildConfiguration.bPrintPerformanceInfo)
 				{
-					var PCHGenTime = (DateTime.UtcNow - PCHGenTimerStart).TotalSeconds;
+					double PCHGenTime = (DateTime.UtcNow - PCHGenTimerStart).TotalSeconds;
 					TotalPCHGenTime += PCHGenTime;
 				}
 			}
 
 			if (!bWasModuleCodeCompiled && SourceFilesToBuild.CPPFiles.Count > 0)
 			{
-				var CPPFilesToCompile = SourceFilesToBuild.CPPFiles;
+				List<FileItem> CPPFilesToCompile = SourceFilesToBuild.CPPFiles;
 				if (bModuleUsesUnityBuild)
 				{
 					CPPFilesToCompile = Unity.GenerateUnityCPPs(ToolChain, Target, CPPFilesToCompile, CPPCompileEnvironment, Name);
@@ -1456,9 +1292,9 @@ namespace UnrealBuildTool
 				string[] GeneratedFiles = Directory.GetFiles(Path.GetDirectoryName(AutoGenerateCppInfo.BuildInfo.FileWildcard), Path.GetFileName(AutoGenerateCppInfo.BuildInfo.FileWildcard));
 				foreach (string GeneratedFilename in GeneratedFiles)
 				{
-					var GeneratedCppFileItem = FileItem.GetItemByPath(GeneratedFilename);
+					FileItem GeneratedCppFileItem = FileItem.GetItemByPath(GeneratedFilename);
 
-					CachePCHUsageForModuleSourceFile(this.Target, CPPCompileEnvironment, GeneratedCppFileItem);
+					CachePCHUsageForModuleSourceFile(Target, CPPCompileEnvironment, GeneratedCppFileItem);
 
 					// @todo ubtmake: Check for ALL other places where we might be injecting .cpp or .rc files for compiling without caching CachedCPPIncludeInfo first (anything platform specific?)
 					LinkInputFiles.AddRange(ToolChain.CompileCPPFiles(Target, CPPCompileEnvironment, new List<FileItem> { GeneratedCppFileItem }, Name).ObjectFiles);
@@ -1483,8 +1319,8 @@ namespace UnrealBuildTool
 		private PrecompileHeaderEnvironment ApplySharedPCH(CPPEnvironment GlobalCompileEnvironment, CPPEnvironment CompileEnvironment, CPPEnvironment ModuleCompileEnvironment, List<FileItem> CPPFiles, ref FileItem SharedPCHHeaderFile)
 		{
 			// Check to see if we have a PCH header already setup that we can use
-			var SharedPCHHeaderFileCopy = SharedPCHHeaderFile;
-			var SharedPCHEnvironment = GlobalCompileEnvironment.SharedPCHEnvironments.Find(Env => Env.PrecompiledHeaderIncludeFilename == SharedPCHHeaderFileCopy);
+			FileItem SharedPCHHeaderFileCopy = SharedPCHHeaderFile;
+			PrecompileHeaderEnvironment SharedPCHEnvironment = GlobalCompileEnvironment.SharedPCHEnvironments.Find(Env => Env.PrecompiledHeaderIncludeFilename == SharedPCHHeaderFileCopy);
 			if (SharedPCHEnvironment == null)
 			{
 				return null;
@@ -1506,8 +1342,8 @@ namespace UnrealBuildTool
 			}
 
 			// Don't mix non-optimized code with optimized code (PCHs won't be compatible)
-			var SharedPCHCodeOptimization = SharedPCHEnvironment.OptimizeCode;
-			var ModuleCodeOptimization = ModuleCompileEnvironment.Config.OptimizeCode;
+			ModuleRules.CodeOptimization SharedPCHCodeOptimization = SharedPCHEnvironment.OptimizeCode;
+			ModuleRules.CodeOptimization ModuleCodeOptimization = ModuleCompileEnvironment.Config.OptimizeCode;
 
 			if (CompileEnvironment.Config.Target.Configuration != CPPTargetConfiguration.Debug)
 			{
@@ -1539,20 +1375,20 @@ namespace UnrealBuildTool
 				throw new BuildException("Required source file not found: " + CPPFile.AbsolutePath);
 			}
 
-			var PCHCacheTimerStart = DateTime.UtcNow;
+			DateTime PCHCacheTimerStart = DateTime.UtcNow;
 
-			var BuildPlatform = UEBuildPlatform.GetBuildPlatformForCPPTargetPlatform(ModuleCompileEnvironment.Config.Target.Platform);
-			var IncludePathsToSearch = ModuleCompileEnvironment.Config.CPPIncludeInfo.GetIncludesPathsToSearch(CPPFile);
+			UEBuildPlatform BuildPlatform = UEBuildPlatform.GetBuildPlatformForCPPTargetPlatform(ModuleCompileEnvironment.Config.Target.Platform);
+			List<string> IncludePathsToSearch = ModuleCompileEnvironment.Config.CPPIncludeInfo.GetIncludesPathsToSearch(CPPFile);
 
 			// Store the module compile environment along with the .cpp file.  This is so that we can use it later on when looking
 			// for header dependencies
 			CPPFile.CachedCPPIncludeInfo = ModuleCompileEnvironment.Config.CPPIncludeInfo;
 
-			var PCHFile = CachePCHUsageForCPPFile(Target, CPPFile, BuildPlatform, IncludePathsToSearch, ModuleCompileEnvironment.Config.CPPIncludeInfo.IncludeFileSearchDictionary);
+			FileItem PCHFile = CachePCHUsageForCPPFile(Target, CPPFile, BuildPlatform, IncludePathsToSearch, ModuleCompileEnvironment.Config.CPPIncludeInfo.IncludeFileSearchDictionary);
 
 			if (BuildConfiguration.bPrintPerformanceInfo)
 			{
-				var PCHCacheTime = (DateTime.UtcNow - PCHCacheTimerStart).TotalSeconds;
+				double PCHCacheTime = (DateTime.UtcNow - PCHCacheTimerStart).TotalSeconds;
 				TotalPCHCacheTime += PCHCacheTime;
 			}
 
@@ -1560,28 +1396,28 @@ namespace UnrealBuildTool
 		}
 
 
-		public void CachePCHUsageForModuleSourceFiles(CPPEnvironment ModuleCompileEnvironment)
+		public void CachePCHUsageForModuleSourceFiles(UEBuildTarget Target, CPPEnvironment ModuleCompileEnvironment)
 		{
 			if (ProcessedDependencies == null)
 			{
-				var PCHCacheTimerStart = DateTime.UtcNow;
+				DateTime PCHCacheTimerStart = DateTime.UtcNow;
 
-				var BuildPlatform = UEBuildPlatform.GetBuildPlatformForCPPTargetPlatform(ModuleCompileEnvironment.Config.Target.Platform);
+				UEBuildPlatform BuildPlatform = UEBuildPlatform.GetBuildPlatformForCPPTargetPlatform(ModuleCompileEnvironment.Config.Target.Platform);
 
 				bool bFoundAProblemWithPCHs = false;
 
 				FileItem UniquePCH = null;
-				foreach (var CPPFile in SourceFilesFound.CPPFiles)	// @todo ubtmake: We're not caching CPPEnvironments for .c/.mm files, etc.  Even though they don't use PCHs, they still have #includes!  This can break dependency checking!
+				foreach (FileItem CPPFile in SourceFilesFound.CPPFiles)	// @todo ubtmake: We're not caching CPPEnvironments for .c/.mm files, etc.  Even though they don't use PCHs, they still have #includes!  This can break dependency checking!
 				{
 					// Build a single list of include paths to search.
-					var IncludePathsToSearch = ModuleCompileEnvironment.Config.CPPIncludeInfo.GetIncludesPathsToSearch(CPPFile);
+					List<string> IncludePathsToSearch = ModuleCompileEnvironment.Config.CPPIncludeInfo.GetIncludesPathsToSearch(CPPFile);
 
 					// Store the module compile environment along with the .cpp file.  This is so that we can use it later on when looking
 					// for header dependencies
 					CPPFile.CachedCPPIncludeInfo = ModuleCompileEnvironment.Config.CPPIncludeInfo;
 
 					// Find headers used by the source file.
-					var PCH = UEBuildModuleCPP.CachePCHUsageForCPPFile(Target, CPPFile, BuildPlatform, IncludePathsToSearch, ModuleCompileEnvironment.Config.CPPIncludeInfo.IncludeFileSearchDictionary);
+					FileItem PCH = UEBuildModuleCPP.CachePCHUsageForCPPFile(Target, CPPFile, BuildPlatform, IncludePathsToSearch, ModuleCompileEnvironment.Config.CPPIncludeInfo.IncludeFileSearchDictionary);
 					if (PCH == null)
 					{
 						throw new BuildException("Source file \"{0}\" is not including any headers.  We expect all modules to include a header file for precompiled header generation.  Please add an #include statement.", CPPFile.AbsolutePath);
@@ -1605,8 +1441,8 @@ namespace UnrealBuildTool
 				if (bFoundAProblemWithPCHs)
 				{
 					// Map from pch header string to the source files that use that PCH
-					var UsageMapPCH = new Dictionary<FileReference, List<FileItem>>();
-					foreach (var CPPFile in SourceFilesToBuild.CPPFiles)
+					Dictionary<FileReference, List<FileItem>> UsageMapPCH = new Dictionary<FileReference, List<FileItem>>();
+					foreach (FileItem CPPFile in SourceFilesToBuild.CPPFiles)
 					{
 						// Create a new entry if not in the pch usage map
 						UsageMapPCH.GetOrAddNew(CPPFile.PrecompiledHeaderIncludeFilename).Add(CPPFile);
@@ -1616,7 +1452,7 @@ namespace UnrealBuildTool
 					{
 						Log.TraceVerbose("{0} PCH files for module {1}:", UsageMapPCH.Count, Name);
 						int MostFilesIncluded = 0;
-						foreach (var CurPCH in UsageMapPCH)
+						foreach (KeyValuePair<FileReference, List<FileItem>> CurPCH in UsageMapPCH)
 						{
 							if (CurPCH.Value.Count > MostFilesIncluded)
 							{
@@ -1632,17 +1468,17 @@ namespace UnrealBuildTool
 						// Keep track of the PCH file that is most used within this module
 						FileReference MostFilesAreIncludingPCH = null;
 						int MostFilesIncluded = 0;
-						foreach (var CurPCH in UsageMapPCH.Where(PCH => PCH.Value.Count > MostFilesIncluded))
+						foreach (KeyValuePair<FileReference, List<FileItem>> CurPCH in UsageMapPCH.Where(PCH => PCH.Value.Count > MostFilesIncluded))
 						{
 							MostFilesAreIncludingPCH = CurPCH.Key;
 							MostFilesIncluded = CurPCH.Value.Count;
 						}
 
 						// Find all of the files that are not including our "best" PCH header
-						var FilesNotIncludingBestPCH = new StringBuilder();
-						foreach (var CurPCH in UsageMapPCH.Where(PCH => PCH.Key != MostFilesAreIncludingPCH))
+						StringBuilder FilesNotIncludingBestPCH = new StringBuilder();
+						foreach (KeyValuePair<FileReference, List<FileItem>> CurPCH in UsageMapPCH.Where(PCH => PCH.Key != MostFilesAreIncludingPCH))
 						{
-							foreach (var SourceFile in CurPCH.Value)
+							foreach (FileItem SourceFile in CurPCH.Value)
 							{
 								FilesNotIncludingBestPCH.AppendFormat("{0} (including {1})\n", SourceFile.AbsolutePath, CurPCH.Key);
 							}
@@ -1659,7 +1495,7 @@ namespace UnrealBuildTool
 
 				if (BuildConfiguration.bPrintPerformanceInfo)
 				{
-					var PCHCacheTime = (DateTime.UtcNow - PCHCacheTimerStart).TotalSeconds;
+					double PCHCacheTime = (DateTime.UtcNow - PCHCacheTimerStart).TotalSeconds;
 					TotalPCHCacheTime += PCHCacheTime;
 				}
 			}
@@ -1682,7 +1518,7 @@ namespace UnrealBuildTool
 				return null;
 			}
 
-			var FirstInclude = DirectIncludeFilenames[0];
+			DependencyInclude FirstInclude = DirectIncludeFilenames[0];
 
 			// The pch header should always be the first include in the source file.
 			// NOTE: This is not an absolute path.  This is just the literal include string from the source file!
@@ -1718,9 +1554,9 @@ namespace UnrealBuildTool
 		/// </summary>
 		/// <param name="BaseCompileEnvironment">An existing environment to base the module compile environment on.</param>
 		/// <returns>The new module compile environment.</returns>
-		public CPPEnvironment CreateModuleCompileEnvironment(CPPEnvironment BaseCompileEnvironment)
+		public CPPEnvironment CreateModuleCompileEnvironment(UEBuildTarget Target, CPPEnvironment BaseCompileEnvironment)
 		{
-			var Result = BaseCompileEnvironment.DeepCopy();
+			CPPEnvironment Result = BaseCompileEnvironment.DeepCopy();
 
 			// Override compile environment
 			Result.Config.bFasterWithoutUnity = Rules.bFasterWithoutUnity;
@@ -1787,19 +1623,19 @@ namespace UnrealBuildTool
 		/// If any of this module's source files contain UObject definitions, this will return those header files back to the caller
 		/// </summary>
 		/// <returns>
-		public UHTModuleInfoCacheType GetCachedUHTModuleInfo()
+		public UHTModuleInfoCacheType GetCachedUHTModuleInfo(EGeneratedCodeVersion GeneratedCodeVersion)
 		{
 			if (UHTModuleInfoCache == null)
 			{
 				IEnumerable<string> HeaderFilenames = Directory.GetFiles(ModuleDirectory.FullName, "*.h", SearchOption.AllDirectories);
-				UHTModuleInfo Info = ExternalExecution.CreateUHTModuleInfo(HeaderFilenames, Target, Name, ModuleDirectory, Type);
+				UHTModuleInfo Info = ExternalExecution.CreateUHTModuleInfo(HeaderFilenames, Name, ModuleDirectory, Type, GeneratedCodeVersion);
 				UHTModuleInfoCache = new UHTModuleInfoCacheType(Info.PublicUObjectHeaders.Concat(Info.PublicUObjectClassesHeaders).Concat(Info.PrivateUObjectHeaders).Select(x => x.AbsolutePath).ToList(), Info);
 			}
 
 			return UHTModuleInfoCache;
 		}
 
-		public override void GetAllDependencyModules(CaselessDictionary<ModuleIndexPair> ReferencedModules, bool bIncludeDynamicallyLoaded, bool bForceCircular, bool bOnlyDirectDependencies)
+		public override void GetAllDependencyModules(List<UEBuildModule> ReferencedModules, HashSet<UEBuildModule> IgnoreReferencedModules, bool bIncludeDynamicallyLoaded, bool bForceCircular, bool bOnlyDirectDependencies)
 		{
 			List<UEBuildModule> AllDependencyModules = new List<UEBuildModule>();
 			AllDependencyModules.AddRange(PrivateDependencyModules);
@@ -1812,21 +1648,21 @@ namespace UnrealBuildTool
 
 			foreach (UEBuildModule DependencyModule in AllDependencyModules)
 			{
-				if (!ReferencedModules.ContainsKey(DependencyModule.Name))
+				if (!IgnoreReferencedModules.Contains(DependencyModule))
 				{
 					// Don't follow circular back-references!
 					bool bIsCircular = HasCircularDependencyOn(DependencyModule.Name);
 					if (bForceCircular || !bIsCircular)
 					{
-						ReferencedModules[DependencyModule.Name] = null;
+						IgnoreReferencedModules.Add(DependencyModule);
 
 						if (!bOnlyDirectDependencies)
 						{
 							// Recurse into dependent modules first
-							DependencyModule.GetAllDependencyModules(ReferencedModules, bIncludeDynamicallyLoaded, bForceCircular, bOnlyDirectDependencies);
+							DependencyModule.GetAllDependencyModules(ReferencedModules, IgnoreReferencedModules, bIncludeDynamicallyLoaded, bForceCircular, bOnlyDirectDependencies);
 						}
 
-						ReferencedModules[DependencyModule.Name] = new ModuleIndexPair { Module = DependencyModule, Index = ReferencedModules.Where(x => x.Value != null).Count() - 1 };
+						ReferencedModules.Add(DependencyModule);
 					}
 				}
 			}
@@ -1877,9 +1713,8 @@ namespace UnrealBuildTool
 		HashSet<string> PrivateAssemblyReferences;
 
 		public UEBuildModuleCPPCLR(
-			UEBuildTarget InTarget,
 			string InName,
-			UEBuildModuleType InType,
+			UHTModuleType InType,
 			DirectoryReference InModuleDirectory,
 			DirectoryReference InGeneratedCodeDirectory,
 			IntelliSenseGatherer InIntelliSenseGatherer,
@@ -1888,7 +1723,7 @@ namespace UnrealBuildTool
 			bool bInBuildSourceFiles,
 			FileReference InRulesFile
 			)
-			: base(InTarget, InName, InType, InModuleDirectory, InGeneratedCodeDirectory, InIntelliSenseGatherer,
+			: base(InName, InType, InModuleDirectory, InGeneratedCodeDirectory, InIntelliSenseGatherer,
 			InSourceFiles, InRules,
 			bInBuildSourceFiles, InRulesFile)
 		{
@@ -1896,21 +1731,21 @@ namespace UnrealBuildTool
 		}
 
 		// UEBuildModule interface.
-		public override List<FileItem> Compile(UEToolChain ToolChain, CPPEnvironment GlobalCompileEnvironment, CPPEnvironment CompileEnvironment)
+		public override List<FileItem> Compile(UEBuildTarget Target, UEToolChain ToolChain, CPPEnvironment GlobalCompileEnvironment, CPPEnvironment CompileEnvironment)
 		{
-			var ModuleCLREnvironment = CompileEnvironment.DeepCopy();
+			CPPEnvironment ModuleCLREnvironment = CompileEnvironment.DeepCopy();
 
 			// Setup the module environment for the project CLR mode
 			ModuleCLREnvironment.Config.CLRMode = CPPCLRMode.CLREnabled;
 
 			// Add the private assembly references to the compile environment.
-			foreach (var PrivateAssemblyReference in PrivateAssemblyReferences)
+			foreach (string PrivateAssemblyReference in PrivateAssemblyReferences)
 			{
 				ModuleCLREnvironment.AddPrivateAssembly(PrivateAssemblyReference);
 			}
 
 			// Pass the CLR compilation environment to the standard C++ module compilation code.
-			return base.Compile(ToolChain, GlobalCompileEnvironment, ModuleCLREnvironment);
+			return base.Compile(Target, ToolChain, GlobalCompileEnvironment, ModuleCLREnvironment);
 		}
 
 		public override void SetupPrivateLinkEnvironment(
@@ -2040,7 +1875,7 @@ namespace UnrealBuildTool
 				FileItem DummyPCH = FileItem.CreateIntermediateTextFile(DummyPath, TmpPCHHeaderContents);
 
 				// Create a new C++ environment that is used to create the PCH.
-				var ProjectPCHEnvironment = ProjectCPPEnvironment.DeepCopy();
+				CPPEnvironment ProjectPCHEnvironment = ProjectCPPEnvironment.DeepCopy();
 				ProjectPCHEnvironment.Config.PrecompiledHeaderAction = PrecompiledHeaderAction.Create;
 				ProjectPCHEnvironment.Config.PrecompiledHeaderIncludeFilename = PrecompiledHeaderIncludeFilename.Reference;
 				ProjectPCHEnvironment.Config.PCHHeaderNameInCode = PCHHeaderNameInCode;
@@ -2048,10 +1883,10 @@ namespace UnrealBuildTool
 
 				if (!bAllowDLLExports)
 				{
-					for (var CurDefinitionIndex = 0; CurDefinitionIndex < ProjectPCHEnvironment.Config.Definitions.Count; ++CurDefinitionIndex)
+					for (int CurDefinitionIndex = 0; CurDefinitionIndex < ProjectPCHEnvironment.Config.Definitions.Count; ++CurDefinitionIndex)
 					{
 						// We change DLLEXPORT to DLLIMPORT for "shared" PCH headers
-						var OldDefinition = ProjectPCHEnvironment.Config.Definitions[CurDefinitionIndex];
+						string OldDefinition = ProjectPCHEnvironment.Config.Definitions[CurDefinitionIndex];
 						if (OldDefinition.EndsWith("=DLLEXPORT"))
 						{
 							ProjectPCHEnvironment.Config.Definitions[CurDefinitionIndex] = OldDefinition.Replace("DLLEXPORT", "DLLIMPORT");
