Copying //UE4/Dev-Build to //UE4/Main

diff --git a/Engine/Source/Programs/UnrealBuildTool/Configuration/UEBuildPlatform.cs b/Engine/Source/Programs/UnrealBuildTool/Configuration/UEBuildPlatform.cs
--- a/Engine/Source/Programs/UnrealBuildTool/Configuration/UEBuildPlatform.cs
+++ b/Engine/Source/Programs/UnrealBuildTool/Configuration/UEBuildPlatform.cs
@@ -467,7 +467,7 @@ namespace UnrealBuildTool
 		/// <param name="Only">  If this is not unknown, then only run that platform</param>
 		public static void PlatformModifyHostModuleRules(string ModuleName, ModuleRules Rules, TargetInfo Target, UnrealTargetPlatform Only = UnrealTargetPlatform.Unknown)
 		{
-			foreach (var PlatformEntry in BuildPlatformDictionary)
+			foreach (KeyValuePair<UnrealTargetPlatform, UEBuildPlatform> PlatformEntry in BuildPlatformDictionary)
 			{
 				if (Only == UnrealTargetPlatform.Unknown || PlatformEntry.Key == Only)
 				{
@@ -531,7 +531,7 @@ namespace UnrealBuildTool
 		public static bool PlatformRequiresMonolithicBuilds(UnrealTargetPlatform InPlatform, UnrealTargetConfiguration InConfiguration)
 		{
 			// Some platforms require monolithic builds...
-			var BuildPlatform = GetBuildPlatform(InPlatform, true);
+			UEBuildPlatform BuildPlatform = GetBuildPlatform(InPlatform, true);
 			if (BuildPlatform != null)
 			{
 				return BuildPlatform.ShouldCompileMonolithicBinary(InPlatform);
@@ -1100,7 +1100,7 @@ namespace UnrealBuildTool
 					//HookProcess.StartInfo.RedirectStandardOutput = true;
 					//HookProcess.StartInfo.RedirectStandardError = true;					
 
-					using (var HookTimer = new ScopedTimer("Time to run hook: ", bShouldLogInfo ? LogEventType.Log : LogEventType.Verbose))
+					using (ScopedTimer HookTimer = new ScopedTimer("Time to run hook: ", bShouldLogInfo ? LogEventType.Log : LogEventType.Verbose))
 					{
 						//installers may require administrator access to succeed. so run as an admmin.
 						HookProcess.StartInfo.Verb = "runas";
diff --git a/Engine/Source/Programs/UnrealBuildTool/Configuration/UEBuildTarget.cs b/Engine/Source/Programs/UnrealBuildTool/Configuration/UEBuildTarget.cs
--- a/Engine/Source/Programs/UnrealBuildTool/Configuration/UEBuildTarget.cs
+++ b/Engine/Source/Programs/UnrealBuildTool/Configuration/UEBuildTarget.cs
@@ -144,13 +144,40 @@ namespace UnrealBuildTool
 	}
 
 	/// <summary>
-	/// A list of extenal files required to build a given target
+	/// A list of external files required to build a given target
 	/// </summary>
 	public class ExternalFileList
 	{
 		public readonly List<string> FileNames = new List<string>();
 	}
 
+	[Serializable]
+	public class FlatModuleCsDataType : ISerializable
+	{
+		public FlatModuleCsDataType(SerializationInfo Info, StreamingContext Context)
+		{
+			BuildCsFilename = Info.GetString("bf");
+			ModuleSourceFolder = (DirectoryReference)Info.GetValue("mf", typeof(DirectoryReference));
+			UHTHeaderNames = (List<string>)Info.GetValue("hn", typeof(List<string>));
+		}
+
+		public void GetObjectData(SerializationInfo Info, StreamingContext Context)
+		{
+			Info.AddValue("bf", BuildCsFilename);
+			Info.AddValue("mf", ModuleSourceFolder);
+			Info.AddValue("hn", UHTHeaderNames);
+		}
+
+		public FlatModuleCsDataType(string InBuildCsFilename)
+		{
+			BuildCsFilename = InBuildCsFilename;
+		}
+
+		public string BuildCsFilename;
+		public DirectoryReference ModuleSourceFolder;
+		public List<string> UHTHeaderNames = new List<string>();
+	}
+
 	[Serializable]
 	public class OnlyModule : ISerializable
 	{
@@ -252,15 +279,15 @@ namespace UnrealBuildTool
 
 		public static List<TargetDescriptor> ParseTargetCommandLine(string[] SourceArguments, ref FileReference ProjectFile)
 		{
-			var Targets = new List<TargetDescriptor>();
+			List<TargetDescriptor> Targets = new List<TargetDescriptor>();
 
 			string TargetName = null;
-			var AdditionalDefinitions = new List<string>();
-			var Platform = UnrealTargetPlatform.Unknown;
-			var Configuration = UnrealTargetConfiguration.Unknown;
+			List<string> AdditionalDefinitions = new List<string>();
+			UnrealTargetPlatform Platform = UnrealTargetPlatform.Unknown;
+			UnrealTargetConfiguration Configuration = UnrealTargetConfiguration.Unknown;
 			string Architecture = null;
 			string RemoteRoot = null;
-			var OnlyModules = new List<OnlyModule>();
+			List<OnlyModule> OnlyModules = new List<OnlyModule>();
 			List<FileReference> ForeignPlugins = new List<FileReference>();
 			string ForceReceiptFileName = null;
 
@@ -294,7 +321,7 @@ namespace UnrealBuildTool
 								{
 									throw new BuildException("Expected module name after -Module argument, but found nothing.");
 								}
-								var OnlyModuleName = Arguments[++ArgumentIndex];
+								string OnlyModuleName = Arguments[++ArgumentIndex];
 
 								OnlyModules.Add(new OnlyModule(OnlyModuleName));
 							}
@@ -308,8 +335,8 @@ namespace UnrealBuildTool
 									throw new BuildException("Expected module name and module suffix -ModuleWithSuffix argument");
 								}
 
-								var OnlyModuleName = Arguments[++ArgumentIndex];
-								var OnlyModuleSuffix = Arguments[++ArgumentIndex];
+								string OnlyModuleName = Arguments[++ArgumentIndex];
+								string OnlyModuleSuffix = Arguments[++ArgumentIndex];
 
 								OnlyModules.Add(new OnlyModule(OnlyModuleName, OnlyModuleSuffix));
 							}
@@ -505,7 +532,7 @@ namespace UnrealBuildTool
 				// We have possible targets!
 				foreach (string PossibleTargetName in PossibleTargetNames)
 				{
-					// If running Rocket, the PossibleTargetName could contain a path
+					// If Engine is installed, the PossibleTargetName could contain a path
 					TargetName = PossibleTargetName;
 
 					// If a project file was not specified see if we can find one
@@ -557,7 +584,7 @@ namespace UnrealBuildTool
 		/// <returns>The build target object for the specified build rules source file</returns>
 		public static UEBuildTarget CreateTarget(TargetDescriptor Desc)
 		{
-			var CreateTargetStartTime = DateTime.UtcNow;
+			DateTime CreateTargetStartTime = DateTime.UtcNow;
 
 			RulesAssembly RulesAssembly;
 			if (Desc.ProjectFile != null)
@@ -581,7 +608,7 @@ namespace UnrealBuildTool
 			if (Desc.bIsEditorRecompile)
 			{
 				// Now that we found the actual Editor target, make sure we're no longer using the old TargetName (which is the Game target)
-				var TargetSuffixIndex = RulesObject.TargetName.LastIndexOf("Target");
+				int TargetSuffixIndex = RulesObject.TargetName.LastIndexOf("Target");
 				Desc.TargetName = (TargetSuffixIndex > 0) ? RulesObject.TargetName.Substring(0, TargetSuffixIndex) : RulesObject.TargetName;
 			}
 			if ((ProjectFileGenerator.bGenerateProjectFiles == false) && (RulesObject.SupportsPlatform(Desc.Platform) == false))
@@ -616,7 +643,7 @@ namespace UnrealBuildTool
 
 			if (BuildConfiguration.bPrintPerformanceInfo)
 			{
-				var CreateTargetTime = (DateTime.UtcNow - CreateTargetStartTime).TotalSeconds;
+				double CreateTargetTime = (DateTime.UtcNow - CreateTargetStartTime).TotalSeconds;
 				Log.TraceInformation("CreateTarget for " + Desc.TargetName + " took " + CreateTargetTime + "s");
 			}
 
@@ -636,7 +663,7 @@ namespace UnrealBuildTool
 			Platform = UnrealTargetPlatform.Unknown;
 			Configuration = UnrealTargetConfiguration.Unknown;
 
-			foreach (var CurArgument in SourceArguments)
+			foreach (string CurArgument in SourceArguments)
 			{
 				UnrealTargetPlatform ParsedPlatform = UEBuildPlatform.ConvertStringToPlatform(CurArgument);
 				if (ParsedPlatform != UnrealTargetPlatform.Unknown)
@@ -875,7 +902,7 @@ namespace UnrealBuildTool
 		public List<PluginInfo> BuildPlugins;
 
 		/// <summary>
-		/// All plugin dependencies for this target. This differs from the list of plugins that is built for Rocket, where we build everything, but link in only the enabled plugins.
+		/// All plugin dependencies for this target. This differs from the list of plugins that is built for Launcher, where we build everything, but link in only the enabled plugins.
 		/// </summary>
 		[NonSerialized]
 		public List<PluginInfo> EnabledPlugins;
@@ -936,33 +963,6 @@ namespace UnrealBuildTool
 		[NonSerialized]
 		private Dictionary<string, UEBuildModule> Modules = new CaselessDictionary<UEBuildModule>();
 
-		[Serializable]
-		public class FlatModuleCsDataType : ISerializable
-		{
-			public FlatModuleCsDataType(SerializationInfo Info, StreamingContext Context)
-			{
-				BuildCsFilename = Info.GetString("bf");
-				ModuleSourceFolder = (DirectoryReference)Info.GetValue("mf", typeof(DirectoryReference));
-				UHTHeaderNames = (List<string>)Info.GetValue("hn", typeof(List<string>));
-			}
-
-			public void GetObjectData(SerializationInfo Info, StreamingContext Context)
-			{
-				Info.AddValue("bf", BuildCsFilename);
-				Info.AddValue("mf", ModuleSourceFolder);
-				Info.AddValue("hn", UHTHeaderNames);
-			}
-
-			public FlatModuleCsDataType(string InBuildCsFilename)
-			{
-				BuildCsFilename = InBuildCsFilename;
-			}
-
-			public string BuildCsFilename;
-			public DirectoryReference ModuleSourceFolder;
-			public List<string> UHTHeaderNames = new List<string>();
-		}
-
 		/// <summary>
 		/// Used to map names of modules to their .Build.cs filename
 		/// </summary>
@@ -1059,8 +1059,8 @@ namespace UnrealBuildTool
 			bEditorRecompile = Info.GetBoolean("er");
 			OnlyModules = (List<OnlyModule>)Info.GetValue("om", typeof(List<OnlyModule>));
 			bCompileMonolithic = Info.GetBoolean("cm");
-			var FlatModuleCsDataKeys = (string[])Info.GetValue("fk", typeof(string[]));
-			var FlatModuleCsDataValues = (FlatModuleCsDataType[])Info.GetValue("fv", typeof(FlatModuleCsDataType[]));
+			string[] FlatModuleCsDataKeys = (string[])Info.GetValue("fk", typeof(string[]));
+			FlatModuleCsDataType[] FlatModuleCsDataValues = (FlatModuleCsDataType[])Info.GetValue("fv", typeof(FlatModuleCsDataType[]));
 			for (int Index = 0; Index != FlatModuleCsDataKeys.Length; ++Index)
 			{
 				FlatModuleCsData.Add(FlatModuleCsDataKeys[Index], FlatModuleCsDataValues[Index]);
@@ -1316,7 +1316,7 @@ namespace UnrealBuildTool
 		{
 			if (!UnrealBuildTool.IsEngineInstalled())
 			{
-				var UBTArguments = new StringBuilder();
+				StringBuilder UBTArguments = new StringBuilder();
 
 				UBTArguments.Append("UnrealHeaderTool");
 				// Which desktop platform do we need to clean UHT for?
@@ -1348,12 +1348,12 @@ namespace UnrealBuildTool
 			{
 				Log.TraceVerbose("Cleaning target {0} - AppName {1}", TargetName, AppName);
 
-				var TargetFilename = RulesAssembly.GetTargetFileName(TargetName);
+				FileReference TargetFilename = RulesAssembly.GetTargetFileName(TargetName);
 				Log.TraceVerbose("\tTargetFilename {0}", TargetFilename);
 
 				// Collect all files to delete.
-				var AdditionalFileExtensions = new string[] { ".lib", ".exp", ".dll.response" };
-				var AllFilesToDelete = new List<string>();
+				string[] AdditionalFileExtensions = new string[] { ".lib", ".exp", ".dll.response" };
+				List<string> AllFilesToDelete = new List<string>();
 				foreach (BuildProduct BuildProduct in ReceiptWithFullPaths.BuildProducts)
 				{
 					// If we're cleaning, don't add any precompiled binaries to the manifest. We don't want to delete them.
@@ -1375,13 +1375,13 @@ namespace UnrealBuildTool
 					}
 
 					AllFilesToDelete.Add(BuildProduct.Path);
-					var FileExt = Path.GetExtension(BuildProduct.Path);
+					string FileExt = Path.GetExtension(BuildProduct.Path);
 					if (FileExt == ".dll" || FileExt == ".exe")
 					{
-						var ManifestFileWithoutExtension = Utils.GetPathWithoutExtension(BuildProduct.Path);
-						foreach (var AdditionalExt in AdditionalFileExtensions)
+						string ManifestFileWithoutExtension = Utils.GetPathWithoutExtension(BuildProduct.Path);
+						foreach (string AdditionalExt in AdditionalFileExtensions)
 						{
-							var AdditionalFileToDelete = ManifestFileWithoutExtension + AdditionalExt;
+							string AdditionalFileToDelete = ManifestFileWithoutExtension + AdditionalExt;
 							AllFilesToDelete.Add(AdditionalFileToDelete);
 						}
 					}
@@ -1393,7 +1393,7 @@ namespace UnrealBuildTool
 
 				//@todo. This does not clean up files that are no longer built by the target...				
 				// Delete all output files listed in the manifest as well as any additional files.
-				foreach (var FileToDelete in AllFilesToDelete)
+				foreach (string FileToDelete in AllFilesToDelete)
 				{
 					if (File.Exists(FileToDelete))
 					{
@@ -1415,9 +1415,9 @@ namespace UnrealBuildTool
 				}
 
 				// Generate a list of all the modules of each AppBinaries entry
-				var ModuleList = new List<string>();
-				var bTargetUsesUObjectModule = false;
-				foreach (var AppBin in AppBinaries)
+				List<string> ModuleList = new List<string>();
+				bool bTargetUsesUObjectModule = false;
+				foreach (UEBuildBinary AppBin in AppBinaries)
 				{
 					UEBuildBinaryCPP AppBinCPP = AppBin as UEBuildBinaryCPP;
 					if (AppBinCPP != null)
@@ -1443,11 +1443,11 @@ namespace UnrealBuildTool
 					}
 				}
 
-				var BaseEngineBuildDataFolder = Path.GetFullPath(BuildConfiguration.BaseIntermediatePath).Replace("\\", "/");
-				var PlatformEngineBuildDataFolder = BuildConfiguration.BaseIntermediatePath;
+				string BaseEngineBuildDataFolder = Path.GetFullPath(BuildConfiguration.BaseIntermediatePath).Replace("\\", "/");
+				string PlatformEngineBuildDataFolder = BuildConfiguration.BaseIntermediatePath;
 
 				// Delete generated header files
-				foreach (var ModuleName in ModuleList)
+				foreach (string ModuleName in ModuleList)
 				{
 					UEBuildModuleCPP Module = GetModuleByName(ModuleName) as UEBuildModuleCPP;
 					if (Module != null && Module.GeneratedCodeDirectory != null && Module.GeneratedCodeDirectory.Exists())
@@ -1463,7 +1463,7 @@ namespace UnrealBuildTool
 
 				//
 				{
-					var AppEnginePath = Path.Combine(PlatformEngineBuildDataFolder, TargetName, Configuration.ToString());
+					string AppEnginePath = Path.Combine(PlatformEngineBuildDataFolder, TargetName, Configuration.ToString());
 					if (Directory.Exists(AppEnginePath))
 					{
 						CleanDirectory(AppEnginePath);
@@ -1499,13 +1499,13 @@ namespace UnrealBuildTool
 
 				// Delete the dependency caches
 				{
-					var FlatCPPIncludeDependencyCacheFilename = FlatCPPIncludeDependencyCache.GetDependencyCachePathForTarget(this);
+					FileReference FlatCPPIncludeDependencyCacheFilename = FlatCPPIncludeDependencyCache.GetDependencyCachePathForTarget(this);
 					if (FlatCPPIncludeDependencyCacheFilename.Exists())
 					{
 						Log.TraceVerbose("\tDeleting " + FlatCPPIncludeDependencyCacheFilename);
 						CleanFile(FlatCPPIncludeDependencyCacheFilename.FullName);
 					}
-					var DependencyCacheFilename = DependencyCache.GetDependencyCachePathForTarget(this);
+					FileReference DependencyCacheFilename = DependencyCache.GetDependencyCachePathForTarget(this);
 					if (DependencyCacheFilename.Exists())
 					{
 						Log.TraceVerbose("\tDeleting " + DependencyCacheFilename);
@@ -1517,7 +1517,7 @@ namespace UnrealBuildTool
 				{
 					// @todo ubtmake: Does not yet support cleaning Makefiles that were generated for multiple targets (that code path is currently not used though.)
 
-					var TargetDescs = new List<TargetDescriptor>();
+					List<TargetDescriptor> TargetDescs = new List<TargetDescriptor>();
 					TargetDescs.Add(new TargetDescriptor
 						{
 							TargetName = GetTargetName(),
@@ -1527,7 +1527,7 @@ namespace UnrealBuildTool
 
 					// Normal makefile
 					{
-						var UBTMakefilePath = UnrealBuildTool.GetUBTMakefilePath(TargetDescs);
+						FileReference UBTMakefilePath = UnrealBuildTool.GetUBTMakefilePath(TargetDescs);
 						if (UBTMakefilePath.Exists())
 						{
 							Log.TraceVerbose("\tDeleting " + UBTMakefilePath);
@@ -1538,7 +1538,7 @@ namespace UnrealBuildTool
 					// Hot reload makefile
 					{
 						UEBuildConfiguration.bHotReloadFromIDE = true;
-						var UBTMakefilePath = UnrealBuildTool.GetUBTMakefilePath(TargetDescs);
+						FileReference UBTMakefilePath = UnrealBuildTool.GetUBTMakefilePath(TargetDescs);
 						if (UBTMakefilePath.Exists())
 						{
 							Log.TraceVerbose("\tDeleting " + UBTMakefilePath);
@@ -1550,7 +1550,7 @@ namespace UnrealBuildTool
 
 				// Delete the action history
 				{
-					var ActionHistoryFilename = ActionHistory.GeneratePathForTarget(this);
+					FileReference ActionHistoryFilename = ActionHistory.GeneratePathForTarget(this);
 					if (ActionHistoryFilename.Exists())
 					{
 						Log.TraceVerbose("\tDeleting " + ActionHistoryFilename);
@@ -1610,8 +1610,8 @@ namespace UnrealBuildTool
 
 			{
 				// Collect all files to build
-				var AdditionalFileExtensions = new string[] { ".lib", ".exp", ".dll.response" };
-				var AllFilesToDelete = new List<string>();
+				string[] AdditionalFileExtensions = new string[] { ".lib", ".exp", ".dll.response" };
+				List<string> AllFilesToDelete = new List<string>();
 				foreach (BuildProduct BuildProduct in BuiltReceiptWithFullPaths.BuildProducts)
 				{
 					// don't add any precompiled binaries to the manifest. We don't want to delete them.
@@ -1635,13 +1635,13 @@ namespace UnrealBuildTool
 					if (!ReceiptWithFullPaths.BuildProducts.Contains(BuildProduct, new UEBuildTarget.BuldProductComparer()))
 					{
 						AllFilesToDelete.Add(BuildProduct.Path);
-						var FileExt = Path.GetExtension(BuildProduct.Path);
+						string FileExt = Path.GetExtension(BuildProduct.Path);
 						if (FileExt == ".dll" || FileExt == ".exe")
 						{
-							var ManifestFileWithoutExtension = Utils.GetPathWithoutExtension(BuildProduct.Path);
-							foreach (var AdditionalExt in AdditionalFileExtensions)
+							string ManifestFileWithoutExtension = Utils.GetPathWithoutExtension(BuildProduct.Path);
+							foreach (string AdditionalExt in AdditionalFileExtensions)
 							{
-								var AdditionalFileToDelete = ManifestFileWithoutExtension + AdditionalExt;
+								string AdditionalFileToDelete = ManifestFileWithoutExtension + AdditionalExt;
 								AllFilesToDelete.Add(AdditionalFileToDelete);
 							}
 						}
@@ -1650,7 +1650,7 @@ namespace UnrealBuildTool
 
 				//@todo. This does not clean up files that are no longer built by the target...				
 				// Delete all output files listed in the manifest as well as any additional files.
-				foreach (var FileToDelete in AllFilesToDelete)
+				foreach (string FileToDelete in AllFilesToDelete)
 				{
 					if (File.Exists(FileToDelete))
 					{
@@ -2041,130 +2041,6 @@ namespace UnrealBuildTool
 			}
 		}
 
-		/// <summary>
-		/// Holds a cache of dependency tests between nodes.
-		/// </summary>
-		/// <typeparam name="NodeType">The type of node to be tested for dependencies.</typeparam>
-		class NodeDependencyCache<NodeType>
-		{
-			public NodeDependencyCache(Func<NodeType, IEnumerable<NodeType>> InNodeDependencies)
-			{
-				NodeDependencies = InNodeDependencies;
-				PreviousResults  = new Dictionary<NodeType, Dictionary<NodeType, bool>>();
-			}
-
-			/// <summary>
-			/// Tests if NodeA is dependent on NodeB.
-			/// </summary>
-			/// <param name="NodeA">The dependent to test.</param>
-			/// <param name="NodeB">The dependency to test.</param>
-			/// <returns>true if NodeA is dependent on NodeB, false otherwise.</returns>
-			public bool DependsOn(NodeType NodeA, NodeType NodeB)
-			{
-				// Return any previous result, if there is one.
-				bool Result;
-				Dictionary<NodeType, bool> InnerDictionary;
-				if (PreviousResults.TryGetValue(NodeA, out InnerDictionary))
-				{
-					if (InnerDictionary.TryGetValue(NodeB, out Result))
-					{
-						return Result;
-					}
-				}
-				else
-				{
-					InnerDictionary = new Dictionary<NodeType, bool>();
-					PreviousResults.Add(NodeA, InnerDictionary);
-				}
-
-				Result = DependsOnRecursive(NodeA, NodeB, new HashSet<NodeType>(), InnerDictionary);
-				InnerDictionary.Add(NodeB, Result);
-				return Result;
-			}
-
-			private bool DependsOnRecursive(NodeType NodeA, NodeType NodeB, HashSet<NodeType> Visited, Dictionary<NodeType, bool> InnerDictionary)
-			{
-				if (Visited.Contains(NodeA))
-				{
-					// Nodes are not dependent on themselves
-					return false;
-				}
-
-				Visited.Add(NodeA);
-
-				IEnumerable<NodeType> Deps = NodeDependencies(NodeA);
-
-				if (Deps.Contains(NodeB))
-				{
-					// NodeB is an immediate dependent of NodeA
-					return true;
-				}
-
-				foreach (var i in Deps)
-				{
-					bool Result;
-					if (InnerDictionary.TryGetValue(NodeB, out Result) && Result)
-					{
-						// We've calculated this result before
-						return true;
-					}
-
-					if (DependsOnRecursive(i, NodeB, Visited, InnerDictionary))
-					{
-						return true;
-					}
-				}
-
-				// Didn't find any dependency
-				return false;
-			}
-
-			private Func<NodeType, IEnumerable<NodeType>>            NodeDependencies;
-			private Dictionary<NodeType, Dictionary<NodeType, bool>> PreviousResults;
-		}
-
-		/// <summary>
-		/// Returns a copy of Nodes sorted by dependency.  Independent or circularly-dependent nodes should
-		/// remain in their same relative order within the original Nodes sequence.
-		/// </summary>
-		/// <typeparam name="NodeType">The type of node to sort.</typeparam>
-		/// <param name="Nodes">The sequence of nodes to sort.</param>
-		/// <param name="NodeDependencies">A function which returns a node's immediate dependencies.</param>
-		/// <returns>A copy of Nodes sorted by NodeDependencies, with circular references remaining in the same relative order.</returns>
-		public static IEnumerable<NodeType> StableTopologicalSort<NodeType>(IEnumerable<NodeType> Nodes, Func<NodeType, IEnumerable<NodeType>> NodeDependencies)
-		{
-			List<NodeType> NodeList  = Nodes.ToList();
-			int            NodeCount = NodeList.Count;
-
-			var Cache = new NodeDependencyCache<NodeType>(NodeDependencies);
-			for (int Index1 = 0; Index1 != NodeCount; ++Index1)
-			{
-				var Node1 = NodeList[Index1];
-
-				for (int Index2 = 0; Index2 != Index1; ++Index2)
-				{
-					var Node2 = NodeList[Index2];
-
-					if (Cache.DependsOn(Node2, Node1) && !Cache.DependsOn(Node1, Node2))
-					{
-						// Rotate element at Index1 into position at Index2
-						for (int Index3 = Index1; Index3 != Index2; )
-						{
-							--Index3;
-							NodeList[Index3 + 1] = NodeList[Index3];
-						}
-						NodeList[Index2] = Node1;
-
-						// Break out of this loop, because this iteration must have covered all existing cases
-						// involving the node formerly at position Index1
-						break;
-					}
-				}
-			}
-
-			return NodeList;
-		}
-
 		/// <summary>
 		/// Gathers dependency modules for given binaries list.
 		/// </summary>
@@ -2172,12 +2048,12 @@ namespace UnrealBuildTool
 		/// <returns>Dependency modules set.</returns>
 		static HashSet<UEBuildModuleCPP> GatherDependencyModules(List<UEBuildBinary> Binaries)
 		{
-			var Output = new HashSet<UEBuildModuleCPP>();
+			HashSet<UEBuildModuleCPP> Output = new HashSet<UEBuildModuleCPP>();
 
-			foreach (var Binary in Binaries)
+			foreach (UEBuildBinary Binary in Binaries)
 			{
-				var DependencyModules = Binary.GetAllDependencyModules(bIncludeDynamicallyLoaded: false, bForceCircular: false);
-				foreach (var Module in DependencyModules.OfType<UEBuildModuleCPP>())
+				List<UEBuildModule> DependencyModules = Binary.GetAllDependencyModules(bIncludeDynamicallyLoaded: false, bForceCircular: false);
+				foreach (UEBuildModuleCPP Module in DependencyModules.OfType<UEBuildModuleCPP>())
 				{
 					if (Module.Binary != null)
 					{
@@ -2186,10 +2062,7 @@ namespace UnrealBuildTool
 				}
 			}
 
-			// Sort modules by type, then by dependency
-			var ModulesSortedByType        = Output.OrderBy(c => c.Type).ToList();
-			var ModulesTopologicallySorted = StableTopologicalSort(ModulesSortedByType, x => x.GetAllDependencyModules().OfType<UEBuildModuleCPP>());
-			return new HashSet<UEBuildModuleCPP>(ModulesTopologicallySorted.AsEnumerable());
+			return Output;
 		}
 
 		/// <summary>
@@ -2215,11 +2088,11 @@ namespace UnrealBuildTool
 			// If we're compiling monolithic, make sure the executable knows about all referenced modules
 			if (ShouldCompileMonolithic())
 			{
-				var ExecutableBinary = AppBinaries[0];
+				UEBuildBinary ExecutableBinary = AppBinaries[0];
 
 				// Add all the modules that the executable depends on. Plugins will be already included in this list.
-				var AllReferencedModules = ExecutableBinary.GetAllDependencyModules(bIncludeDynamicallyLoaded: true, bForceCircular: true);
-				foreach (var CurModule in AllReferencedModules)
+				List<UEBuildModule> AllReferencedModules = ExecutableBinary.GetAllDependencyModules(bIncludeDynamicallyLoaded: true, bForceCircular: true);
+				foreach (UEBuildModule CurModule in AllReferencedModules)
 				{
 					if (CurModule.Binary == null || CurModule.Binary == ExecutableBinary || CurModule.Binary.Config.Type == UEBuildBinaryType.StaticLibrary)
 					{
@@ -2308,64 +2181,20 @@ namespace UnrealBuildTool
 
 			if ((BuildConfiguration.bXGEExport && UEBuildConfiguration.bGenerateManifest) || (!UEBuildConfiguration.bGenerateManifest && !UEBuildConfiguration.bCleanProject && !ProjectFileGenerator.bGenerateProjectFiles))
 			{
-				var UObjectDiscoveryStartTime = DateTime.UtcNow;
-
 				// Reconstruct a full list of binaries. Binaries which aren't compiled are stripped out of AppBinaries, but we still need to scan them for UHT.
 				List<UEBuildBinary> AllAppBinaries = AppBinaries.Union(PrecompiledBinaries).ToList();
 
-				var ModulesToGenerateHeadersFor = GatherDependencyModules(AllAppBinaries);
+				HashSet<UEBuildModuleCPP> ModulesToGenerateHeadersFor = GatherDependencyModules(AllAppBinaries);
 
 				if (OnlyModules.Count > 0)
 				{
-					var CorrectlyOrderedModules = GatherDependencyModules(NonFilteredModules);
+					HashSet<UEBuildModuleCPP> CorrectlyOrderedModules = GatherDependencyModules(NonFilteredModules);
 
 					CorrectlyOrderedModules.RemoveWhere((Module) => !ModulesToGenerateHeadersFor.Contains(Module));
 					ModulesToGenerateHeadersFor = CorrectlyOrderedModules;
 				}
 
-				foreach (var Module in ModulesToGenerateHeadersFor)
-				{
-					var UHTModuleInfo = Module.GetCachedUHTModuleInfo();
-					if (UHTModuleInfo.Info.PublicUObjectClassesHeaders.Count > 0 || UHTModuleInfo.Info.PrivateUObjectHeaders.Count > 0 || UHTModuleInfo.Info.PublicUObjectHeaders.Count > 0)
-					{
-						// If we've got this far and there are no source files then it's likely we're running Rocket and ignoring
-						// engine files, so we don't need a .generated.cpp either
-						UEBuildModuleCPP.AutoGenerateCppInfoClass.BuildInfoClass BuildInfo = null;
-						UHTModuleInfo.Info.GeneratedCPPFilenameBase = Path.Combine(Module.GeneratedCodeDirectory.FullName, UHTModuleInfo.Info.ModuleName) + ".generated";
-						if (Module.SourceFilesToBuild.Count != 0)
-						{
-							BuildInfo = new UEBuildModuleCPP.AutoGenerateCppInfoClass.BuildInfoClass(UHTModuleInfo.Info.GeneratedCPPFilenameBase + "*.cpp");
-						}
-
-						Module.AutoGenerateCppInfo = new UEBuildModuleCPP.AutoGenerateCppInfoClass(BuildInfo);
-
-						// If we're running in "gather" mode only, we'll go ahead and cache PCH information for each module right now, so that we don't
-						// have to do it in the assembling phase.  It's OK for gathering to take a bit longer, even if UObject headers are not out of
-						// date in order to save a lot of time in the assembling runs.
-						UHTModuleInfo.Info.PCH = "";
-						if (UnrealBuildTool.IsGatheringBuild && !UnrealBuildTool.IsAssemblingBuild)
-						{
-							// We need to figure out which PCH header this module is including, so that UHT can inject an include statement for it into any .cpp files it is synthesizing
-							var ModuleCompileEnvironment = Module.CreateModuleCompileEnvironment(GlobalCompileEnvironment);
-							Module.CachePCHUsageForModuleSourceFiles(ModuleCompileEnvironment);
-							if (Module.ProcessedDependencies.UniquePCHHeaderFile != null)
-							{
-								UHTModuleInfo.Info.PCH = Module.ProcessedDependencies.UniquePCHHeaderFile.AbsolutePath;
-							}
-						}
-
-						UObjectModules.Add(UHTModuleInfo.Info);
-						FlatModuleCsData[Module.Name].ModuleSourceFolder = Module.ModuleDirectory;
-						FlatModuleCsData[Module.Name].UHTHeaderNames = UHTModuleInfo.HeaderFilenames.ToList();
-						Log.TraceVerbose("Detected UObject module: " + UHTModuleInfo.Info.ModuleName);
-					}
-				}
-
-				if (BuildConfiguration.bPrintPerformanceInfo)
-				{
-					double UObjectDiscoveryTime = (DateTime.UtcNow - UObjectDiscoveryStartTime).TotalSeconds;
-					Trace.TraceInformation("UObject discovery time: " + UObjectDiscoveryTime + "s");
-				}
+				ExternalExecution.SetupUObjectModules(ModulesToGenerateHeadersFor, this, GlobalCompileEnvironment, UObjectModules, FlatModuleCsData, Rules.GetGeneratedCodeVersion());
 
 				// NOTE: Even in Gather mode, we need to run UHT to make sure the files exist for the static action graph to be setup correctly.  This is because UHT generates .cpp
 				// files that are injected as top level prerequisites.  If UHT only emitted included header files, we wouldn't need to run it during the Gather phase at all.
@@ -2393,14 +2222,14 @@ namespace UnrealBuildTool
 			GlobalLinkEnvironment.bShouldCompileMonolithic = ShouldCompileMonolithic();
 
 			// Build the target's binaries.
-			foreach (var Binary in AppBinaries)
+			foreach (UEBuildBinary Binary in AppBinaries)
 			{
-				OutputItems.AddRange(Binary.Build(TargetToolChain, GlobalCompileEnvironment, GlobalLinkEnvironment));
+				OutputItems.AddRange(Binary.Build(this, TargetToolChain, GlobalCompileEnvironment, GlobalLinkEnvironment));
 			}
 
 			if (BuildConfiguration.bPrintPerformanceInfo)
 			{
-				foreach (var SharedPCH in GlobalCompileEnvironment.SharedPCHHeaderFiles)
+				foreach (SharedPCHHeaderInfo SharedPCH in GlobalCompileEnvironment.SharedPCHHeaderFiles)
 				{
 					Log.TraceInformation("Shared PCH '" + SharedPCH.Module.Name + "': Used " + SharedPCH.NumModulesUsingThisPCH + " times");
 				}
@@ -2420,11 +2249,11 @@ namespace UnrealBuildTool
 			if (TargetType != TargetRules.TargetType.Editor && TargetType != TargetRules.TargetType.Program && Configuration == UnrealTargetConfiguration.Shipping &&
 				BuildConfiguration.bCheckLicenseViolations)
 			{
-				var RedistributionErrorMessageBuilder = new StringBuilder();
-				foreach (var Binary in AppBinaries)
+				StringBuilder RedistributionErrorMessageBuilder = new StringBuilder();
+				foreach (UEBuildBinary Binary in AppBinaries)
 				{
-					var NonRedistModules = Binary.GetAllDependencyModules(true, false).Where((DependencyModule) =>
-							!DependencyModule.IsRedistributable() && DependencyModule.Name != Binary.Target.AppName
+					IEnumerable<UEBuildModule> NonRedistModules = Binary.GetAllDependencyModules(true, false).Where((DependencyModule) =>
+							!IsRedistributable(DependencyModule) && DependencyModule.Name != Binary.Target.AppName
 						);
 
 					if (NonRedistModules.Count() != 0)
@@ -2450,10 +2279,27 @@ namespace UnrealBuildTool
 			return EULAViolationWarning;
 		}
 
+		/// <summary>
+		/// Tells if this module can be redistributed.
+		/// </summary>
+		public static bool IsRedistributable(UEBuildModule Module)
+		{
+			if(Module.Rules != null && Module.Rules.IsRedistributableOverride.HasValue)
+			{
+				return Module.Rules.IsRedistributableOverride.Value;
+			}
+
+			if(Module.RulesFile != null)
+			{
+				return !Module.RulesFile.IsUnderDirectory(UnrealBuildTool.EngineSourceDeveloperDirectory) && !Module.RulesFile.IsUnderDirectory(UnrealBuildTool.EngineSourceEditorDirectory);
+			}
+
+			return true;
+		}
+
 		/// <summary>
 		/// Setup target before build. This method finds dependencies, sets up global environment etc.
 		/// </summary>
-		/// <returns>Special Rocket lib files that are build products.</returns>
 		public void PreBuildSetup(UEToolChain TargetToolChain)
 		{
 			// Set up the global compile and link environment in GlobalCompileEnvironment and GlobalLinkEnvironment.
@@ -2506,7 +2352,7 @@ namespace UnrealBuildTool
 			// Create all the modules referenced by the existing binaries
 			foreach(UEBuildBinary Binary in AppBinaries)
 			{
-				Binary.CreateAllDependentModules();
+				Binary.CreateAllDependentModules(this);
 			}
 
 			// Bind every referenced C++ module to a binary
@@ -2580,9 +2426,9 @@ namespace UnrealBuildTool
 
 			// Filter out binaries that were already built and are just used for linking. We will not build these binaries but we need them for link information
 			{
-				var FilteredBinaries = new List<UEBuildBinary>();
+				List<UEBuildBinary> FilteredBinaries = new List<UEBuildBinary>();
 
-				foreach (var DLLBinary in AppBinaries)
+				foreach (UEBuildBinary DLLBinary in AppBinaries)
 				{
 					if (DLLBinary.Config.bAllowCompilation)
 					{
@@ -2602,24 +2448,24 @@ namespace UnrealBuildTool
 			// If we are only interested in platform specific binaries, filter everything else out now
 			if (UnrealBuildTool.GetOnlyPlatformSpecificFor() != UnrealTargetPlatform.Unknown)
 			{
-				var FilteredBinaries = new List<UEBuildBinary>();
+				List<UEBuildBinary> FilteredBinaries = new List<UEBuildBinary>();
 
-				var OtherThingsWeNeedToBuild = new List<OnlyModule>();
+				List<OnlyModule> OtherThingsWeNeedToBuild = new List<OnlyModule>();
 
-				foreach (var DLLBinary in AppBinaries)
+				foreach (UEBuildBinary DLLBinary in AppBinaries)
 				{
 					if (DLLBinary.Config.bIsCrossTarget)
 					{
 						FilteredBinaries.Add(DLLBinary);
 						bool bIncludeDynamicallyLoaded = false;
-						var AllReferencedModules = DLLBinary.GetAllDependencyModules(bIncludeDynamicallyLoaded, bForceCircular: true);
-						foreach (var Other in AllReferencedModules)
+						List<UEBuildModule> AllReferencedModules = DLLBinary.GetAllDependencyModules(bIncludeDynamicallyLoaded, bForceCircular: true);
+						foreach (UEBuildModule Other in AllReferencedModules)
 						{
 							OtherThingsWeNeedToBuild.Add(new OnlyModule(Other.Name));
 						}
 					}
 				}
-				foreach (var DLLBinary in AppBinaries)
+				foreach (UEBuildBinary DLLBinary in AppBinaries)
 				{
 					if (!FilteredBinaries.Contains(DLLBinary) && DLLBinary.FindOnlyModule(OtherThingsWeNeedToBuild) != null)
 					{
@@ -2637,10 +2483,10 @@ namespace UnrealBuildTool
 			[REDACTED_EMAIL]: Will users be able to rebuild UnrealHeaderTool? NO
 			if (!ProjectFileGenerator.bGenerateProjectFiles && UnrealBuildTool.IsEngineInstalled() && AppName != "UnrealHeaderTool")
 			{
-				var FilteredBinaries = new List<UEBuildBinary>();
+				List<UEBuildBinary> FilteredBinaries = new List<UEBuildBinary>();
 
-				// We only want to build rocket projects...
-				foreach (var DLLBinary in AppBinaries)
+				// We only want to build projects outside of the Engine Directory...
+				foreach (UEBuildBinary DLLBinary in AppBinaries)
 				{
 					if (!DLLBinary.Config.OutputFilePath.IsUnderDirectory(UnrealBuildTool.EngineDirectory))
 					{
@@ -2657,7 +2503,7 @@ namespace UnrealBuildTool
 
 				if (AppBinaries.Count == 0)
 				{
-					throw new BuildException("Rocket: No modules found to build?");
+					throw new BuildException("No modules found to build. All requested binaries were already part of the installed engine data.");
 				}
 			}
 
@@ -2670,14 +2516,14 @@ namespace UnrealBuildTool
 
 				// We only want to build the binaries for this single file
 				List<UEBuildBinary> FilteredBinaries = new List<UEBuildBinary>();
-				foreach (var Dependency in Dependencies)
+				foreach (UEBuildModuleCPP Dependency in Dependencies)
 				{
 					bool bFileExistsInDependency = Dependency.SourceFilesFound.CPPFiles.Exists(x => x.AbsolutePath == SingleFileItem.AbsolutePath);
 					if (bFileExistsInDependency)
 					{
 						FilteredBinaries.Add(Dependency.Binary);
 
-						var EmptySourceFileList = new UEBuildModuleCPP.SourceFilesClass();
+						UEBuildModuleCPP.SourceFilesClass EmptySourceFileList = new UEBuildModuleCPP.SourceFilesClass();
 						Dependency.SourceFilesToBuild.CopyFrom(EmptySourceFileList);
 						Dependency.SourceFilesToBuild.CPPFiles.Add(SingleFileItem);
 					}
@@ -2689,7 +2535,7 @@ namespace UnrealBuildTool
 			{
 				// Check the distribution level of all binaries based on the dependencies they have
 				Dictionary<UEBuildModule, UEBuildModuleDistribution> ModuleDistributionCache = new Dictionary<UEBuildModule, UEBuildModuleDistribution>();
-				foreach (var Binary in AppBinaries)
+				foreach (UEBuildBinary Binary in AppBinaries)
 				{
 					Binary.CheckOutputDistributionLevelAgainstDependencies(ModuleDistributionCache);
 				}
@@ -2840,22 +2686,22 @@ namespace UnrealBuildTool
 
 		private static FileReference AddModuleFilenameSuffix(string ModuleName, FileReference FilePath, string Suffix)
 		{
-			var MatchPos = FilePath.FullName.LastIndexOf(ModuleName, StringComparison.InvariantCultureIgnoreCase);
+			int MatchPos = FilePath.FullName.LastIndexOf(ModuleName, StringComparison.InvariantCultureIgnoreCase);
 			if (MatchPos < 0)
 			{
 				throw new BuildException("Failed to find module name \"{0}\" specified on the command line inside of the output filename \"{1}\" to add appendage.", ModuleName, FilePath);
 			}
-			var Appendage = "-" + Suffix;
+			string Appendage = "-" + Suffix;
 			return new FileReference(FilePath.FullName.Insert(MatchPos + ModuleName.Length, Appendage));
 		}
 
 		private static List<UEBuildBinary> GetFilteredOnlyModules(List<UEBuildBinary> Binaries, List<OnlyModule> OnlyModules)
 		{
-			var Result = new List<UEBuildBinary>();
+			List<UEBuildBinary> Result = new List<UEBuildBinary>();
 
-			foreach (var DLLBinary in Binaries)
+			foreach (UEBuildBinary DLLBinary in Binaries)
 			{
-				var FoundOnlyModule = DLLBinary.FindOnlyModule(OnlyModules);
+				OnlyModule FoundOnlyModule = DLLBinary.FindOnlyModule(OnlyModules);
 				if (FoundOnlyModule != null)
 				{
 					Result.Add(DLLBinary);
@@ -2873,16 +2719,16 @@ namespace UnrealBuildTool
 
 		private static List<UEBuildBinary> GetFilteredGameModules(List<UEBuildBinary> Binaries)
 		{
-			var Result = new List<UEBuildBinary>();
+			List<UEBuildBinary> Result = new List<UEBuildBinary>();
 
-			foreach (var DLLBinary in Binaries)
+			foreach (UEBuildBinary DLLBinary in Binaries)
 			{
-				var GameModules = DLLBinary.FindGameModules();
+				List<UEBuildModule> GameModules = DLLBinary.FindGameModules();
 				if (GameModules != null && GameModules.Count > 0)
 				{
 					Result.Add(DLLBinary);
 
-					var UniqueSuffix = (new Random((int)(DateTime.Now.Ticks % Int32.MaxValue)).Next(10000)).ToString();
+					string UniqueSuffix = (new Random((int)(DateTime.Now.Ticks % Int32.MaxValue)).Next(10000)).ToString();
 
 					DLLBinary.Config.OriginalOutputFilePaths = DLLBinary.Config.OutputFilePaths;
 					DLLBinary.Config.OutputFilePaths = DLLBinary.Config.OutputFilePaths.Select(Path => AddModuleFilenameSuffix(GameModules[0].Name, Path, UniqueSuffix)).ToList();
@@ -2899,14 +2745,14 @@ namespace UnrealBuildTool
 		/// </summary>
 		private void CreateLinkerFixupsCPPFile()
 		{
-			var ExecutableBinary = AppBinaries[0];
+			UEBuildBinary ExecutableBinary = AppBinaries[0];
 
 			List<string> PrivateDependencyModuleNames = new List<string>();
 
 			UEBuildBinaryCPP BinaryCPP = ExecutableBinary as UEBuildBinaryCPP;
 			if (BinaryCPP != null)
 			{
-				foreach (var TargetModule in BinaryCPP.Modules)
+				foreach (UEBuildModule TargetModule in BinaryCPP.Modules)
 				{
 					ModuleRules CheckRules = TargetModule.Rules;
 					if (CheckRules.Type != ModuleRules.ModuleType.External)
@@ -2935,22 +2781,22 @@ namespace UnrealBuildTool
 
 				// Create the source file list (just the one cpp file)
 				List<FileItem> SourceFiles = new List<FileItem>();
-				var LinkerFixupCPPFileItem = FileItem.GetItemByFileReference(LinkerFixupCPPFilename);
+				FileItem LinkerFixupCPPFileItem = FileItem.GetItemByFileReference(LinkerFixupCPPFilename);
 				SourceFiles.Add(LinkerFixupCPPFileItem);
 
 				// Create the CPP module
-				var FakeModuleDirectory = LinkerFixupCPPFilename.Directory;
-				var NewModule = CreateArtificialModule(LinkerFixupsName, FakeModuleDirectory, SourceFiles, PrivateDependencyModuleNames);
+				DirectoryReference FakeModuleDirectory = LinkerFixupCPPFilename.Directory;
+				UEBuildModuleCPP NewModule = CreateArtificialModule(LinkerFixupsName, FakeModuleDirectory, SourceFiles, PrivateDependencyModuleNames);
 
 				// Now bind this new module to the executable binary so it will link the plugin libs correctly
 				NewModule.bSkipDefinitionsForCompileEnvironment = true;
 				NewModule.Rules.PCHUsage = ModuleRules.PCHUsageMode.NoSharedPCHs;
-				NewModule.RecursivelyCreateModules();
+				NewModule.RecursivelyCreateModules(this);
 				BindArtificialModuleToBinary(NewModule, ExecutableBinary);
 
 				// Create the cpp file
 				NewModule.bSkipDefinitionsForCompileEnvironment = false;
-				var LinkerFixupsFileContents = GenerateLinkerFixupsContents(ExecutableBinary, NewModule.CreateModuleCompileEnvironment(GlobalCompileEnvironment), HeaderFilename, LinkerFixupsName, PrivateDependencyModuleNames);
+				List<string> LinkerFixupsFileContents = GenerateLinkerFixupsContents(ExecutableBinary, NewModule.CreateModuleCompileEnvironment(this, GlobalCompileEnvironment), HeaderFilename, LinkerFixupsName, PrivateDependencyModuleNames);
 				NewModule.bSkipDefinitionsForCompileEnvironment = true;
 
 				// Determine if the file changed. Write it if it either doesn't exist or the contents are different.
@@ -2978,12 +2824,12 @@ namespace UnrealBuildTool
 
 		private List<string> GenerateLinkerFixupsContents(UEBuildBinary ExecutableBinary, CPPEnvironment CompileEnvironment, string HeaderFilename, string LinkerFixupsName, List<string> PrivateDependencyModuleNames)
 		{
-			var Result = new List<string>();
+			List<string> Result = new List<string>();
 
 			Result.Add("#include \"" + HeaderFilename + "\"");
 
 			// To reduce the size of the command line for the compiler, we're going to put all definitions inside of the cpp file.
-			foreach (var Definition in CompileEnvironment.Config.Definitions)
+			foreach (string Definition in CompileEnvironment.Config.Definitions)
 			{
 				string MacroName;
 				string MacroValue = String.Empty;
@@ -3007,13 +2853,13 @@ namespace UnrealBuildTool
 			Result.Add("{");
 
 			// Fill out the body of the function with the empty function calls. This is what causes the static libraries to be considered relevant
-			var DependencyModules = ExecutableBinary.GetAllDependencyModules(bIncludeDynamicallyLoaded: false, bForceCircular: false);
+			List<UEBuildModule> DependencyModules = ExecutableBinary.GetAllDependencyModules(bIncludeDynamicallyLoaded: false, bForceCircular: false);
 			foreach (string ModuleName in DependencyModules.OfType<UEBuildModuleCPP>().Where(CPPModule => CPPModule.AutoGenerateCppInfo != null).Select(CPPModule => CPPModule.Name).Distinct())
 			{
 				Result.Add("    extern void EmptyLinkFunctionForGeneratedCode" + ModuleName + "();");
 				Result.Add("    EmptyLinkFunctionForGeneratedCode" + ModuleName + "();");
 			}
-			foreach (var DependencyModuleName in PrivateDependencyModuleNames)
+			foreach (string DependencyModuleName in PrivateDependencyModuleNames)
 			{
 				Result.Add("    extern void EmptyLinkFunctionForStaticInitialization" + DependencyModuleName + "();");
 				Result.Add("    EmptyLinkFunctionForStaticInitialization" + DependencyModuleName + "();");
@@ -3035,7 +2881,7 @@ namespace UnrealBuildTool
 			Module.Binary = Binary;
 
 			// Process dependencies for this new module
-			Module.CachePCHUsageForModuleSourceFiles(Module.CreateModuleCompileEnvironment(GlobalCompileEnvironment));
+			Module.CachePCHUsageForModuleSourceFiles(this, Module.CreateModuleCompileEnvironment(this, GlobalCompileEnvironment));
 
 			// Add module to binary
 			Binary.AddModule(Module);
@@ -3055,9 +2901,8 @@ namespace UnrealBuildTool
 			Rules.PrivateDependencyModuleNames.AddRange(PrivateDependencyModuleNames);
 
 			return new UEBuildModuleCPP(
-				InTarget: this,
 				InName: Name,
-				InType: UEBuildModuleType.GameRuntime,
+				InType: UHTModuleType.GameRuntime,
 				InModuleDirectory: Directory,
 				InGeneratedCodeDirectory: null,
 				InIntelliSenseGatherer: null,
@@ -3077,18 +2922,18 @@ namespace UnrealBuildTool
 		private List<SharedPCHHeaderInfo> FindSharedPCHHeaders()
 		{
 			// List of modules, with all of the dependencies of that module
-			var SharedPCHHeaderFiles = new List<SharedPCHHeaderInfo>();
+			List<SharedPCHHeaderInfo> SharedPCHHeaderFiles = new List<SharedPCHHeaderInfo>();
 
 			// Build up our list of modules with "shared PCH headers".  The list will be in dependency order, with modules
 			// that depend on previous modules appearing later in the list
-			foreach (var Binary in AppBinaries)
+			foreach (UEBuildBinary Binary in AppBinaries)
 			{
-				var CPPBinary = Binary as UEBuildBinaryCPP;
+				UEBuildBinaryCPP CPPBinary = Binary as UEBuildBinaryCPP;
 				if (CPPBinary != null)
 				{
-					foreach (var Module in CPPBinary.Modules)
+					foreach (UEBuildModule Module in CPPBinary.Modules)
 					{
-						var CPPModule = Module as UEBuildModuleCPP;
+						UEBuildModuleCPP CPPModule = Module as UEBuildModuleCPP;
 						if (CPPModule != null)
 						{
 							if (!String.IsNullOrEmpty(CPPModule.Rules.SharedPCHHeaderFile) && CPPModule.Binary.Config.bAllowCompilation)
@@ -3096,24 +2941,24 @@ namespace UnrealBuildTool
 								// @todo SharedPCH: Ideally we could figure the PCH header name automatically, and simply use a boolean in the module
 								//     definition to opt into exposing a shared PCH.  Unfortunately we don't determine which private PCH header "goes with"
 								//     a module until a bit later in the process.  It shouldn't be hard to change that though.
-								var SharedPCHHeaderFilePath = ProjectFileGenerator.RootRelativePath + "/Engine/Source/" + CPPModule.Rules.SharedPCHHeaderFile;
-								var SharedPCHHeaderFileItem = FileItem.GetExistingItemByPath(SharedPCHHeaderFilePath);
+								string SharedPCHHeaderFilePath = ProjectFileGenerator.RootRelativePath + "/Engine/Source/" + CPPModule.Rules.SharedPCHHeaderFile;
+								FileItem SharedPCHHeaderFileItem = FileItem.GetExistingItemByPath(SharedPCHHeaderFilePath);
 								if (SharedPCHHeaderFileItem != null)
 								{
-									var ModuleDependencies = new CaselessDictionary<UEBuildModule.ModuleIndexPair>();
+									List<UEBuildModule> ModuleDependencies = new List<UEBuildModule>();
 									bool bIncludeDynamicallyLoaded = false;
-									CPPModule.GetAllDependencyModules(ModuleDependencies, bIncludeDynamicallyLoaded, bForceCircular: false, bOnlyDirectDependencies: false);
+									CPPModule.GetAllDependencyModules(ModuleDependencies, new HashSet<UEBuildModule>(), bIncludeDynamicallyLoaded, bForceCircular: false, bOnlyDirectDependencies: false);
 
 									// Figure out where to insert the shared PCH into our list, based off the module dependency ordering
 									int InsertAtIndex = SharedPCHHeaderFiles.Count;
-									for (var ExistingModuleIndex = SharedPCHHeaderFiles.Count - 1; ExistingModuleIndex >= 0; --ExistingModuleIndex)
+									for (int ExistingModuleIndex = SharedPCHHeaderFiles.Count - 1; ExistingModuleIndex >= 0; --ExistingModuleIndex)
 									{
-										var ExistingModule = SharedPCHHeaderFiles[ExistingModuleIndex].Module;
-										var ExistingModuleDependencies = SharedPCHHeaderFiles[ExistingModuleIndex].Dependencies;
+										UEBuildModule ExistingModule = SharedPCHHeaderFiles[ExistingModuleIndex].Module;
+										Dictionary<string, UEBuildModule> ExistingModuleDependencies = SharedPCHHeaderFiles[ExistingModuleIndex].Dependencies;
 
 										// If the module to add to the list is dependent on any modules already in our header list, that module
 										// must be inserted after any of those dependencies in the list
-										foreach (var ExistingModuleDependency in ExistingModuleDependencies)
+										foreach (KeyValuePair<string, UEBuildModule> ExistingModuleDependency in ExistingModuleDependencies)
 										{
 											if (ExistingModuleDependency.Value == CPPModule)
 											{
@@ -3134,10 +2979,10 @@ namespace UnrealBuildTool
 										}
 									}
 
-									var NewSharedPCHHeaderInfo = new SharedPCHHeaderInfo();
+									SharedPCHHeaderInfo NewSharedPCHHeaderInfo = new SharedPCHHeaderInfo();
 									NewSharedPCHHeaderInfo.PCHHeaderFile = SharedPCHHeaderFileItem;
 									NewSharedPCHHeaderInfo.Module = CPPModule;
-									NewSharedPCHHeaderInfo.Dependencies = ModuleDependencies.Values.OrderBy(x => x.Index).Select(x => x.Module).ToDictionary(M => M.Name);
+									NewSharedPCHHeaderInfo.Dependencies = ModuleDependencies.ToDictionary(M => M.Name);
 									SharedPCHHeaderFiles.Insert(InsertAtIndex, NewSharedPCHHeaderInfo);
 								}
 								else
@@ -3153,7 +2998,7 @@ namespace UnrealBuildTool
 			if (SharedPCHHeaderFiles.Count > 0)
 			{
 				Log.TraceVerbose("Detected {0} shared PCH headers (listed in dependency order):", SharedPCHHeaderFiles.Count);
-				foreach (var CurSharedPCH in SharedPCHHeaderFiles)
+				foreach (SharedPCHHeaderInfo CurSharedPCH in SharedPCHHeaderFiles)
 				{
 					Log.TraceVerbose("	" + CurSharedPCH.PCHHeaderFile.AbsolutePath + "  (module: " + CurSharedPCH.Module.Name + ")");
 				}
@@ -3204,7 +3049,7 @@ namespace UnrealBuildTool
 		protected virtual void AddExtraModules()
 		{
 			// Add extra modules that will either link into the main binary (monolithic), or be linked into separate DLL files (modular)
-			foreach (var ModuleName in ExtraModuleNames)
+			foreach (string ModuleName in ExtraModuleNames)
 			{
 				UEBuildModule Module = FindOrCreateModuleByName(ModuleName);
 				AddModuleToBinary(Module, false);
@@ -3261,17 +3106,16 @@ namespace UnrealBuildTool
 					List<DirectoryReference> Directories = new List<DirectoryReference>();
 					if (TargetType == TargetRules.TargetType.Editor)
 					{
-						Directories.Add(DirectoryReference.Combine(UnrealBuildTool.EngineSourceDirectory, "Editor"));
+						Directories.Add(UnrealBuildTool.EngineSourceEditorDirectory);
 					}
-					Directories.Add(DirectoryReference.Combine(UnrealBuildTool.EngineSourceDirectory, "Runtime"));
+					Directories.Add(UnrealBuildTool.EngineSourceRuntimeDirectory);
 
 					// Also allow anything in the developer directory in non-shipping configurations (though we blacklist by default unless the PrecompileForTargets
 					// setting indicates that it's actually useful at runtime).
 					bool bAllowDeveloperModules = false;
-					DirectoryReference DeveloperDirectory = DirectoryReference.Combine(UnrealBuildTool.EngineSourceDirectory, "Developer");
 					if(Configuration != UnrealTargetConfiguration.Shipping)
 					{
-						Directories.Add(DeveloperDirectory);
+						Directories.Add(UnrealBuildTool.EngineSourceDeveloperDirectory);
 						bAllowDeveloperModules = true;
 					}
 
@@ -3335,7 +3179,7 @@ namespace UnrealBuildTool
 									bCanPrecompile = false;
 									break;
 								case ModuleRules.PrecompileTargetsType.Default:
-									bCanPrecompile = !ModuleFileName.IsUnderDirectory(DeveloperDirectory) || TargetType == TargetRules.TargetType.Editor;
+									bCanPrecompile = !ModuleFileName.IsUnderDirectory(UnrealBuildTool.EngineSourceDeveloperDirectory) || TargetType == TargetRules.TargetType.Editor;
 									break;
 								case ModuleRules.PrecompileTargetsType.Game:
 									bCanPrecompile = (TargetType == TargetRules.TargetType.Client || TargetType == TargetRules.TargetType.Server || TargetType == TargetRules.TargetType.Game);
@@ -3353,7 +3197,7 @@ namespace UnrealBuildTool
 						if (bCanPrecompile)
 						{
 							UEBuildModule Module = FindOrCreateModuleByName(FilteredModuleName);
-							Module.RecursivelyCreateModules();
+							Module.RecursivelyCreateModules(this);
 							PrecompiledModules.Add(Module);
 						}
 					}
@@ -3564,7 +3408,7 @@ namespace UnrealBuildTool
 		/// </summary>
 		protected virtual void SetupModules()
 		{
-			var BuildPlatform = UEBuildPlatform.GetBuildPlatform(Platform);
+			UEBuildPlatform BuildPlatform = UEBuildPlatform.GetBuildPlatform(Platform);
 			List<string> PlatformExtraModules = new List<string>();
 			PlatformContext.AddExtraModules(TargetInfo, PlatformExtraModules);
 			ExtraModuleNames.AddRange(PlatformExtraModules);
@@ -3718,7 +3562,7 @@ namespace UnrealBuildTool
 		/// </summary>
 		public virtual void SetupGlobalEnvironment(UEToolChain ToolChain)
 		{
-			var BuildPlatform = UEBuildPlatform.GetBuildPlatform(Platform);
+			UEBuildPlatform BuildPlatform = UEBuildPlatform.GetBuildPlatform(Platform);
 
 			ToolChain.SetUpGlobalEnvironment();
 
@@ -3765,7 +3609,7 @@ namespace UnrealBuildTool
 			UEBuildConfiguration.ValidateConfiguration();
 
 			// Add the 'Engine/Source' path as a global include path for all modules
-			var EngineSourceDirectory = Path.GetFullPath(Path.Combine("..", "..", "Engine", "Source"));
+			string EngineSourceDirectory = Path.GetFullPath(Path.Combine("..", "..", "Engine", "Source"));
 			if (!Directory.Exists(EngineSourceDirectory))
 			{
 				throw new BuildException("Couldn't find Engine/Source directory using relative path");
@@ -3975,7 +3819,7 @@ namespace UnrealBuildTool
 
 		void SetUpPlatformEnvironment()
 		{
-			var BuildPlatform = UEBuildPlatform.GetBuildPlatform(Platform);
+			UEBuildPlatform BuildPlatform = UEBuildPlatform.GetBuildPlatform(Platform);
 
 			CPPTargetPlatform MainCompilePlatform = BuildPlatform.DefaultCppPlatform;
 
@@ -4009,16 +3853,6 @@ namespace UnrealBuildTool
 			PlatformContext.SetUpProjectEnvironment();
 		}
 
-		/// <summary>
-		/// Registers a module with this target.
-		/// </summary>
-		public void RegisterModule(UEBuildModule Module)
-		{
-			Debug.Assert(Module.Target == this);
-			Modules.Add(Module.Name, Module);
-			FlatModuleCsData.Add(Module.Name, new FlatModuleCsDataType((Module.RulesFile == null) ? null : Module.RulesFile.FullName));
-		}
-
 		/// <summary>
 		/// Create a rules object for the given module, and set any default values for this target
 		/// </summary>
@@ -4045,7 +3879,7 @@ namespace UnrealBuildTool
 					throw new BuildException("Module rules for '{0}' may not specify PrivateAssemblyReferences unless it is a CPlusPlusCLR module type.", ModuleName);
 				}
 
-				var InvalidDependencies = RulesObject.DynamicallyLoadedModuleNames.Intersect(RulesObject.PublicDependencyModuleNames.Concat(RulesObject.PrivateDependencyModuleNames)).ToList();
+				List<string> InvalidDependencies = RulesObject.DynamicallyLoadedModuleNames.Intersect(RulesObject.PublicDependencyModuleNames.Concat(RulesObject.PrivateDependencyModuleNames)).ToList();
 				if (InvalidDependencies.Count != 0)
 				{
 					throw new BuildException("Module rules for '{0}' should not be dependent on modules which are also dynamically loaded: {1}", ModuleName, String.Join(", ", InvalidDependencies));
@@ -4103,7 +3937,7 @@ namespace UnrealBuildTool
 				DirectoryReference ModuleDirectory = ModuleFileName.Directory;
 
 				// Get the type of module we're creating
-				UEBuildModuleType? ModuleType = null;
+				UHTModuleType? ModuleType = null;
 
 				// Get the plugin for this module
 				PluginInfo Plugin;
@@ -4116,7 +3950,7 @@ namespace UnrealBuildTool
 					PluginModuleDesc = Plugin.Descriptor.Modules.FirstOrDefault(x => x.Name == ModuleName);
 					if (PluginModuleDesc != null && PluginModuleDesc.Type == ModuleHostType.Program)
 					{
-						ModuleType = UEBuildModuleType.Program;
+						ModuleType = UHTModuleType.Program;
 					}
 				}
 
@@ -4124,18 +3958,18 @@ namespace UnrealBuildTool
 				{
 					if (RulesObject.Type == ModuleRules.ModuleType.External)
 					{
-						ModuleType = UEBuildModuleType.EngineThirdParty;
+						ModuleType = UHTModuleType.EngineThirdParty;
 					}
 					else
 					{
 						if (!ModuleType.HasValue && PluginModuleDesc != null)
 						{
-							ModuleType = UEBuildModule.GetEngineModuleTypeFromDescriptor(PluginModuleDesc);
+							ModuleType = ExternalExecution.GetEngineModuleTypeFromDescriptor(PluginModuleDesc);
 						}
 
 						if (!ModuleType.HasValue)
 						{
-							ModuleType = UEBuildModule.GetEngineModuleTypeBasedOnLocation(ModuleName, ModuleFileName);
+							ModuleType = ExternalExecution.GetEngineModuleTypeBasedOnLocation(ModuleFileName);
 						}
 					}
 				}
@@ -4143,13 +3977,13 @@ namespace UnrealBuildTool
 				{
 					if (RulesObject.Type == ModuleRules.ModuleType.External)
 					{
-						ModuleType = UEBuildModuleType.GameThirdParty;
+						ModuleType = UHTModuleType.GameThirdParty;
 					}
 					else
 					{
 						if (!ModuleType.HasValue && PluginModuleDesc != null)
 						{
-							ModuleType = UEBuildModule.GetGameModuleTypeFromDescriptor(PluginModuleDesc);
+							ModuleType = ExternalExecution.GetGameModuleTypeFromDescriptor(PluginModuleDesc);
 						}
 
 						if (!ModuleType.HasValue)
@@ -4159,12 +3993,12 @@ namespace UnrealBuildTool
 								ModuleDescriptor ProjectModule = ProjectDescriptor.Modules.FirstOrDefault(x => x.Name == ModuleName);
 								if (ProjectModule != null)
 								{
-									ModuleType = UEBuildModuleTypeExtensions.GameModuleTypeFromHostType(ProjectModule.Type);
+									ModuleType = UHTModuleTypeExtensions.GameModuleTypeFromHostType(ProjectModule.Type);
 								}
 								else
 								{
 									// No descriptor file or module was not on the list
-									ModuleType = UEBuildModuleType.GameRuntime;
+									ModuleType = UHTModuleType.GameRuntime;
 								}
 							}
 						}
@@ -4176,8 +4010,6 @@ namespace UnrealBuildTool
 					throw new BuildException("Unable to determine module type for {0}", ModuleFileName);
 				}
 
-				bool bIsGameModuleOrProgram = ModuleType.Value.IsGameModule() || ModuleType.Value.IsProgramModule();
-
 				// Get the base directory for paths referenced by the module. If the module's under the UProject source directory use that, otherwise leave it relative to the Engine source directory.
 				if (ProjectFile != null)
 				{
@@ -4214,7 +4046,8 @@ namespace UnrealBuildTool
 				// Don't generate include paths for third party modules; they don't follow our conventions. Core is a special-case... leave it alone
 				if (RulesObject.Type != ModuleRules.ModuleType.External && ModuleName != "Core")
 				{
-					// Add the default include paths to the module rules, if they exist.
+					// Add the default include paths to the module rules, if they exist. Would be nice not to include game plugins here, but it would be a regression to change now.
+					bool bIsGameModuleOrProgram = ModuleFileName.IsUnderDirectory(TargetCsFilename.Directory) || (Plugin != null && Plugin.LoadedFrom == PluginLoadedFrom.GameProject);
 					AddDefaultIncludePathsToModuleRules(ModuleFileName, bIsGameModuleOrProgram, Plugin, RulesObject);
 
 					// Add the path to the generated headers 
@@ -4247,7 +4080,7 @@ namespace UnrealBuildTool
 
 						if (ProjectFileForIDE != null)
 						{
-							foreach (var SourceFile in ProjectFileForIDE.SourceFiles)
+							foreach (ProjectFile.SourceFile SourceFile in ProjectFileForIDE.SourceFiles)
 							{
 								SourceFilePaths.Add(SourceFile.Reference);
 							}
@@ -4277,6 +4110,8 @@ namespace UnrealBuildTool
 
 				// Now, go ahead and create the module builder instance
 				Module = InstantiateModule(RulesObject, ModuleName, ModuleType.Value, ModuleDirectory, GeneratedCodeDirectory, IntelliSenseGatherer, FoundSourceFiles, bBuildFiles, ModuleFileName);
+				Modules.Add(Module.Name, Module);
+				FlatModuleCsData.Add(Module.Name, new FlatModuleCsDataType((Module.RulesFile == null) ? null : Module.RulesFile.FullName));
 			}
 			return Module;
 		}
@@ -4284,7 +4119,7 @@ namespace UnrealBuildTool
 		protected virtual UEBuildModule InstantiateModule(
 			ModuleRules RulesObject,
 			string ModuleName,
-			UEBuildModuleType ModuleType,
+			UHTModuleType ModuleType,
 			DirectoryReference ModuleDirectory,
 			DirectoryReference GeneratedCodeDirectory,
 			IntelliSenseGatherer IntelliSenseGatherer,
@@ -4296,7 +4131,6 @@ namespace UnrealBuildTool
 			{
 				case ModuleRules.ModuleType.CPlusPlus:
 					return new UEBuildModuleCPP(
-							InTarget: this,
 							InName: ModuleName,
 							InType: ModuleType,
 							InModuleDirectory: ModuleDirectory,
@@ -4310,7 +4144,6 @@ namespace UnrealBuildTool
 
 				case ModuleRules.ModuleType.CPlusPlusCLR:
 					return new UEBuildModuleCPPCLR(
-							InTarget: this,
 							InName: ModuleName,
 							InType: ModuleType,
 							InModuleDirectory: ModuleDirectory,
@@ -4324,7 +4157,6 @@ namespace UnrealBuildTool
 
 				case ModuleRules.ModuleType.External:
 					return new UEBuildExternalModule(
-							InTarget: this,
 							InName: ModuleName,
 							InType: ModuleType,
 							InModuleDirectory: ModuleDirectory,
@@ -4450,18 +4282,18 @@ namespace UnrealBuildTool
 		{
 			// Make a list of all platform name strings that we're *not* currently compiling, to speed
 			// up file path comparisons later on
-			var SupportedPlatforms = new List<UnrealTargetPlatform>();
+			List<UnrealTargetPlatform> SupportedPlatforms = new List<UnrealTargetPlatform>();
 			SupportedPlatforms.Add(TargetPlatform);
-			var OtherPlatformNameStrings = Utils.MakeListOfUnsupportedPlatforms(SupportedPlatforms);
+			List<string> OtherPlatformNameStrings = Utils.MakeListOfUnsupportedPlatforms(SupportedPlatforms);
 
 
 			// @todo projectfiles: Consider saving out cached list of source files for modules so we don't need to harvest these each time
 
-			var FilteredFileItems = new List<FileItem>();
+			List<FileItem> FilteredFileItems = new List<FileItem>();
 			FilteredFileItems.Capacity = SourceFiles.Count;
 
 			// @todo projectfiles: hard-coded source file set.  Should be made extensible by platform tool chains.
-			var CompilableSourceFileTypes = new string[]
+			string[] CompilableSourceFileTypes = new string[]
 				{
 					".cpp",
 					".c",
@@ -4473,11 +4305,11 @@ namespace UnrealBuildTool
 				};
 
 			// When generating project files, we have no file to extract source from, so we'll locate the code files manually
-			foreach (var SourceFilePath in SourceFiles)
+			foreach (FileReference SourceFilePath in SourceFiles)
 			{
 				// We're only able to compile certain types of files
 				bool IsCompilableSourceFile = false;
-				foreach (var CurExtension in CompilableSourceFileTypes)
+				foreach (string CurExtension in CompilableSourceFileTypes)
 				{
 					if (SourceFilePath.HasExtension(CurExtension))
 					{
@@ -4491,7 +4323,7 @@ namespace UnrealBuildTool
 					if (SourceFilePath.IsUnderDirectory(SourceFilesBaseDirectory))
 					{
 						// Store the path as relative to the project file
-						var RelativeFilePath = SourceFilePath.MakeRelativeTo(SourceFilesBaseDirectory);
+						string RelativeFilePath = SourceFilePath.MakeRelativeTo(SourceFilesBaseDirectory);
 
 						// All compiled files should always be in a sub-directory under the project file directory.  We enforce this here.
 						if (Path.IsPathRooted(RelativeFilePath) || RelativeFilePath.StartsWith(".."))
@@ -4502,7 +4334,7 @@ namespace UnrealBuildTool
 						// Check for source files that don't belong to the platform we're currently compiling.  We'll filter
 						// those source files out
 						bool IncludeThisFile = true;
-						foreach (var CurPlatformName in OtherPlatformNameStrings)
+						foreach (string CurPlatformName in OtherPlatformNameStrings)
 						{
 							if (RelativeFilePath.IndexOf(Path.DirectorySeparatorChar + CurPlatformName + Path.DirectorySeparatorChar, StringComparison.InvariantCultureIgnoreCase) != -1
 								|| RelativeFilePath.StartsWith(CurPlatformName + Path.DirectorySeparatorChar))
diff --git a/Engine/Source/Programs/UnrealBuildTool/IOS/IOSToolChain.cs b/Engine/Source/Programs/UnrealBuildTool/IOS/IOSToolChain.cs
--- a/Engine/Source/Programs/UnrealBuildTool/IOS/IOSToolChain.cs
+++ b/Engine/Source/Programs/UnrealBuildTool/IOS/IOSToolChain.cs
@@ -320,7 +320,7 @@ namespace UnrealBuildTool
 				throw new BuildException("GetRemoteIntermediateFrameworkZipPath: No owning module for framework {0}", Framework.FrameworkName);
 			}
 
-			string IntermediatePath = Framework.OwningModule.Target.ProjectDirectory + "/Intermediate/UnzippedFrameworks/" + Framework.OwningModule.Name;
+			string IntermediatePath = ((ProjectFile != null)? ProjectFile.Directory : UnrealBuildTool.EngineDirectory) + "/Intermediate/UnzippedFrameworks/" + Framework.OwningModule.Name;
 			IntermediatePath = Path.GetFullPath((IntermediatePath + Framework.FrameworkZipPath).Replace(".zip", ""));
 
 			if (BuildHostPlatform.Current.Platform != UnrealTargetPlatform.Mac)
@@ -1066,11 +1066,6 @@ namespace UnrealBuildTool
 							continue;		// Only care if we need to copy bundle assets
 						}
 
-						if (Framework.OwningModule.Target != Target)
-						{
-							continue;		// This framework item doesn't belong to this target, skip it
-						}
-
 						string UnpackedZipPath = GetRemoteIntermediateFrameworkZipPath(Framework);
 
 						// For now, this is hard coded, but we need to loop over all modules, and copy bundled assets that need it
@@ -1276,11 +1271,6 @@ namespace UnrealBuildTool
 							continue;		// Only care if we need to copy bundle assets
 						}
 
-						if (Framework.OwningModule.Target != Target)
-						{
-							continue;		// This framework item doesn't belong to this target, skip it
-						}
-
 						string RemoteZipPath = GetRemoteIntermediateFrameworkZipPath(Framework);
 
 						RemoteZipPath = RemoteZipPath.Replace(".zip", "");
diff --git a/Engine/Source/Programs/UnrealBuildTool/System/ActionGraph.cs b/Engine/Source/Programs/UnrealBuildTool/System/ActionGraph.cs
--- a/Engine/Source/Programs/UnrealBuildTool/System/ActionGraph.cs
+++ b/Engine/Source/Programs/UnrealBuildTool/System/ActionGraph.cs
@@ -295,26 +295,26 @@ namespace UnrealBuildTool
 		/// </summary>
 		public static List<Action> GetActionsToExecute(Action[] PrerequisiteActions, List<UEBuildTarget> Targets, out Dictionary<UEBuildTarget, List<FileItem>> TargetToOutdatedPrerequisitesMap)
 		{
-			var CheckOutdatednessStartTime = DateTime.UtcNow;
+			DateTime CheckOutdatednessStartTime = DateTime.UtcNow;
 
 			// Build a set of all actions needed for this target.
-			var IsActionOutdatedMap = new Dictionary<Action, bool>();
-			foreach (var Action in PrerequisiteActions)
+			Dictionary<Action, bool> IsActionOutdatedMap = new Dictionary<Action, bool>();
+			foreach (Action Action in PrerequisiteActions)
 			{
 				IsActionOutdatedMap.Add(Action, true);
 			}
 
 			// For all targets, build a set of all actions that are outdated.
-			var OutdatedActionDictionary = new Dictionary<Action, bool>();
-			var HistoryList = new List<ActionHistory>();
-			var OpenHistoryFiles = new HashSet<FileReference>();
+			Dictionary<Action, bool> OutdatedActionDictionary = new Dictionary<Action, bool>();
+			List<ActionHistory> HistoryList = new List<ActionHistory>();
+			HashSet<FileReference> OpenHistoryFiles = new HashSet<FileReference>();
 			TargetToOutdatedPrerequisitesMap = new Dictionary<UEBuildTarget, List<FileItem>>();
-			foreach (var BuildTarget in Targets)	// @todo ubtmake: Optimization: Ideally we don't even need to know about targets for ubtmake -- everything would come from the files
+			foreach (UEBuildTarget BuildTarget in Targets)	// @todo ubtmake: Optimization: Ideally we don't even need to know about targets for ubtmake -- everything would come from the files
 			{
-				var HistoryFilename = ActionHistory.GeneratePathForTarget(BuildTarget);
+				FileReference HistoryFilename = ActionHistory.GeneratePathForTarget(BuildTarget);
 				if (!OpenHistoryFiles.Contains(HistoryFilename))		// @todo ubtmake: Optimization: We should be able to move the command-line outdatedness and build product deletion over to the 'gather' phase, as the command-lines won't change between assembler runs
 				{
-					var History = new ActionHistory(HistoryFilename.FullName);
+					ActionHistory History = new ActionHistory(HistoryFilename.FullName);
 					HistoryList.Add(History);
 					OpenHistoryFiles.Add(HistoryFilename);
 					GatherAllOutdatedActions(BuildTarget, History, ref OutdatedActionDictionary, TargetToOutdatedPrerequisitesMap);
@@ -328,7 +328,7 @@ namespace UnrealBuildTool
 				AllActions.RemoveAll(x => x.ActionType != ActionType.Compile);
 
 				// Check all of the leftover compilation actions for the one we want... that one is always outdated.
-				foreach (var Action in AllActions)
+				foreach (Action Action in AllActions)
 				{
 					// Slightly hacky.  We know a compiling X:/Blah/Actor.cpp ends up with StatusDescription == "Actor.cpp".
 					bool bIsSingleFileAction = BuildConfiguration.SingleFileToCompile.EndsWith(Action.StatusDescription.ToLowerInvariant());
@@ -345,7 +345,7 @@ namespace UnrealBuildTool
 			// Save the action history.
 			// This must happen after deleting outdated produced items to ensure that the action history on disk doesn't have
 			// command-lines that don't match the produced items on disk.
-			foreach (var TargetHistory in HistoryList)
+			foreach (ActionHistory TargetHistory in HistoryList)
 			{
 				TargetHistory.Save();
 			}
@@ -376,7 +376,7 @@ namespace UnrealBuildTool
 
 			if (BuildConfiguration.bPrintPerformanceInfo)
 			{
-				var CheckOutdatednessTime = (DateTime.UtcNow - CheckOutdatednessStartTime).TotalSeconds;
+				double CheckOutdatednessTime = (DateTime.UtcNow - CheckOutdatednessStartTime).TotalSeconds;
 				Log.TraceInformation("Checking outdatedness took " + CheckOutdatednessTime + "s");
 			}
 
@@ -557,7 +557,7 @@ namespace UnrealBuildTool
 				}
 			}
 			PlatformSuffix = "";
-			foreach (var CurPlatform in Enum.GetValues(typeof(UnrealTargetPlatform)))
+			foreach (object CurPlatform in Enum.GetValues(typeof(UnrealTargetPlatform)))
 			{
 				string Test = "-" + CurPlatform;
 				if (WorkingString.EndsWith(Test))
@@ -576,7 +576,7 @@ namespace UnrealBuildTool
 		/// </summary>
 		public static void DeleteStaleHotReloadDLLs()
 		{
-			var DeleteStartTime = DateTime.UtcNow;
+			DateTime DeleteStartTime = DateTime.UtcNow;
 
 			foreach (Action BuildAction in AllActions)
 			{
@@ -641,7 +641,7 @@ namespace UnrealBuildTool
 
 			if (BuildConfiguration.bPrintPerformanceInfo)
 			{
-				var DeleteTime = (DateTime.UtcNow - DeleteStartTime).TotalSeconds;
+				double DeleteTime = (DateTime.UtcNow - DeleteStartTime).TotalSeconds;
 				Log.TraceInformation("Deleting stale hot reload DLLs took " + DeleteTime + "s");
 			}
 		}
@@ -937,16 +937,16 @@ namespace UnrealBuildTool
 					// Scan this file for included headers that may be out of date.  Note that it's OK if we break out early because we found
 					// the action to be outdated.  For outdated actions, we kick off a separate include scan in a background thread later on to
 					// catch all of the other includes and form an exhaustive set.
-					var BuildPlatform = UEBuildPlatform.GetBuildPlatform(Target.Platform);
+					UEBuildPlatform BuildPlatform = UEBuildPlatform.GetBuildPlatform(Target.Platform);
 					foreach (FileItem PrerequisiteItem in RootAction.PrerequisiteItems)
 					{
 						// @todo ubtmake: Make sure we are catching RC files here too.  Anything that the toolchain would have tried it on.  Logic should match the CACHING stuff below
 						if (PrerequisiteItem.CachedCPPIncludeInfo != null)
 						{
-							var IncludedFileList = CPPEnvironment.FindAndCacheAllIncludedFiles(Target, PrerequisiteItem, BuildPlatform, PrerequisiteItem.CachedCPPIncludeInfo, bOnlyCachedDependencies: !bPerformExhaustiveIncludeSearchAndUpdateCache);
+							List<FileItem> IncludedFileList = CPPEnvironment.FindAndCacheAllIncludedFiles(Target, PrerequisiteItem, BuildPlatform, PrerequisiteItem.CachedCPPIncludeInfo, bOnlyCachedDependencies: !bPerformExhaustiveIncludeSearchAndUpdateCache);
 							if (IncludedFileList != null)
 							{
-								foreach (var IncludedFile in IncludedFileList)
+								foreach (FileItem IncludedFile in IncludedFileList)
 								{
 									if (IncludedFile.bExists)
 									{
@@ -1044,7 +1044,7 @@ namespace UnrealBuildTool
 				// includes for this file, so that we'll be able to determine whether it is out of date next time very quickly.
 				if (BuildConfiguration.bUseUBTMakefiles)
 				{
-					var DeepIncludeScanStartTime = DateTime.UtcNow;
+					DateTime DeepIncludeScanStartTime = DateTime.UtcNow;
 
 					// @todo ubtmake: we may be scanning more files than we need to here -- indirectly outdated files are bIsOutdated=true by this point (for example basemost includes when deeper includes are dirty)
 					if (bIsOutdated && RootAction.ActionType == ActionType.Compile)	// @todo ubtmake: Does this work with RC files?  See above too.
@@ -1100,16 +1100,16 @@ namespace UnrealBuildTool
 		/// </summary>
 		static void GatherAllOutdatedActions(UEBuildTarget Target, ActionHistory ActionHistory, ref Dictionary<Action, bool> OutdatedActions, Dictionary<UEBuildTarget, List<FileItem>> TargetToOutdatedPrerequisitesMap)
 		{
-			var CheckOutdatednessStartTime = DateTime.UtcNow;
+			DateTime CheckOutdatednessStartTime = DateTime.UtcNow;
 
-			foreach (var Action in AllActions)
+			foreach (Action Action in AllActions)
 			{
 				IsActionOutdated(Target, Action, ref OutdatedActions, ActionHistory, TargetToOutdatedPrerequisitesMap);
 			}
 
 			if (BuildConfiguration.bPrintPerformanceInfo)
 			{
-				var CheckOutdatednessTime = (DateTime.UtcNow - CheckOutdatednessStartTime).TotalSeconds;
+				double CheckOutdatednessTime = (DateTime.UtcNow - CheckOutdatednessStartTime).TotalSeconds;
 				Log.TraceInformation("Checking actions for " + Target.GetTargetName() + " took " + CheckOutdatednessTime + "s");
 			}
 		}
@@ -1259,27 +1259,27 @@ namespace UnrealBuildTool
 		public static void SaveActionGraphVisualization(UEBuildTarget Target, string Filename, string Description, ActionGraphVisualizationType VisualizationType, List<Action> Actions, bool IncludeCompileActions = true)
 		{
 			// True if we should include individual files in the graph network, or false to include only the build actions
-			var IncludeFiles = VisualizationType != ActionGraphVisualizationType.OnlyActions;
-			var OnlyIncludeCPlusPlusFiles = VisualizationType == ActionGraphVisualizationType.OnlyCPlusPlusFilesAndHeaders;
+			bool IncludeFiles = VisualizationType != ActionGraphVisualizationType.OnlyActions;
+			bool OnlyIncludeCPlusPlusFiles = VisualizationType == ActionGraphVisualizationType.OnlyCPlusPlusFilesAndHeaders;
 
 			// True if want to show actions in the graph, otherwise we're only showing files
-			var IncludeActions = VisualizationType != ActionGraphVisualizationType.OnlyFilesAndHeaders && VisualizationType != ActionGraphVisualizationType.OnlyCPlusPlusFilesAndHeaders;
+			bool IncludeActions = VisualizationType != ActionGraphVisualizationType.OnlyFilesAndHeaders && VisualizationType != ActionGraphVisualizationType.OnlyCPlusPlusFilesAndHeaders;
 
 			// True if C++ header dependencies should be expanded into the graph, or false to only have .cpp files
-			var ExpandCPPHeaderDependencies = IncludeFiles && (VisualizationType == ActionGraphVisualizationType.ActionsWithFilesAndHeaders || VisualizationType == ActionGraphVisualizationType.OnlyFilesAndHeaders || VisualizationType == ActionGraphVisualizationType.OnlyCPlusPlusFilesAndHeaders);
+			bool ExpandCPPHeaderDependencies = IncludeFiles && (VisualizationType == ActionGraphVisualizationType.ActionsWithFilesAndHeaders || VisualizationType == ActionGraphVisualizationType.OnlyFilesAndHeaders || VisualizationType == ActionGraphVisualizationType.OnlyCPlusPlusFilesAndHeaders);
 
-			var TimerStartTime = DateTime.UtcNow;
+			DateTime TimerStartTime = DateTime.UtcNow;
 
-			var GraphNodes = new List<GraphNode>();
+			List<GraphNode> GraphNodes = new List<GraphNode>();
 
-			var FileToGraphNodeMap = new Dictionary<FileItem, GraphNode>();
+			Dictionary<FileItem, GraphNode> FileToGraphNodeMap = new Dictionary<FileItem, GraphNode>();
 
 			// Filter our list of actions
-			var FilteredActions = new List<Action>();
+			List<Action> FilteredActions = new List<Action>();
 			{
-				for (var ActionIndex = 0; ActionIndex < Actions.Count; ++ActionIndex)
+				for (int ActionIndex = 0; ActionIndex < Actions.Count; ++ActionIndex)
 				{
-					var Action = Actions[ActionIndex];
+					Action Action = Actions[ActionIndex];
 
 					if (!IncludeActions || IncludeCompileActions || (Action.ActionType != ActionType.Compile))
 					{
@@ -1289,14 +1289,14 @@ namespace UnrealBuildTool
 			}
 
 
-			var FilesToCreateNodesFor = new HashSet<FileItem>();
-			for (var ActionIndex = 0; ActionIndex < FilteredActions.Count; ++ActionIndex)
+			HashSet<FileItem> FilesToCreateNodesFor = new HashSet<FileItem>();
+			for (int ActionIndex = 0; ActionIndex < FilteredActions.Count; ++ActionIndex)
 			{
-				var Action = FilteredActions[ActionIndex];
+				Action Action = FilteredActions[ActionIndex];
 
 				if (IncludeActions)
 				{
-					var GraphNode = new GraphNode()
+					GraphNode GraphNode = new GraphNode()
 					{
 						Id = GraphNodes.Count,
 
@@ -1326,7 +1326,7 @@ namespace UnrealBuildTool
 
 				if (IncludeFiles)
 				{
-					foreach (var ProducedFileItem in Action.ProducedItems)
+					foreach (FileItem ProducedFileItem in Action.ProducedItems)
 					{
 						if (!OnlyIncludeCPlusPlusFiles || IsCPPFile(ProducedFileItem))
 						{
@@ -1334,7 +1334,7 @@ namespace UnrealBuildTool
 						}
 					}
 
-					foreach (var PrerequisiteFileItem in Action.PrerequisiteItems)
+					foreach (FileItem PrerequisiteFileItem in Action.PrerequisiteItems)
 					{
 						if (!OnlyIncludeCPlusPlusFiles || IsCPPFile(PrerequisiteFileItem))
 						{
@@ -1345,14 +1345,14 @@ namespace UnrealBuildTool
 			}
 
 
-			var OverriddenPrerequisites = new Dictionary<FileItem, List<FileItem>>();
+			Dictionary<FileItem, List<FileItem>> OverriddenPrerequisites = new Dictionary<FileItem, List<FileItem>>();
 
 			// Determine the average size of all of the C++ source files
 			Int64 AverageCPPFileSize;
 			{
 				Int64 TotalFileSize = 0;
 				int CPPFileCount = 0;
-				foreach (var FileItem in FilesToCreateNodesFor)
+				foreach (FileItem FileItem in FilesToCreateNodesFor)
 				{
 					if (IsCPPFile(FileItem))
 					{
@@ -1371,11 +1371,11 @@ namespace UnrealBuildTool
 				}
 			}
 
-			var BuildPlatform = UEBuildPlatform.GetBuildPlatform(UnrealTargetPlatform.Win64);
+			UEBuildPlatform BuildPlatform = UEBuildPlatform.GetBuildPlatform(UnrealTargetPlatform.Win64);
 
-			foreach (var FileItem in FilesToCreateNodesFor)
+			foreach (FileItem FileItem in FilesToCreateNodesFor)
 			{
-				var FileGraphNode = new GraphNode()
+				GraphNode FileGraphNode = new GraphNode()
 				{
 					Id = GraphNodes.Count,
 					Label = Path.GetFileName(FileItem.AbsolutePath)
@@ -1401,15 +1401,15 @@ namespace UnrealBuildTool
 				}
 
 				// Set the size of the file node based on the size of the file on disk
-				var bIsCPPFile = IsCPPFile(FileItem);
+				bool bIsCPPFile = IsCPPFile(FileItem);
 				if (bIsCPPFile)
 				{
-					var MinNodeSize = 0.25f;
-					var MaxNodeSize = 2.0f;
-					var FileSize = new FileInfo(FileItem.AbsolutePath).Length;
+					float MinNodeSize = 0.25f;
+					float MaxNodeSize = 2.0f;
+					long FileSize = new FileInfo(FileItem.AbsolutePath).Length;
 					float FileSizeScale = (float)((double)FileSize / (double)AverageCPPFileSize);
 
-					var SourceFileSizeScaleFactor = 0.1f;		// How much to make nodes for files bigger or larger based on their difference from the average file's size
+					float SourceFileSizeScaleFactor = 0.1f;		// How much to make nodes for files bigger or larger based on their difference from the average file's size
 					FileGraphNode.Size = Math.Min(Math.Max(1.0f + SourceFileSizeScaleFactor * FileSizeScale, MinNodeSize), MaxNodeSize);
 				}
 
@@ -1427,7 +1427,7 @@ namespace UnrealBuildTool
 					List<DependencyInclude> DirectlyIncludedFilenames = CPPEnvironment.GetDirectIncludeDependencies(Target, FileItem, BuildPlatform, bOnlyCachedDependencies: false);
 
 					// Resolve the included file name to an actual file.
-					var DirectlyIncludedFiles =
+					List<FileItem> DirectlyIncludedFiles =
 						DirectlyIncludedFilenames
 						.Where(DirectlyIncludedFilename => (DirectlyIncludedFilename.IncludeResolvedNameIfSuccessful != null))
 						.Select(DirectlyIncludedFilename => DirectlyIncludedFilename.IncludeResolvedNameIfSuccessful)
@@ -1442,20 +1442,20 @@ namespace UnrealBuildTool
 
 
 			// Connect everything together
-			var GraphEdges = new List<GraphEdge>();
+			List<GraphEdge> GraphEdges = new List<GraphEdge>();
 
 			if (IncludeActions)
 			{
-				for (var ActionIndex = 0; ActionIndex < FilteredActions.Count; ++ActionIndex)
+				for (int ActionIndex = 0; ActionIndex < FilteredActions.Count; ++ActionIndex)
 				{
-					var Action = FilteredActions[ActionIndex];
-					var ActionGraphNode = GraphNodes[ActionIndex];
+					Action Action = FilteredActions[ActionIndex];
+					GraphNode ActionGraphNode = GraphNodes[ActionIndex];
 
 					List<FileItem> ActualPrerequisiteItems = Action.PrerequisiteItems;
 					if (IncludeFiles && ExpandCPPHeaderDependencies && Action.ActionType == ActionType.Compile)
 					{
 						// The first prerequisite is always the .cpp file to compile
-						var CPPFile = Action.PrerequisiteItems[0];
+						FileItem CPPFile = Action.PrerequisiteItems[0];
 						if (!IsCPPFile(CPPFile))
 						{
 							throw new BuildException("Was expecting a C++ file as the first prerequisite for a Compile action");
@@ -1466,7 +1466,7 @@ namespace UnrealBuildTool
 					}
 
 
-					foreach (var PrerequisiteFileItem in ActualPrerequisiteItems)
+					foreach (FileItem PrerequisiteFileItem in ActualPrerequisiteItems)
 					{
 						if (IncludeFiles)
 						{
@@ -1474,7 +1474,7 @@ namespace UnrealBuildTool
 							if (FileToGraphNodeMap.TryGetValue(PrerequisiteFileItem, out PrerequisiteFileGraphNode))
 							{
 								// Connect a file our action is dependent on, to our action itself
-								var GraphEdge = new GraphEdge()
+								GraphEdge GraphEdge = new GraphEdge()
 								{
 									Id = GraphEdges.Count,
 									Source = PrerequisiteFileGraphNode,
@@ -1492,12 +1492,12 @@ namespace UnrealBuildTool
 						else if (PrerequisiteFileItem.ProducingAction != null)
 						{
 							// Not showing files, so connect the actions together
-							var ProducingActionIndex = FilteredActions.IndexOf(PrerequisiteFileItem.ProducingAction);
+							int ProducingActionIndex = FilteredActions.IndexOf(PrerequisiteFileItem.ProducingAction);
 							if (ProducingActionIndex != -1)
 							{
-								var SourceGraphNode = GraphNodes[ProducingActionIndex];
+								GraphNode SourceGraphNode = GraphNodes[ProducingActionIndex];
 
-								var GraphEdge = new GraphEdge()
+								GraphEdge GraphEdge = new GraphEdge()
 								{
 									Id = GraphEdges.Count,
 									Source = SourceGraphNode,
@@ -1513,15 +1513,15 @@ namespace UnrealBuildTool
 						}
 					}
 
-					foreach (var ProducedFileItem in Action.ProducedItems)
+					foreach (FileItem ProducedFileItem in Action.ProducedItems)
 					{
 						if (IncludeFiles)
 						{
 							if (!OnlyIncludeCPlusPlusFiles || IsCPPFile(ProducedFileItem))
 							{
-								var ProducedFileGraphNode = FileToGraphNodeMap[ProducedFileItem];
+								GraphNode ProducedFileGraphNode = FileToGraphNodeMap[ProducedFileItem];
 
-								var GraphEdge = new GraphEdge()
+								GraphEdge GraphEdge = new GraphEdge()
 								{
 									Id = GraphEdges.Count,
 									Source = ActionGraphNode,
@@ -1538,18 +1538,18 @@ namespace UnrealBuildTool
 			if (IncludeFiles && ExpandCPPHeaderDependencies)
 			{
 				// Fill in overridden prerequisites
-				foreach (var FileAndPrerequisites in OverriddenPrerequisites)
+				foreach (KeyValuePair<FileItem, List<FileItem>> FileAndPrerequisites in OverriddenPrerequisites)
 				{
-					var FileItem = FileAndPrerequisites.Key;
-					var FilePrerequisites = FileAndPrerequisites.Value;
+					FileItem FileItem = FileAndPrerequisites.Key;
+					List<FileItem> FilePrerequisites = FileAndPrerequisites.Value;
 
-					var FileGraphNode = FileToGraphNodeMap[FileItem];
-					foreach (var PrerequisiteFileItem in FilePrerequisites)
+					GraphNode FileGraphNode = FileToGraphNodeMap[FileItem];
+					foreach (FileItem PrerequisiteFileItem in FilePrerequisites)
 					{
 						GraphNode PrerequisiteFileGraphNode;
 						if (FileToGraphNodeMap.TryGetValue(PrerequisiteFileItem, out PrerequisiteFileGraphNode))
 						{
-							var GraphEdge = new GraphEdge()
+							GraphEdge GraphEdge = new GraphEdge()
 							{
 								Id = GraphEdges.Count,
 								Source = PrerequisiteFileGraphNode,
@@ -1571,7 +1571,7 @@ namespace UnrealBuildTool
 
 			if (BuildConfiguration.bPrintPerformanceInfo)
 			{
-				var TimerDuration = DateTime.UtcNow - TimerStartTime;
+				TimeSpan TimerDuration = DateTime.UtcNow - TimerStartTime;
 				Log.TraceInformation("Generating and saving ActionGraph took " + TimerDuration.TotalSeconds + "s");
 			}
 		}
diff --git a/Engine/Source/Programs/UnrealBuildTool/System/ActionHistory.cs b/Engine/Source/Programs/UnrealBuildTool/System/ActionHistory.cs
--- a/Engine/Source/Programs/UnrealBuildTool/System/ActionHistory.cs
+++ b/Engine/Source/Programs/UnrealBuildTool/System/ActionHistory.cs
@@ -32,7 +32,7 @@ namespace UnrealBuildTool
 		{
 			FilePath = Path.GetFullPath(InFilePath);
 
-			var bFoundCache = false;
+			bool bFoundCache = false;
 			if (File.Exists(FilePath) == true)
 			{
 				try
diff --git a/Engine/Source/Programs/UnrealBuildTool/System/CPPHeaders.cs b/Engine/Source/Programs/UnrealBuildTool/System/CPPHeaders.cs
--- a/Engine/Source/Programs/UnrealBuildTool/System/CPPHeaders.cs
+++ b/Engine/Source/Programs/UnrealBuildTool/System/CPPHeaders.cs
@@ -59,7 +59,7 @@ namespace UnrealBuildTool
 		public List<string> GetIncludesPathsToSearch(FileItem SourceFile)
 		{
 			// Build a single list of include paths to search.
-			var IncludePathsToSearch = new List<string>();
+			List<string> IncludePathsToSearch = new List<string>();
 
 			if (SourceFile != null)
 			{
@@ -250,7 +250,7 @@ namespace UnrealBuildTool
 
 				Result = new List<FileItem>();
 
-				var IncludedFileList = new IncludedFilesSet();
+				IncludedFilesSet IncludedFileList = new IncludedFilesSet();
 				CPPEnvironment.FindAndCacheAllIncludedFiles(Target, SourceFile, BuildPlatform, CPPIncludeInfo, ref IncludedFileList, bOnlyCachedDependencies: bOnlyCachedDependencies);
 				foreach (FileItem IncludedFile in IncludedFileList)
 				{
@@ -260,8 +260,8 @@ namespace UnrealBuildTool
 				// Update cache
 				if (bUseFlatCPPIncludeDependencyCache && !bOnlyCachedDependencies)
 				{
-					var Dependencies = new List<FileReference>();
-					foreach (var IncludedFile in Result)
+					List<FileReference> Dependencies = new List<FileReference>();
+					foreach (FileItem IncludedFile in Result)
 					{
 						Dependencies.Add(IncludedFile.Reference);
 					}
@@ -283,10 +283,10 @@ namespace UnrealBuildTool
 		public static bool FindAndCacheAllIncludedFiles(UEBuildTarget Target, FileItem CPPFile, UEBuildPlatform BuildPlatform, CPPIncludeInfo CPPIncludeInfo, ref IncludedFilesSet Result, bool bOnlyCachedDependencies)
 		{
 			IncludedFilesSet IncludedFileList;
-			var IncludedFilesMap = bOnlyCachedDependencies ? OnlyCachedIncludedFilesMap : ExhaustiveIncludedFilesMap;
+			Dictionary<FileItem, IncludedFilesSet> IncludedFilesMap = bOnlyCachedDependencies ? OnlyCachedIncludedFilesMap : ExhaustiveIncludedFilesMap;
 			if (!IncludedFilesMap.TryGetValue(CPPFile, out IncludedFileList))
 			{
-				var TimerStartTime = DateTime.UtcNow;
+				DateTime TimerStartTime = DateTime.UtcNow;
 
 				IncludedFileList = new IncludedFilesSet();
 
@@ -297,7 +297,7 @@ namespace UnrealBuildTool
 				List<DependencyInclude> DirectIncludes = GetDirectIncludeDependencies(Target, CPPFile, BuildPlatform, bOnlyCachedDependencies: bOnlyCachedDependencies);
 
 				// Build a list of the unique set of files that are included by this file.
-				var DirectlyIncludedFiles = new HashSet<FileItem>();
+				HashSet<FileItem> DirectlyIncludedFiles = new HashSet<FileItem>();
 				// require a for loop here because we need to keep track of the index in the list.
 				for (int DirectlyIncludedFileNameIndex = 0; DirectlyIncludedFileNameIndex < DirectIncludes.Count; ++DirectlyIncludedFileNameIndex)
 				{
@@ -332,7 +332,7 @@ namespace UnrealBuildTool
 				TotalDirectIncludeResolves += DirectIncludes.Count;
 
 				// Convert the dictionary of files included by this file into a list.
-				foreach (var DirectlyIncludedFile in DirectlyIncludedFiles)
+				foreach (FileItem DirectlyIncludedFile in DirectlyIncludedFiles)
 				{
 					// Add the file we're directly including
 					IncludedFileList.Add(DirectlyIncludedFile);
@@ -353,7 +353,7 @@ namespace UnrealBuildTool
 
 				// All dependencies have been processed by now so update all circular dependencies
 				// with the full list.
-				foreach (var CircularDependency in IncludedFileList.CircularDependencies)
+				foreach (FileItem CircularDependency in IncludedFileList.CircularDependencies)
 				{
 					IncludedFilesSet CircularDependencyIncludedFiles = IncludedFilesMap[CircularDependency];
 					foreach (FileItem IncludedFile in IncludedFileList)
@@ -367,7 +367,7 @@ namespace UnrealBuildTool
 				// Done collecting files.
 				IncludedFileList.bIsInitialized = true;
 
-				var TimerDuration = DateTime.UtcNow - TimerStartTime;
+				TimeSpan TimerDuration = DateTime.UtcNow - TimerStartTime;
 				TotalTimeSpentGettingIncludes += TimerDuration.TotalSeconds;
 			}
 
@@ -435,13 +435,13 @@ namespace UnrealBuildTool
 				return Info;
 			}
 
-			var Result = new List<DependencyInclude>();
+			List<DependencyInclude> Result = new List<DependencyInclude>();
 			if (bOnlyCachedDependencies)
 			{
 				return Result;
 			}
 
-			var TimerStartTime = DateTime.UtcNow;
+			DateTime TimerStartTime = DateTime.UtcNow;
 			++CPPEnvironment.TotalDirectIncludeCacheMisses;
 
 			// Get the adjusted filename
@@ -509,7 +509,7 @@ namespace UnrealBuildTool
 		/// <param name="EndIndex"></param>
 		private static List<DependencyInclude> CollectHeaders(FileReference ProjectFile, FileItem CPPFile, string FileToRead, string FileContents, string InstalledFolder, int StartIndex, int EndIndex)
 		{
-			var Result = new List<DependencyInclude>();
+			List<DependencyInclude> Result = new List<DependencyInclude>();
 
 			Match M = CPPHeaderRegex.Match(FileContents, StartIndex, EndIndex - StartIndex);
 			CaptureCollection Captures = M.Groups["HeaderFile"].Captures;
diff --git a/Engine/Source/Programs/UnrealBuildTool/System/DependencyCache.cs b/Engine/Source/Programs/UnrealBuildTool/System/DependencyCache.cs
--- a/Engine/Source/Programs/UnrealBuildTool/System/DependencyCache.cs
+++ b/Engine/Source/Programs/UnrealBuildTool/System/DependencyCache.cs
@@ -96,7 +96,7 @@ namespace UnrealBuildTool
 					Log.TraceInformation("Loading existing IncludeFileCache: " + CacheFile.FullName);
 				}
 
-				var TimerStartTime = DateTime.UtcNow;
+				DateTime TimerStartTime = DateTime.UtcNow;
 
 				// Deserialize cache from disk if there is one.
 				DependencyCache Result = Load(CacheFile);
@@ -105,7 +105,7 @@ namespace UnrealBuildTool
 					// Successfully serialize, create the transient variables and return cache.
 					Result.UpdateTimeUtc = DateTime.UtcNow;
 
-					var TimerDuration = DateTime.UtcNow - TimerStartTime;
+					TimeSpan TimerDuration = DateTime.UtcNow - TimerStartTime;
 					if (BuildConfiguration.bPrintPerformanceInfo)
 					{
 						Log.TraceInformation("Loading IncludeFileCache took " + TimerDuration.TotalSeconds + "s");
@@ -239,7 +239,7 @@ namespace UnrealBuildTool
 			// Only save if we've made changes to it since load.
 			if (bIsDirty)
 			{
-				var TimerStartTime = DateTime.UtcNow;
+				DateTime TimerStartTime = DateTime.UtcNow;
 
 				// Save update date as new creation date.
 				CreationTimeUtc = UpdateTimeUtc;
@@ -261,7 +261,7 @@ namespace UnrealBuildTool
 
 				if (BuildConfiguration.bPrintPerformanceInfo)
 				{
-					var TimerDuration = DateTime.UtcNow - TimerStartTime;
+					TimeSpan TimerDuration = DateTime.UtcNow - TimerStartTime;
 					Log.TraceInformation("Saving IncludeFileCache took " + TimerDuration.TotalSeconds + "s");
 				}
 			}
@@ -310,7 +310,7 @@ namespace UnrealBuildTool
 			}
 
 			// Check if any of the resolved includes is missing
-			foreach (var Include in Includes)
+			foreach (DependencyInclude Include in Includes)
 			{
 				if (Include.IncludeResolvedNameIfSuccessful != null)
 				{
@@ -363,9 +363,9 @@ namespace UnrealBuildTool
 		/// </summary>
 		public void ResetUnresolvedDependencies()
 		{
-			foreach (var Dependency in DependencyMap)
+			foreach (KeyValuePair<FileReference, List<DependencyInclude>> Dependency in DependencyMap)
 			{
-				foreach (var Include in Dependency.Value)
+				foreach (DependencyInclude Include in Dependency.Value)
 				{
 					if (Include.HasAttemptedResolve && Include.IncludeResolvedNameIfSuccessful == null)
 					{
@@ -387,8 +387,8 @@ namespace UnrealBuildTool
 		{
 			if (BuildConfiguration.bUseIncludeDependencyResolveCache)
 			{
-				var Includes = DependencyMap[File.Reference];
-				var IncludeToResolve = Includes[DirectlyIncludedFileNameIndex];
+				List<DependencyInclude> Includes = DependencyMap[File.Reference];
+				DependencyInclude IncludeToResolve = Includes[DirectlyIncludedFileNameIndex];
 				if (BuildConfiguration.bTestIncludeDependencyResolveCache)
 				{
 					// test whether there are resolve conflicts between modules with different include paths.
diff --git a/Engine/Source/Programs/UnrealBuildTool/System/DynamicCompilation.cs b/Engine/Source/Programs/UnrealBuildTool/System/DynamicCompilation.cs
--- a/Engine/Source/Programs/UnrealBuildTool/System/DynamicCompilation.cs
+++ b/Engine/Source/Programs/UnrealBuildTool/System/DynamicCompilation.cs
@@ -41,7 +41,7 @@ namespace UnrealBuildTool
 				{
 					// Make sure we have a manifest of source files used to compile the output assembly.  If it doesn't exist
 					// for some reason (not an expected case) then we'll need to recompile.
-					var AssemblySourceListFile = new FileInfo(AssemblySourceListFilePath.FullName);
+					FileInfo AssemblySourceListFile = new FileInfo(AssemblySourceListFilePath.FullName);
 					if (!AssemblySourceListFile.Exists)
 					{
 						return true;
@@ -50,13 +50,13 @@ namespace UnrealBuildTool
 					{
 						// Make sure the source files we're compiling are the same as the source files that were compiled
 						// for the assembly that we want to load
-						var ExistingAssemblySourceFileNames = new List<FileReference>();
+						List<FileReference> ExistingAssemblySourceFileNames = new List<FileReference>();
 						{
-							using (var Reader = AssemblySourceListFile.OpenRead())
+							using (FileStream Reader = AssemblySourceListFile.OpenRead())
 							{
-								using (var TextReader = new StreamReader(Reader))
+								using (StreamReader TextReader = new StreamReader(Reader))
 								{
-									for (var ExistingSourceFileName = TextReader.ReadLine(); ExistingSourceFileName != null; ExistingSourceFileName = TextReader.ReadLine())
+									for (string ExistingSourceFileName = TextReader.ReadLine(); ExistingSourceFileName != null; ExistingSourceFileName = TextReader.ReadLine())
 									{
 										FileReference FullExistingSourceFileName = new FileReference(ExistingSourceFileName);
 
@@ -73,7 +73,7 @@ namespace UnrealBuildTool
 						}
 
 						// Test against source file time stamps
-						foreach (var SourceFileName in SourceFileNames)
+						foreach (FileReference SourceFileName in SourceFileNames)
 						{
 							// Was the existing assembly compiled without this source file?  If so, then we definitely need to recompile it!
 							if (!ExistingAssemblySourceFileNames.Contains(SourceFileName))
@@ -81,7 +81,7 @@ namespace UnrealBuildTool
 								return true;
 							}
 
-							var SourceFileInfo = new FileInfo(SourceFileName.FullName);
+							FileInfo SourceFileInfo = new FileInfo(SourceFileName.FullName);
 
 							// Check to see if the source file exists
 							if (!SourceFileInfo.Exists)
@@ -121,10 +121,10 @@ namespace UnrealBuildTool
 		 */
 		private static Assembly CompileAssembly(FileReference OutputAssemblyPath, List<FileReference> SourceFileNames, List<string> ReferencedAssembies, List<string> PreprocessorDefines = null, bool TreatWarningsAsErrors = false)
 		{
-			var TemporaryFiles = new TempFileCollection();
+			TempFileCollection TemporaryFiles = new TempFileCollection();
 
 			// Setup compile parameters
-			var CompileParams = new CompilerParameters();
+			CompilerParameters CompileParams = new CompilerParameters();
 			{
 				// Always compile the assembly to a file on disk, so that we can load a cached version later if we have one
 				CompileParams.GenerateInMemory = false;
@@ -170,7 +170,7 @@ namespace UnrealBuildTool
 					}
 
 					// The assembly will depend on this application
-					var UnrealBuildToolAssembly = Assembly.GetExecutingAssembly();
+					Assembly UnrealBuildToolAssembly = Assembly.GetExecutingAssembly();
 					CompileParams.ReferencedAssemblies.Add(UnrealBuildToolAssembly.Location);
 				}
 
@@ -210,8 +210,8 @@ namespace UnrealBuildTool
 			try
 			{
 				// Enable .NET 4.0 as we want modern language features like 'var'
-				var ProviderOptions = new Dictionary<string, string>() { { "CompilerVersion", "v4.0" } };
-				var Compiler = new CSharpCodeProvider(ProviderOptions);
+				Dictionary<string, string> ProviderOptions = new Dictionary<string, string>() { { "CompilerVersion", "v4.0" } };
+				CSharpCodeProvider Compiler = new CSharpCodeProvider(ProviderOptions);
 				CompileResults = Compiler.CompileAssemblyFromFile(CompileParams, SourceFileNames.Select(x => x.FullName).ToArray());
 			}
 			catch (Exception Ex)
@@ -305,7 +305,7 @@ namespace UnrealBuildTool
 				// since the previous time we compiled the assembly.  In that case, we'll always want to recompile it!
 				{
 					FileInfo AssemblySourcesListFile = new FileInfo(AssemblySourcesListFilePath.FullName);
-					using (var Writer = AssemblySourcesListFile.CreateText())
+					using (StreamWriter Writer = AssemblySourcesListFile.CreateText())
 					{
 						SourceFileNames.ForEach(x => Writer.WriteLine(x));
 					}
