Copying //UE4/Dev-Build to //UE4/Main

diff --git a/Engine/Source/Programs/UnrealBuildTool/System/ExternalExecution.cs b/Engine/Source/Programs/UnrealBuildTool/System/ExternalExecution.cs
--- a/Engine/Source/Programs/UnrealBuildTool/System/ExternalExecution.cs
+++ b/Engine/Source/Programs/UnrealBuildTool/System/ExternalExecution.cs
@@ -67,6 +67,74 @@ namespace UnrealBuildTool
 		}
 	}
 
+	/// <summary>
+	/// Type of module. Mirrored in UHT as EBuildModuleType.
+	/// This should be sorted by the order in which we expect modules to be built.
+	/// </summary>
+	public enum UHTModuleType
+	{
+		Program,
+		EngineRuntime,
+		EngineDeveloper,
+		EngineEditor,
+		EngineThirdParty,
+		GameRuntime,
+		GameDeveloper,
+		GameEditor,
+		GameThirdParty,
+	}
+	public static class UHTModuleTypeExtensions
+	{
+		public static bool IsProgramModule(this UHTModuleType ModuleType)
+		{
+			return ModuleType == UHTModuleType.Program;
+		}
+		public static bool IsEngineModule(this UHTModuleType ModuleType)
+		{
+			return ModuleType == UHTModuleType.EngineRuntime || ModuleType == UHTModuleType.EngineDeveloper || ModuleType == UHTModuleType.EngineEditor || ModuleType == UHTModuleType.EngineThirdParty;
+		}
+		public static bool IsGameModule(this UHTModuleType ModuleType)
+		{
+			return ModuleType == UHTModuleType.GameRuntime || ModuleType == UHTModuleType.GameDeveloper || ModuleType == UHTModuleType.GameEditor || ModuleType == UHTModuleType.GameThirdParty;
+		}
+		public static UHTModuleType? EngineModuleTypeFromHostType(ModuleHostType ModuleType)
+		{
+			switch (ModuleType)
+			{
+				case ModuleHostType.Runtime:
+				case ModuleHostType.RuntimeNoCommandlet:
+				case ModuleHostType.RuntimeAndProgram:
+				case ModuleHostType.ServerOnly:
+					return UHTModuleType.EngineRuntime;
+				case ModuleHostType.Developer:
+					return UHTModuleType.EngineDeveloper;
+				case ModuleHostType.Editor:
+				case ModuleHostType.EditorNoCommandlet:
+					return UHTModuleType.EngineEditor;
+				default:
+					return null;
+			}
+		}
+		public static UHTModuleType? GameModuleTypeFromHostType(ModuleHostType ModuleType)
+		{
+			switch (ModuleType)
+			{
+				case ModuleHostType.Runtime:
+				case ModuleHostType.RuntimeNoCommandlet:
+				case ModuleHostType.RuntimeAndProgram:
+				case ModuleHostType.ServerOnly:
+					return UHTModuleType.GameRuntime;
+				case ModuleHostType.Developer:
+					return UHTModuleType.GameDeveloper;
+				case ModuleHostType.Editor:
+				case ModuleHostType.EditorNoCommandlet:
+					return UHTModuleType.GameEditor;
+				default:
+					return null;
+			}
+		}
+	}
+
 	/// <summary>
 	/// Information about a module that needs to be passed to UnrealHeaderTool for code generation
 	/// </summary>
@@ -230,21 +298,20 @@ namespace UnrealBuildTool
 		/// Generates a UHTModuleInfo for a particular named module under a directory.
 		/// </summary>
 		/// <returns>
-		public static UHTModuleInfo CreateUHTModuleInfo(IEnumerable<string> HeaderFilenames, UEBuildTarget Target, string ModuleName, DirectoryReference ModuleDirectory, UEBuildModuleType ModuleType)
+		public static UHTModuleInfo CreateUHTModuleInfo(IEnumerable<string> HeaderFilenames, string ModuleName, DirectoryReference ModuleDirectory, UHTModuleType ModuleType, EGeneratedCodeVersion GeneratedCodeVersion)
 		{
-			var ClassesFolder = DirectoryReference.Combine(ModuleDirectory, "Classes");
-			var PublicFolder = DirectoryReference.Combine(ModuleDirectory, "Public");
-			var BuildPlatform = UEBuildPlatform.GetBuildPlatform(Target.Platform);
+			DirectoryReference ClassesFolder = DirectoryReference.Combine(ModuleDirectory, "Classes");
+			DirectoryReference PublicFolder = DirectoryReference.Combine(ModuleDirectory, "Public");
 
-			var AllClassesHeaders = new List<FileItem>();
-			var PublicUObjectHeaders = new List<FileItem>();
-			var PrivateUObjectHeaders = new List<FileItem>();
+			List<FileItem> AllClassesHeaders = new List<FileItem>();
+			List<FileItem> PublicUObjectHeaders = new List<FileItem>();
+			List<FileItem> PrivateUObjectHeaders = new List<FileItem>();
 
-			foreach (var Header in HeaderFilenames)
+			foreach (string Header in HeaderFilenames)
 			{
 				// Check to see if we know anything about this file.  If we have up-to-date cached information about whether it has
 				// UObjects or not, we can skip doing a test here.
-				var UObjectHeaderFileItem = FileItem.GetExistingItemByPath(Header);
+				FileItem UObjectHeaderFileItem = FileItem.GetExistingItemByPath(Header);
 
 				if (CPPEnvironment.DoesFileContainUObjects(UObjectHeaderFileItem.AbsolutePath))
 				{
@@ -263,7 +330,7 @@ namespace UnrealBuildTool
 				}
 			}
 
-			var Result = new UHTModuleInfo
+			UHTModuleInfo Result = new UHTModuleInfo
 			{
 				ModuleName = ModuleName,
 				ModuleDirectory = ModuleDirectory.FullName,
@@ -271,7 +338,7 @@ namespace UnrealBuildTool
 				PublicUObjectClassesHeaders = AllClassesHeaders,
 				PublicUObjectHeaders = PublicUObjectHeaders,
 				PrivateUObjectHeaders = PrivateUObjectHeaders,
-				GeneratedCodeVersion = Target.Rules.GetGeneratedCodeVersion()
+				GeneratedCodeVersion = GeneratedCodeVersion
 			};
 
 			return Result;
@@ -281,6 +348,249 @@ namespace UnrealBuildTool
 		{
 		}
 
+		public static UHTModuleType GetEngineModuleTypeFromDescriptor(ModuleDescriptor Module)
+		{
+			switch (Module.Type)
+			{
+				case ModuleHostType.Developer:
+					return UHTModuleType.EngineDeveloper;
+				case ModuleHostType.Editor:
+				case ModuleHostType.EditorNoCommandlet:
+					return UHTModuleType.EngineEditor;
+				case ModuleHostType.Runtime:
+				case ModuleHostType.RuntimeNoCommandlet:
+				case ModuleHostType.RuntimeAndProgram:
+					return UHTModuleType.EngineRuntime;
+				default:
+					throw new BuildException("Unhandled engine module type {0}", Module.Type.ToString());
+			}
+		}
+
+		public static UHTModuleType GetGameModuleTypeFromDescriptor(ModuleDescriptor Module)
+		{
+			switch (Module.Type)
+			{
+				case ModuleHostType.Developer:
+					return UHTModuleType.GameDeveloper;
+				case ModuleHostType.Editor:
+				case ModuleHostType.EditorNoCommandlet:
+					return UHTModuleType.GameEditor;
+				case ModuleHostType.Runtime:
+				case ModuleHostType.RuntimeNoCommandlet:
+				case ModuleHostType.RuntimeAndProgram:
+					return UHTModuleType.GameRuntime;
+				default:
+					throw new BuildException("Unhandled game module type {0}", Module.Type.ToString());
+			}
+		}
+
+		public static UHTModuleType? GetEngineModuleTypeBasedOnLocation(FileReference ModuleFileName)
+		{
+			if (ModuleFileName.IsUnderDirectory(UnrealBuildTool.EngineSourceRuntimeDirectory))
+			{
+				return UHTModuleType.EngineRuntime;
+			}
+
+			if (ModuleFileName.IsUnderDirectory(UnrealBuildTool.EngineSourceDeveloperDirectory))
+			{
+				return UHTModuleType.EngineDeveloper;
+			}
+
+			if (ModuleFileName.IsUnderDirectory(UnrealBuildTool.EngineSourceEditorDirectory))
+			{
+				return UHTModuleType.EngineEditor;
+			}
+	
+			if (ModuleFileName.IsUnderDirectory(UnrealBuildTool.EngineSourceProgramsDirectory))
+			{
+				return UHTModuleType.Program;
+			}
+
+			if (ModuleFileName.IsUnderDirectory(UnrealBuildTool.EngineSourceThirdPartyDirectory))
+			{
+				return UHTModuleType.EngineThirdParty;
+			}
+
+			return null;
+		}
+
+		/// <summary>
+		/// Holds a cache of dependency tests between nodes.
+		/// </summary>
+		/// <typeparam name="NodeType">The type of node to be tested for dependencies.</typeparam>
+		class NodeDependencyCache<NodeType>
+		{
+			public NodeDependencyCache(Func<NodeType, IEnumerable<NodeType>> InNodeDependencies)
+			{
+				NodeDependencies = InNodeDependencies;
+				PreviousResults  = new Dictionary<NodeType, Dictionary<NodeType, bool>>();
+			}
+
+			/// <summary>
+			/// Tests if NodeA is dependent on NodeB.
+			/// </summary>
+			/// <param name="NodeA">The dependent to test.</param>
+			/// <param name="NodeB">The dependency to test.</param>
+			/// <returns>true if NodeA is dependent on NodeB, false otherwise.</returns>
+			public bool DependsOn(NodeType NodeA, NodeType NodeB)
+			{
+				// Return any previous result, if there is one.
+				bool Result;
+				Dictionary<NodeType, bool> InnerDictionary;
+				if (PreviousResults.TryGetValue(NodeA, out InnerDictionary))
+				{
+					if (InnerDictionary.TryGetValue(NodeB, out Result))
+					{
+						return Result;
+					}
+				}
+				else
+				{
+					InnerDictionary = new Dictionary<NodeType, bool>();
+					PreviousResults.Add(NodeA, InnerDictionary);
+				}
+
+				Result = DependsOnRecursive(NodeA, NodeB, new HashSet<NodeType>(), InnerDictionary);
+				InnerDictionary.Add(NodeB, Result);
+				return Result;
+			}
+
+			private bool DependsOnRecursive(NodeType NodeA, NodeType NodeB, HashSet<NodeType> Visited, Dictionary<NodeType, bool> InnerDictionary)
+			{
+				if (Visited.Contains(NodeA))
+				{
+					// Nodes are not dependent on themselves
+					return false;
+				}
+
+				Visited.Add(NodeA);
+
+				IEnumerable<NodeType> Deps = NodeDependencies(NodeA);
+
+				if (Deps.Contains(NodeB))
+				{
+					// NodeB is an immediate dependent of NodeA
+					return true;
+				}
+
+				foreach (NodeType i in Deps)
+				{
+					bool Result;
+					if (InnerDictionary.TryGetValue(NodeB, out Result) && Result)
+					{
+						// We've calculated this result before
+						return true;
+					}
+
+					if (DependsOnRecursive(i, NodeB, Visited, InnerDictionary))
+					{
+						return true;
+					}
+				}
+
+				// Didn't find any dependency
+				return false;
+			}
+
+			private Func<NodeType, IEnumerable<NodeType>>            NodeDependencies;
+			private Dictionary<NodeType, Dictionary<NodeType, bool>> PreviousResults;
+		}
+
+		/// <summary>
+		/// Returns a copy of Nodes sorted by dependency.  Independent or circularly-dependent nodes should
+		/// remain in their same relative order within the original Nodes sequence.
+		/// </summary>
+		/// <typeparam name="NodeType">The type of node to sort.</typeparam>
+		/// <param name="Nodes">The sequence of nodes to sort.</param>
+		/// <param name="NodeDependencies">A function which returns a node's immediate dependencies.</param>
+		/// <returns>A copy of Nodes sorted by NodeDependencies, with circular references remaining in the same relative order.</returns>
+		public static List<NodeType> StableTopologicalSort<NodeType>(IEnumerable<NodeType> Nodes, Func<NodeType, IEnumerable<NodeType>> NodeDependencies)
+		{
+			List<NodeType> NodeList  = Nodes.ToList();
+			int            NodeCount = NodeList.Count;
+
+			NodeDependencyCache<NodeType> Cache = new NodeDependencyCache<NodeType>(NodeDependencies);
+			for (int Index1 = 0; Index1 != NodeCount; ++Index1)
+			{
+				NodeType Node1 = NodeList[Index1];
+
+				for (int Index2 = 0; Index2 != Index1; ++Index2)
+				{
+					NodeType Node2 = NodeList[Index2];
+
+					if (Cache.DependsOn(Node2, Node1) && !Cache.DependsOn(Node1, Node2))
+					{
+						// Rotate element at Index1 into position at Index2
+						for (int Index3 = Index1; Index3 != Index2; )
+						{
+							--Index3;
+							NodeList[Index3 + 1] = NodeList[Index3];
+						}
+						NodeList[Index2] = Node1;
+
+						// Break out of this loop, because this iteration must have covered all existing cases
+						// involving the node formerly at position Index1
+						break;
+					}
+				}
+			}
+
+			return NodeList;
+		}
+
+		public static void SetupUObjectModules(IEnumerable<UEBuildModuleCPP> ModulesToGenerateHeadersFor, UEBuildTarget Target, CPPEnvironment GlobalCompileEnvironment, List<UHTModuleInfo> UObjectModules, Dictionary<string, FlatModuleCsDataType> FlatModuleCsData, EGeneratedCodeVersion GeneratedCodeVersion)
+		{
+			DateTime UObjectDiscoveryStartTime = DateTime.UtcNow;
+
+			// Sort modules by type, then by dependency
+			List<UEBuildModuleCPP> ModulesSortedByType = ModulesToGenerateHeadersFor.OrderBy(c => c.Type).ToList();
+			ModulesSortedByType = StableTopologicalSort(ModulesSortedByType, x => x.GetDirectDependencyModules().OfType<UEBuildModuleCPP>()).ToList();
+
+			foreach (UEBuildModuleCPP Module in ModulesSortedByType)
+			{
+				UEBuildModuleCPP.UHTModuleInfoCacheType UHTModuleInfo = Module.GetCachedUHTModuleInfo(GeneratedCodeVersion);
+				if (UHTModuleInfo.Info.PublicUObjectClassesHeaders.Count > 0 || UHTModuleInfo.Info.PrivateUObjectHeaders.Count > 0 || UHTModuleInfo.Info.PublicUObjectHeaders.Count > 0)
+				{
+					// If we've got this far and there are no source files then it's likely we're installed and ignoring
+					// engine files, so we don't need a .generated.cpp either
+					UEBuildModuleCPP.AutoGenerateCppInfoClass.BuildInfoClass BuildInfo = null;
+					UHTModuleInfo.Info.GeneratedCPPFilenameBase = Path.Combine(Module.GeneratedCodeDirectory.FullName, UHTModuleInfo.Info.ModuleName) + ".generated";
+					if (Module.SourceFilesToBuild.Count != 0)
+					{
+						BuildInfo = new UEBuildModuleCPP.AutoGenerateCppInfoClass.BuildInfoClass(UHTModuleInfo.Info.GeneratedCPPFilenameBase + "*.cpp");
+					}
+
+					Module.AutoGenerateCppInfo = new UEBuildModuleCPP.AutoGenerateCppInfoClass(BuildInfo);
+
+					// If we're running in "gather" mode only, we'll go ahead and cache PCH information for each module right now, so that we don't
+					// have to do it in the assembling phase.  It's OK for gathering to take a bit longer, even if UObject headers are not out of
+					// date in order to save a lot of time in the assembling runs.
+					UHTModuleInfo.Info.PCH = "";
+					if (UnrealBuildTool.IsGatheringBuild && !UnrealBuildTool.IsAssemblingBuild)
+					{
+						// We need to figure out which PCH header this module is including, so that UHT can inject an include statement for it into any .cpp files it is synthesizing
+						CPPEnvironment ModuleCompileEnvironment = Module.CreateModuleCompileEnvironment(Target, GlobalCompileEnvironment);
+						Module.CachePCHUsageForModuleSourceFiles(Target, ModuleCompileEnvironment);
+						if (Module.ProcessedDependencies.UniquePCHHeaderFile != null)
+						{
+							UHTModuleInfo.Info.PCH = Module.ProcessedDependencies.UniquePCHHeaderFile.AbsolutePath;
+						}
+					}
+
+					UObjectModules.Add(UHTModuleInfo.Info);
+					FlatModuleCsData[Module.Name].ModuleSourceFolder = Module.ModuleDirectory;
+					FlatModuleCsData[Module.Name].UHTHeaderNames = UHTModuleInfo.HeaderFilenames.ToList();
+					Log.TraceVerbose("Detected UObject module: " + UHTModuleInfo.Info.ModuleName);
+				}
+			}
+
+			if (BuildConfiguration.bPrintPerformanceInfo)
+			{
+				double UObjectDiscoveryTime = (DateTime.UtcNow - UObjectDiscoveryStartTime).TotalSeconds;
+				Trace.TraceInformation("UObject discovery time: " + UObjectDiscoveryTime + "s");
+			}
+		}
+
 		/// <summary>
 		/// Gets UnrealHeaderTool.exe path. Does not care if UnrealheaderTool was build as a monolithic exe or not.
 		/// </summary>
@@ -301,7 +611,7 @@ namespace UnrealBuildTool
 		/// </returns>
 		static bool GetHeaderToolTimestamp(out DateTime Timestamp)
 		{
-			using (var TimestampTimer = new ScopedTimer("GetHeaderToolTimestamp"))
+			using (ScopedTimer TimestampTimer = new ScopedTimer("GetHeaderToolTimestamp"))
 			{
 				// Try to read the receipt for UHT.
 				string ReceiptPath = TargetReceipt.GetDefaultPath(BuildConfiguration.RelativeEnginePath, "UnrealHeaderTool", BuildHostPlatform.Current.Platform, UnrealTargetConfiguration.Development, null);
@@ -451,7 +761,7 @@ namespace UnrealBuildTool
 			DateTime? CoreGeneratedTimestamp = null;
 			{
 				// Find the CoreUObject module
-				foreach (var Module in UObjectModules)
+				foreach (UHTModuleInfo Module in UObjectModules)
 				{
 					if (Module.ModuleName.Equals("CoreUObject", StringComparison.InvariantCultureIgnoreCase))
 					{
@@ -466,7 +776,7 @@ namespace UnrealBuildTool
 			}
 
 
-			foreach (var Module in UObjectModules)
+			foreach (UHTModuleInfo Module in UObjectModules)
 			{
 				// If the engine is installed, skip skip checking timestamps for modules that are under the engine directory
 				if (UnrealBuildTool.IsEngineInstalled() && new DirectoryReference(Module.ModuleDirectory).IsUnderDirectory(UnrealBuildTool.EngineDirectory))
@@ -475,8 +785,8 @@ namespace UnrealBuildTool
 				}
 
 				// Make sure we have an existing folder for generated code.  If not, then we definitely need to generate code!
-				var GeneratedCodeDirectory = Path.GetDirectoryName(Module.GeneratedCPPFilenameBase);
-				var TestDirectory = (FileSystemInfo)new DirectoryInfo(GeneratedCodeDirectory);
+				string GeneratedCodeDirectory = Path.GetDirectoryName(Module.GeneratedCPPFilenameBase);
+				FileSystemInfo TestDirectory = (FileSystemInfo)new DirectoryInfo(GeneratedCodeDirectory);
 				if (!TestDirectory.Exists)
 				{
 					// Generated code directory is missing entirely!
@@ -489,7 +799,7 @@ namespace UnrealBuildTool
 				// UObject source files were deleted (or no longer contain UObjects), which means we need to run UHT even
 				// if no other source files were outdated
 				string TimestampFile = Path.Combine(GeneratedCodeDirectory, @"Timestamp");
-				var SavedTimestampFileInfo = (FileSystemInfo)new FileInfo(TimestampFile);
+				FileSystemInfo SavedTimestampFileInfo = (FileSystemInfo)new FileInfo(TimestampFile);
 				if (!SavedTimestampFileInfo.Exists)
 				{
 					// Timestamp file was missing (possibly deleted/cleaned), so headers are out of date
@@ -498,7 +808,7 @@ namespace UnrealBuildTool
 				}
 
 				// Make sure the last UHT run completed after UnrealHeaderTool.exe was compiled last, and after the CoreUObject headers were touched last.
-				var SavedTimestamp = SavedTimestampFileInfo.LastWriteTime;
+				DateTime SavedTimestamp = SavedTimestampFileInfo.LastWriteTime;
 				if (HeaderToolTimestamp > SavedTimestamp || CoreGeneratedTimestamp > SavedTimestamp)
 				{
 					// Generated code is older than UnrealHeaderTool.exe or CoreUObject headers.  Out of date!
@@ -507,14 +817,14 @@ namespace UnrealBuildTool
 				}
 
 				// Iterate over our UObjects headers and figure out if any of them have changed
-				var AllUObjectHeaders = new List<FileItem>();
+				List<FileItem> AllUObjectHeaders = new List<FileItem>();
 				AllUObjectHeaders.AddRange(Module.PublicUObjectClassesHeaders);
 				AllUObjectHeaders.AddRange(Module.PublicUObjectHeaders);
 				AllUObjectHeaders.AddRange(Module.PrivateUObjectHeaders);
 
 				// Load up the old timestamp file and check to see if anything has changed
 				{
-					var UObjectFilesFromPreviousRun = File.ReadAllLines(TimestampFile);
+					string[] UObjectFilesFromPreviousRun = File.ReadAllLines(TimestampFile);
 					if (AllUObjectHeaders.Count != UObjectFilesFromPreviousRun.Length)
 					{
 						Log.TraceVerbose("UnrealHeaderTool needs to run because there are a different number of UObject source files in module {0}", Module.ModuleName);
@@ -530,9 +840,9 @@ namespace UnrealBuildTool
 					}
 				}
 
-				foreach (var HeaderFile in AllUObjectHeaders)
+				foreach (FileItem HeaderFile in AllUObjectHeaders)
 				{
-					var HeaderFileTimestamp = HeaderFile.Info.LastWriteTime;
+					DateTime HeaderFileTimestamp = HeaderFile.Info.LastWriteTime;
 
 					// Has the source header changed since we last generated headers successfully?
 					if (HeaderFileTimestamp > SavedTimestamp)
@@ -553,7 +863,7 @@ namespace UnrealBuildTool
 						// files were actually UObject headers, but because we don't know all of the files we processed
 						// in the previous run, we need to assume our generated code is out of date if the directory timestamp
 						// is newer.
-						var HeaderDirectoryTimestamp = new DirectoryInfo(Path.GetDirectoryName(HeaderFile.AbsolutePath)).LastWriteTime;
+						DateTime HeaderDirectoryTimestamp = new DirectoryInfo(Path.GetDirectoryName(HeaderFile.AbsolutePath)).LastWriteTime;
 						if (HeaderDirectoryTimestamp > SavedTimestamp)
 						{
 							Log.TraceVerbose("UnrealHeaderTool needs to run because the directory containing an existing header ({0}) has changed, and headers may have been added to or deleted from module {1}", HeaderFile.AbsolutePath, Module.ModuleName);
@@ -571,10 +881,10 @@ namespace UnrealBuildTool
 		/// </summary>
 		private static void UpdateDirectoryTimestamps(List<UHTModuleInfo> UObjectModules)
 		{
-			foreach (var Module in UObjectModules)
+			foreach (UHTModuleInfo Module in UObjectModules)
 			{
 				string GeneratedCodeDirectory = Path.GetDirectoryName(Module.GeneratedCPPFilenameBase);
-				var GeneratedCodeDirectoryInfo = new DirectoryInfo(GeneratedCodeDirectory);
+				DirectoryInfo GeneratedCodeDirectoryInfo = new DirectoryInfo(GeneratedCodeDirectory);
 
 				try
 				{
@@ -601,7 +911,7 @@ namespace UnrealBuildTool
 						// files with UObject classes were deleted, so that we'll know to run UHT even if the timestamps of all
 						// of the other source files were unchanged
 						{
-							var AllUObjectFiles = new List<string>();
+							List<string> AllUObjectFiles = new List<string>();
 							AllUObjectFiles.AddRange(Module.PublicUObjectClassesHeaders.ConvertAll(Item => Item.AbsolutePath));
 							AllUObjectFiles.AddRange(Module.PublicUObjectHeaders.ConvertAll(Item => Item.AbsolutePath));
 							AllUObjectFiles.AddRange(Module.PrivateUObjectHeaders.ConvertAll(Item => Item.AbsolutePath));
@@ -621,11 +931,11 @@ namespace UnrealBuildTool
 		/// </summary>
 		public static int RunExternalExecutable(string ExePath, string Commandline)
 		{
-			var ExeInfo = new ProcessStartInfo(ExePath, Commandline);
+			ProcessStartInfo ExeInfo = new ProcessStartInfo(ExePath, Commandline);
 			Log.TraceVerbose("RunExternalExecutable {0} {1}", ExePath, Commandline);
 			ExeInfo.UseShellExecute = false;
 			ExeInfo.RedirectStandardOutput = true;
-			using (var GameProcess = Process.Start(ExeInfo))
+			using (Process GameProcess = Process.Start(ExeInfo))
 			{
 				GameProcess.BeginOutputReadLine();
 				GameProcess.OutputDataReceived += PrintProcessOutputAsync;
@@ -664,9 +974,9 @@ namespace UnrealBuildTool
 			using (ProgressWriter Progress = new ProgressWriter("Generating code...", false))
 			{
 				// We never want to try to execute the header tool when we're already trying to build it!
-				var bIsBuildingUHT = Target.GetTargetName().Equals("UnrealHeaderTool", StringComparison.InvariantCultureIgnoreCase);
+				bool bIsBuildingUHT = Target.GetTargetName().Equals("UnrealHeaderTool", StringComparison.InvariantCultureIgnoreCase);
 
-				var RootLocalPath = Path.GetFullPath(ProjectFileGenerator.RootRelativePath);
+				string RootLocalPath = Path.GetFullPath(ProjectFileGenerator.RootRelativePath);
 
 				// check if UHT is out of date
 				DateTime HeaderToolTimestamp = DateTime.MaxValue;
@@ -680,7 +990,7 @@ namespace UnrealBuildTool
 					// want to save this work until we know that UHT actually needs to be run to speed up best-case iteration times.
 					if (UnrealBuildTool.IsGatheringBuild)		// In assembler-only mode, PCH info is loaded from our UBTMakefile!
 					{
-						foreach (var UHTModuleInfo in UObjectModules)
+						foreach (UHTModuleInfo UHTModuleInfo in UObjectModules)
 						{
 							// Only cache the PCH name if we don't already have one.  When running in 'gather only' mode, this will have already been cached
 							if (string.IsNullOrEmpty(UHTModuleInfo.PCH))
@@ -688,9 +998,9 @@ namespace UnrealBuildTool
 								UHTModuleInfo.PCH = "";
 
 								// We need to figure out which PCH header this module is including, so that UHT can inject an include statement for it into any .cpp files it is synthesizing
-								var DependencyModuleCPP = (UEBuildModuleCPP)Target.GetModuleByName(UHTModuleInfo.ModuleName);
-								var ModuleCompileEnvironment = DependencyModuleCPP.CreateModuleCompileEnvironment(GlobalCompileEnvironment);
-								DependencyModuleCPP.CachePCHUsageForModuleSourceFiles(ModuleCompileEnvironment);
+								UEBuildModuleCPP DependencyModuleCPP = (UEBuildModuleCPP)Target.GetModuleByName(UHTModuleInfo.ModuleName);
+								CPPEnvironment ModuleCompileEnvironment = DependencyModuleCPP.CreateModuleCompileEnvironment(Target, GlobalCompileEnvironment);
+								DependencyModuleCPP.CachePCHUsageForModuleSourceFiles(Target, ModuleCompileEnvironment);
 								if (DependencyModuleCPP.ProcessedDependencies.UniquePCHHeaderFile != null)
 								{
 									UHTModuleInfo.PCH = DependencyModuleCPP.ProcessedDependencies.UniquePCHHeaderFile.AbsolutePath;
@@ -715,7 +1025,7 @@ namespace UnrealBuildTool
 						// If it is there and up to date, it will add 0.8 seconds to the build time.
 						Log.TraceInformation("Building UnrealHeaderTool...");
 
-						var UBTArguments = new StringBuilder();
+						StringBuilder UBTArguments = new StringBuilder();
 
 						UBTArguments.Append("UnrealHeaderTool");
 
@@ -755,7 +1065,7 @@ namespace UnrealBuildTool
 
 					Progress.Write(1, 3);
 
-					var ActualTargetName = String.IsNullOrEmpty(Target.GetTargetName()) ? "UE4" : Target.GetTargetName();
+					string ActualTargetName = String.IsNullOrEmpty(Target.GetTargetName()) ? "UE4" : Target.GetTargetName();
 					Log.TraceInformation("Parsing headers for {0}", ActualTargetName);
 
 					string HeaderToolPath = GetHeaderToolPath();
diff --git a/Engine/Source/Programs/UnrealBuildTool/System/FileItem.cs b/Engine/Source/Programs/UnrealBuildTool/System/FileItem.cs
--- a/Engine/Source/Programs/UnrealBuildTool/System/FileItem.cs
+++ b/Engine/Source/Programs/UnrealBuildTool/System/FileItem.cs
@@ -74,8 +74,8 @@ namespace UnrealBuildTool
 					{
 						// Let's make sure the include paths are the same
 						// @todo ubtmake: We have not seen examples of this actually firing off, so we could probably remove the check for matching includes and simply always make this an error case
-						var CachedIncludePathsToSearch = _CachedCPPIncludeInfo.GetIncludesPathsToSearch(this);
-						var NewIncludePathsToSearch = value.GetIncludesPathsToSearch(this);
+						List<string> CachedIncludePathsToSearch = _CachedCPPIncludeInfo.GetIncludesPathsToSearch(this);
+						List<string> NewIncludePathsToSearch = value.GetIncludesPathsToSearch(this);
 
 						bool bIncludesAreDifferent = false;
 						if (CachedIncludePathsToSearch.Count != NewIncludePathsToSearch.Count)
@@ -84,7 +84,7 @@ namespace UnrealBuildTool
 						}
 						else
 						{
-							for (var IncludeIndex = 0; IncludeIndex < CachedIncludePathsToSearch.Count; ++IncludeIndex)
+							for (int IncludeIndex = 0; IncludeIndex < CachedIncludePathsToSearch.Count; ++IncludeIndex)
 							{
 								if (!CachedIncludePathsToSearch[IncludeIndex].Equals(NewIncludePathsToSearch[IncludeIndex], StringComparison.InvariantCultureIgnoreCase))
 								{
@@ -493,7 +493,7 @@ namespace UnrealBuildTool
 		/// </summary>
 		public static void ResetInfos()
 		{
-			foreach (var Item in UniqueSourceFileMap)
+			foreach (KeyValuePair<FileReference, FileItem> Item in UniqueSourceFileMap)
 			{
 				Item.Value.ResetFileInfo();
 			}
diff --git a/Engine/Source/Programs/UnrealBuildTool/System/FlatCPPIncludeDepencencyCache.cs b/Engine/Source/Programs/UnrealBuildTool/System/FlatCPPIncludeDepencencyCache.cs
--- a/Engine/Source/Programs/UnrealBuildTool/System/FlatCPPIncludeDepencencyCache.cs
+++ b/Engine/Source/Programs/UnrealBuildTool/System/FlatCPPIncludeDepencencyCache.cs
@@ -67,13 +67,13 @@ namespace UnrealBuildTool
 					Log.TraceInformation("Loading existing FlatCPPIncludeDependencyCache: " + CacheFile.FullName);
 				}
 
-				var TimerStartTime = DateTime.UtcNow;
+				DateTime TimerStartTime = DateTime.UtcNow;
 
 				// Deserialize cache from disk if there is one.
 				FlatCPPIncludeDependencyCache Result = Load(CacheFile);
 				if (Result != null)
 				{
-					var TimerDuration = DateTime.UtcNow - TimerStartTime;
+					TimeSpan TimerDuration = DateTime.UtcNow - TimerStartTime;
 					if (BuildConfiguration.bPrintPerformanceInfo)
 					{
 						Log.TraceInformation("Loading FlatCPPIncludeDependencyCache took " + TimerDuration.TotalSeconds + "s");
@@ -153,7 +153,7 @@ namespace UnrealBuildTool
 			// Only save if we've made changes to it since load.
 			if (bIsDirty)
 			{
-				var TimerStartTime = DateTime.UtcNow;
+				DateTime TimerStartTime = DateTime.UtcNow;
 
 				// Serialize the cache to disk.
 				try
@@ -172,7 +172,7 @@ namespace UnrealBuildTool
 
 				if (BuildConfiguration.bPrintPerformanceInfo)
 				{
-					var TimerDuration = DateTime.UtcNow - TimerStartTime;
+					TimeSpan TimerDuration = DateTime.UtcNow - TimerStartTime;
 					Log.TraceInformation("Saving FlatCPPIncludeDependencyCache took " + TimerDuration.TotalSeconds + "s");
 				}
 			}
@@ -271,7 +271,7 @@ namespace UnrealBuildTool
 		/// <param name="Dependencies">List of source dependencies</param>
 		public void SetDependenciesForFile(FileReference AbsoluteFilePath, FileReference PCHName, List<FileReference> Dependencies)
 		{
-			var DependencyInfo = new FlatCPPIncludeDependencyInfo();
+			FlatCPPIncludeDependencyInfo DependencyInfo = new FlatCPPIncludeDependencyInfo();
 			DependencyInfo.PCHName = PCHName;	// @todo ubtmake: Not actually used yet.  The idea is to use this to reduce the number of indirect includes we need to store in the cache.
 			DependencyInfo.Includes = Dependencies;
 			DependencyInfo.IncludeFileItems = null;
diff --git a/Engine/Source/Programs/UnrealBuildTool/System/JunkDeleter.cs b/Engine/Source/Programs/UnrealBuildTool/System/JunkDeleter.cs
--- a/Engine/Source/Programs/UnrealBuildTool/System/JunkDeleter.cs
+++ b/Engine/Source/Programs/UnrealBuildTool/System/JunkDeleter.cs
@@ -15,7 +15,7 @@ namespace UnrealBuildTool
 		/// </summary>
 		static public void DeleteJunk()
 		{
-			var JunkStartTime = DateTime.UtcNow;
+			DateTime JunkStartTime = DateTime.UtcNow;
 
 			if (UnrealBuildTool.IsEngineInstalled() == false)
 			{
@@ -25,7 +25,7 @@ namespace UnrealBuildTool
 
 			if (BuildConfiguration.bPrintPerformanceInfo)
 			{
-				var JunkTime = (DateTime.UtcNow - JunkStartTime).TotalSeconds;
+				double JunkTime = (DateTime.UtcNow - JunkStartTime).TotalSeconds;
 				Log.TraceInformation("DeleteJunk took " + JunkTime + "s");
 			}
 		}
diff --git a/Engine/Source/Programs/UnrealBuildTool/System/LocalExecutor.cs b/Engine/Source/Programs/UnrealBuildTool/System/LocalExecutor.cs
--- a/Engine/Source/Programs/UnrealBuildTool/System/LocalExecutor.cs
+++ b/Engine/Source/Programs/UnrealBuildTool/System/LocalExecutor.cs
@@ -107,7 +107,7 @@ namespace UnrealBuildTool
 		/// <param name="e">  Event arguments (In this case, the line of string output)</param>
 		protected void ActionDebugOutput(object sender, DataReceivedEventArgs e)
 		{
-			var Output = e.Data;
+			string Output = e.Data;
 			if (Output == null)
 			{
 				return;
diff --git a/Engine/Source/Programs/UnrealBuildTool/System/ModuleProcessingException.cs b/Engine/Source/Programs/UnrealBuildTool/System/ModuleProcessingException.cs
--- a/Engine/Source/Programs/UnrealBuildTool/System/ModuleProcessingException.cs
+++ b/Engine/Source/Programs/UnrealBuildTool/System/ModuleProcessingException.cs
@@ -18,7 +18,7 @@ namespace UnrealBuildTool
 			// Our build exceptions do not show the callstack so they are more friendly to users.
 			if (!string.IsNullOrEmpty(Message))
 			{
-				var Result = Message;
+				string Result = Message;
 				if (InnerException != null)
 				{
 					Result += '\n' + InnerException.ToString();
diff --git a/Engine/Source/Programs/UnrealBuildTool/System/Project.cs b/Engine/Source/Programs/UnrealBuildTool/System/Project.cs
--- a/Engine/Source/Programs/UnrealBuildTool/System/Project.cs
+++ b/Engine/Source/Programs/UnrealBuildTool/System/Project.cs
@@ -193,7 +193,7 @@ namespace UnrealBuildTool
 		/// <param name="RelativeBaseFolder">The directory the path within the project will be relative to</param>
 		public void AddFilesToProject(List<FileReference> FilesToAdd, DirectoryReference BaseFolder)
 		{
-			foreach (var CurFile in FilesToAdd)
+			foreach (FileReference CurFile in FilesToAdd)
 			{
 				AddFileToProject(CurFile, BaseFolder);
 			}
@@ -269,13 +269,13 @@ namespace UnrealBuildTool
 			if (ProjectFileGenerator.OnlyGenerateIntelliSenseDataForProject == null ||
 				ProjectFileGenerator.OnlyGenerateIntelliSenseDataForProject == this)
 			{
-				foreach (var CurDef in NewPreprocessorDefinitions)
+				foreach (string CurDef in NewPreprocessorDefinitions)
 				{
 					// Don't add definitions and value combinations that have already been added for this project
 					if (KnownIntelliSensePreprocessorDefinitions.Add(CurDef))
 					{
 						// Go ahead and check to see if the definition already exists, but the value is different
-						var AlreadyExists = false;
+						bool AlreadyExists = false;
 
 						string Def, Value;
 						SplitDefinitionAndValue(CurDef, out Def, out Value);
@@ -310,7 +310,7 @@ namespace UnrealBuildTool
 			if (ProjectFileGenerator.OnlyGenerateIntelliSenseDataForProject == null ||
 				ProjectFileGenerator.OnlyGenerateIntelliSenseDataForProject == this)
 			{
-				foreach (var CurPath in NewIncludePaths)
+				foreach (string CurPath in NewIncludePaths)
 				{
 					if (bAddingSystemIncludes ? KnownIntelliSenseSystemIncludeSearchPaths.Add(CurPath) : KnownIntelliSenseIncludeSearchPaths.Add(CurPath))
 					{
@@ -333,9 +333,9 @@ namespace UnrealBuildTool
 						PathRelativeToProjectFile = PathRelativeToProjectFile.TrimEnd('/', '\\');
 
 						// Make sure that it doesn't exist already
-						var AlreadyExists = false;
+						bool AlreadyExists = false;
 						List<string> SearchPaths = bAddingSystemIncludes ? IntelliSenseSystemIncludeSearchPaths : IntelliSenseIncludeSearchPaths;
-						foreach (var ExistingPath in SearchPaths)
+						foreach (string ExistingPath in SearchPaths)
 						{
 							if (PathRelativeToProjectFile == ExistingPath)
 							{
@@ -360,8 +360,8 @@ namespace UnrealBuildTool
 		public void AddDependsOnProject(ProjectFile InProjectFile)
 		{
 			// Make sure that it doesn't exist already
-			var AlreadyExists = false;
-			foreach (var ExistingDependentOn in DependsOnProjects)
+			bool AlreadyExists = false;
+			foreach (ProjectFile ExistingDependentOn in DependsOnProjects)
 			{
 				if (ExistingDependentOn == InProjectFile)
 				{
diff --git a/Engine/Source/Programs/UnrealBuildTool/System/ProjectFileGenerator.cs b/Engine/Source/Programs/UnrealBuildTool/System/ProjectFileGenerator.cs
--- a/Engine/Source/Programs/UnrealBuildTool/System/ProjectFileGenerator.cs
+++ b/Engine/Source/Programs/UnrealBuildTool/System/ProjectFileGenerator.cs
@@ -47,10 +47,10 @@ namespace UnrealBuildTool
 		{
             MasterProjectFolder ResultFolder = null;
 
-            foreach (var FolderName in SubFolderName.Split(new char[1] {'\\'}, StringSplitOptions.RemoveEmptyEntries))
+            foreach (string FolderName in SubFolderName.Split(new char[1] {'\\'}, StringSplitOptions.RemoveEmptyEntries))
             {
                 bool AlreadyExists = false;
-                foreach (var ExistingFolder in SubFolders)
+                foreach (MasterProjectFolder ExistingFolder in SubFolders)
                 {
                     if (ExistingFolder.FolderName.Equals(FolderName, StringComparison.InvariantCultureIgnoreCase))
                     {
@@ -79,7 +79,7 @@ namespace UnrealBuildTool
 		/// <returns>The found folder that the project is in, or null</returns>
 		public MasterProjectFolder FindFolderForProject( ProjectFile Project )
 		{
-			foreach( var CurFolder in SubFolders )
+			foreach( MasterProjectFolder CurFolder in SubFolders )
 			{
 				MasterProjectFolder FoundFolder = CurFolder.FindFolderForProject( Project );
 				if( FoundFolder != null )
@@ -88,7 +88,7 @@ namespace UnrealBuildTool
 				}
 			}
 
-			foreach( var ChildProject in ChildProjects )
+			foreach( ProjectFile ChildProject in ChildProjects )
 			{
 				if( ChildProject == Project )
 				{
@@ -142,7 +142,6 @@ namespace UnrealBuildTool
 		public static bool bCreateDummyConfigsForUnsupportedPlatforms = true;
 
 		/// Whether we should include configurations for "Test" and "Shipping" in generated projects (pass "-NoShippingConfigs" to disable this)
-		/// NOTE: This is currently ignored when generating Rocket projects
 		public static bool bIncludeTestAndShippingConfigs = true;
 
 		/// True if intellisense data should be generated (takes a while longer)
@@ -255,12 +254,12 @@ namespace UnrealBuildTool
 		/// </summary>
 		void AddAutomationModules(MasterProjectFolder ProgramsFolder)
 		{
-			var Folder = ProgramsFolder.AddSubFolder("Automation");
-			var AllGameFolders = UEBuildTarget.DiscoverAllGameFolders();
-			var BuildFolders = new List<DirectoryReference>(AllGameFolders.Count);
-			foreach (var GameFolder in AllGameFolders)
+			MasterProjectFolder Folder = ProgramsFolder.AddSubFolder("Automation");
+			List<DirectoryReference> AllGameFolders = UEBuildTarget.DiscoverAllGameFolders();
+			List<DirectoryReference> BuildFolders = new List<DirectoryReference>(AllGameFolders.Count);
+			foreach (DirectoryReference GameFolder in AllGameFolders)
 			{
-				var GameBuildFolder = DirectoryReference.Combine(GameFolder, "Build");
+				DirectoryReference GameBuildFolder = DirectoryReference.Combine(GameFolder, "Build");
 				if (GameBuildFolder.Exists())
 				{
 					BuildFolders.Add(GameBuildFolder);
@@ -268,12 +267,12 @@ namespace UnrealBuildTool
 			}
 
 			// Find all the automation modules .csproj files to add
-			var ModuleFiles = RulesCompiler.FindAllRulesSourceFiles(RulesCompiler.RulesFileType.AutomationModule, null, ForeignPlugins:null, AdditionalSearchPaths: BuildFolders );
-			foreach (var ProjectFile in ModuleFiles)
+			List<FileReference> ModuleFiles = RulesCompiler.FindAllRulesSourceFiles(RulesCompiler.RulesFileType.AutomationModule, null, ForeignPlugins:null, AdditionalSearchPaths: BuildFolders );
+			foreach (FileReference ProjectFile in ModuleFiles)
 			{
 				if (ProjectFile.Exists())
 				{
-					var Project = new VCSharpProjectFile(ProjectFile);
+					VCSharpProjectFile Project = new VCSharpProjectFile(ProjectFile);
 					Project.ShouldBuildForAllSolutionTargets = true;
 					AddExistingProjectFile(Project, bForceDevelopmentConfiguration: true);
                     AutomationProjectFiles.Add( Project );
@@ -370,7 +369,7 @@ namespace UnrealBuildTool
 
 			// Find all of the module files.  This will filter out any modules or targets that don't belong to platforms
 			// we're generating project files for.
-			var AllModuleFiles = DiscoverModules(AllGameProjects);
+			List<FileReference> AllModuleFiles = DiscoverModules(AllGameProjects);
 
 			ProjectFile EngineProject = null;
 			Dictionary<DirectoryReference, ProjectFile> GameProjects = null;
@@ -453,7 +452,7 @@ namespace UnrealBuildTool
 						}
 					}
 
-					foreach( var CurGameProject in GameProjects.Values )
+					foreach( ProjectFile CurGameProject in GameProjects.Values )
 					{
 						// Templates go under a different solution folder than games
 						if( TemplateGameProjects.Contains( CurGameProject ) )
@@ -466,7 +465,7 @@ namespace UnrealBuildTool
 						}
 					}
 
-					foreach( var CurProgramProject in ProgramProjects.Values )
+					foreach( ProjectFile CurProgramProject in ProgramProjects.Values )
 					{
                         ProjectTarget Target = CurProgramProject.ProjectTargets.FirstOrDefault(t => !String.IsNullOrEmpty(t.TargetRules.SolutionDirectory));
 
@@ -494,7 +493,7 @@ namespace UnrealBuildTool
 					MasterProjectFolder ProgramsFolder = RootFolder.AddSubFolder( "Programs" );
 
 					// Add EnvVarsToXML to the master project
-					var EnvVarsToXMLProjectFile = AddSimpleCSharpProject("EnvVarsToXML/EnvVarsToXML", bShouldBuildForAllSolutionTargets: true, bForceDevelopmentConfiguration: true);
+					VCSharpProjectFile EnvVarsToXMLProjectFile = AddSimpleCSharpProject("EnvVarsToXML/EnvVarsToXML", bShouldBuildForAllSolutionTargets: true, bForceDevelopmentConfiguration: true);
 					ProgramsFolder.ChildProjects.Add(EnvVarsToXMLProjectFile);
 
 					// Add UnrealBuildTool to the master project
@@ -540,12 +539,12 @@ namespace UnrealBuildTool
 				{
 					// Figure out which targets we need about IntelliSense for.  We only need to worry about targets for projects
 					// that we're actually generating in this session.
-					var IntelliSenseTargetFiles = new List<Tuple<ProjectFile, FileReference>>();
+					List<Tuple<ProjectFile, FileReference>> IntelliSenseTargetFiles = new List<Tuple<ProjectFile, FileReference>>();
 					{
 						// Engine targets
 						if( EngineProject != null)
 						{
-							foreach( var ProjectTarget in EngineProject.ProjectTargets )
+							foreach( ProjectTarget ProjectTarget in EngineProject.ProjectTargets )
 							{
 								if( ProjectTarget.TargetFilePath != null )
 								{
@@ -560,9 +559,9 @@ namespace UnrealBuildTool
 						}
 
 						// Program targets
-						foreach( var ProgramProject in ProgramProjects.Values )
+						foreach( ProjectFile ProgramProject in ProgramProjects.Values )
 						{
-							foreach( var ProjectTarget in ProgramProject.ProjectTargets )
+							foreach( ProjectTarget ProjectTarget in ProgramProject.ProjectTargets )
 							{
 								if( ProjectTarget.TargetFilePath != null )
 								{
@@ -572,9 +571,9 @@ namespace UnrealBuildTool
 						}
 
 						// Game/template targets
-						foreach( var GameProject in GameProjects.Values )
+						foreach( ProjectFile GameProject in GameProjects.Values )
 						{
-							foreach( var ProjectTarget in GameProject.ProjectTargets )
+							foreach( ProjectTarget ProjectTarget in GameProject.ProjectTargets )
 							{
 								if( ProjectTarget.TargetFilePath != null )
 								{
@@ -624,7 +623,7 @@ namespace UnrealBuildTool
 					Path.Combine("Programs", "UnrealBuildTool")
 				));
 
-			foreach(var BuildConfigurationPath in XmlConfigLoader.ConfigLocationHierarchy)
+			foreach(XmlConfigLoader.XmlConfigLocation BuildConfigurationPath in XmlConfigLoader.ConfigLocationHierarchy)
 			{
 				if(!BuildConfigurationPath.bExists)
 				{
@@ -656,7 +655,7 @@ namespace UnrealBuildTool
 		protected virtual void ConfigureProjectFileGeneration( String[] Arguments, ref bool IncludeAllPlatforms )
 		{
 			bool bAlwaysIncludeEngineModules = false;
-			foreach( var CurArgument in Arguments )
+			foreach( string CurArgument in Arguments )
 			{
 				if( CurArgument.StartsWith( "-" ) )
 				{
@@ -792,19 +791,19 @@ namespace UnrealBuildTool
 		/// </summary>
 		private void AddAllGameProjects(Dictionary<DirectoryReference, ProjectFile> GameProjects, string SupportedPlatformNames, MasterProjectFolder ProjectsFolder)
 		{
-			foreach( var GameFolderAndProjectFile in GameProjects )
+			foreach( KeyValuePair<DirectoryReference, ProjectFile> GameFolderAndProjectFile in GameProjects )
 			{
-				var GameProjectDirectory = GameFolderAndProjectFile.Key;
+				DirectoryReference GameProjectDirectory = GameFolderAndProjectFile.Key;
 
 				// @todo projectfiles: We have engine localization files, but should we also add GAME localization files?
 
 				// Game config files
 				if( IncludeConfigFiles )
 				{
-					var GameConfigDirectory = DirectoryReference.Combine(GameProjectDirectory, "Config");
+					DirectoryReference GameConfigDirectory = DirectoryReference.Combine(GameProjectDirectory, "Config");
 					if( GameConfigDirectory.Exists( ) )
 					{
-						var GameProjectFile = GameFolderAndProjectFile.Value;
+						ProjectFile GameProjectFile = GameFolderAndProjectFile.Value;
 						GameProjectFile.AddFilesToProject( SourceFileSearch.FindFiles( GameConfigDirectory ), GameProjectDirectory );
 					}
 				}
@@ -815,7 +814,7 @@ namespace UnrealBuildTool
 		/// Adds all engine localization text files to the specified project
 		private void AddEngineLocalizationFiles( ProjectFile EngineProject )
 		{
-			var EngineLocalizationDirectory = DirectoryReference.Combine( UnrealBuildTool.EngineDirectory, "Content", "Localization" );
+			DirectoryReference EngineLocalizationDirectory = DirectoryReference.Combine( UnrealBuildTool.EngineDirectory, "Content", "Localization" );
 			if( EngineLocalizationDirectory.Exists( ) )
 			{
 				EngineProject.AddFilesToProject( SourceFileSearch.FindFiles( EngineLocalizationDirectory ), UnrealBuildTool.EngineDirectory );
@@ -826,7 +825,7 @@ namespace UnrealBuildTool
         /// Adds all engine template text files to the specified project
         private void AddEngineTemplateFiles( ProjectFile EngineProject )
         {
-            var EngineTemplateDirectory = DirectoryReference.Combine(UnrealBuildTool.EngineDirectory, "Content", "Editor", "Templates");
+            DirectoryReference EngineTemplateDirectory = DirectoryReference.Combine(UnrealBuildTool.EngineDirectory, "Content", "Editor", "Templates");
             if (EngineTemplateDirectory.Exists())
             {
 				EngineProject.AddFilesToProject( SourceFileSearch.FindFiles( EngineTemplateDirectory ), UnrealBuildTool.EngineDirectory );
@@ -837,7 +836,7 @@ namespace UnrealBuildTool
 		/// Adds all engine config files to the specified project
 		private void AddEngineConfigFiles( ProjectFile EngineProject )
 		{
-            var EngineConfigDirectory = DirectoryReference.Combine(UnrealBuildTool.EngineDirectory, "Config" );
+            DirectoryReference EngineConfigDirectory = DirectoryReference.Combine(UnrealBuildTool.EngineDirectory, "Config" );
 			if( EngineConfigDirectory.Exists( ) )
 			{
 				EngineProject.AddFilesToProject( SourceFileSearch.FindFiles( EngineConfigDirectory ), UnrealBuildTool.EngineDirectory );
@@ -847,7 +846,7 @@ namespace UnrealBuildTool
 		/// Adds UnrealHeaderTool config files to the specified project
 		private void AddUnrealHeaderToolConfigFiles(ProjectFile EngineProject)
 		{
-			var UHTConfigDirectory = DirectoryReference.Combine(UnrealBuildTool.EngineDirectory, "Programs", "UnrealHeaderTool", "Config");
+			DirectoryReference UHTConfigDirectory = DirectoryReference.Combine(UnrealBuildTool.EngineDirectory, "Programs", "UnrealHeaderTool", "Config");
 			if (UHTConfigDirectory.Exists())
 			{
 				EngineProject.AddFilesToProject(SourceFileSearch.FindFiles(UHTConfigDirectory), UnrealBuildTool.EngineDirectory);
@@ -864,7 +863,7 @@ namespace UnrealBuildTool
 
 			// Locate all modules (*.Build.cs files)
 			List<FileReference> FoundModuleFiles = RulesCompiler.FindAllRulesSourceFiles( RulesCompiler.RulesFileType.Module, GameFolders: AllGameProjects.Select(x => x.Folder).ToList(), ForeignPlugins:null, AdditionalSearchPaths:null );
-			foreach( var BuildFileName in FoundModuleFiles )
+			foreach( FileReference BuildFileName in FoundModuleFiles )
 			{
 				AllModuleFiles.Add( BuildFileName );
 			}
@@ -888,7 +887,7 @@ namespace UnrealBuildTool
 		/// <returns></returns>
 		private static bool IsNoRedistModule(FileReference ModulePath)
 		{
-			foreach (var NoRedistFolderName in NoRedistFolders)
+			foreach (string NoRedistFolderName in NoRedistFolders)
 			{
 				if (ModulePath.FullName.IndexOf(NoRedistFolderName, StringComparison.InvariantCultureIgnoreCase) >= 0)
 				{
@@ -907,17 +906,17 @@ namespace UnrealBuildTool
 			List<FileReference> AllTargetFiles = new List<FileReference>();
 
 			// Make a list of all platform name strings that we're *not* including in the project files
-			var UnsupportedPlatformNameStrings = Utils.MakeListOfUnsupportedPlatforms( SupportedPlatforms );
+			List<string> UnsupportedPlatformNameStrings = Utils.MakeListOfUnsupportedPlatforms( SupportedPlatforms );
 
 			// Locate all targets (*.Target.cs files)
 			List<FileReference> FoundTargetFiles = RulesCompiler.FindAllRulesSourceFiles( RulesCompiler.RulesFileType.Target, AllGameProjects.Select(x => x.Folder).ToList(), ForeignPlugins:null, AdditionalSearchPaths:null );
 			foreach( FileReference CurTargetFile in FoundTargetFiles )
 			{
-				var CleanTargetFileName = Utils.CleanDirectorySeparators( CurTargetFile.FullName );
+				string CleanTargetFileName = Utils.CleanDirectorySeparators( CurTargetFile.FullName );
 
 				// remove the local root
-				var LocalRoot = Path.GetFullPath(RootRelativePath);
-				var Search = CleanTargetFileName;
+				string LocalRoot = Path.GetFullPath(RootRelativePath);
+				string Search = CleanTargetFileName;
 				if (Search.StartsWith(LocalRoot, StringComparison.InvariantCultureIgnoreCase))
 				{
 					if (LocalRoot.EndsWith("\\") || LocalRoot.EndsWith("/"))
@@ -948,7 +947,7 @@ namespace UnrealBuildTool
 
 				// Skip targets in unsupported platform directories
 				bool IncludeThisTarget = true;
-				foreach( var CurPlatformName in UnsupportedPlatformNameStrings )
+				foreach( string CurPlatformName in UnsupportedPlatformNameStrings )
 				{
 					if (Search.IndexOf(Path.DirectorySeparatorChar + CurPlatformName + Path.DirectorySeparatorChar, StringComparison.InvariantCultureIgnoreCase) != -1)
 					{
@@ -976,7 +975,7 @@ namespace UnrealBuildTool
 		void EliminateRedundantMasterProjectSubFolders( MasterProjectFolder Folder, string ParentMasterProjectFolderPath )
 		{
 			// NOTE: This is for diagnostics output only
-			var MasterProjectFolderPath = String.IsNullOrEmpty( ParentMasterProjectFolderPath ) ? Folder.FolderName : ( ParentMasterProjectFolderPath + "/" + Folder.FolderName );
+			string MasterProjectFolderPath = String.IsNullOrEmpty( ParentMasterProjectFolderPath ) ? Folder.FolderName : ( ParentMasterProjectFolderPath + "/" + Folder.FolderName );
 
 			// We can eliminate folders that meet all of these requirements:
 			//		1) Have only a single project file in them
@@ -988,15 +987,15 @@ namespace UnrealBuildTool
 			// Also, we can kill folders that are completely empty.
 			
 
-			foreach( var SubFolder in Folder.SubFolders )
+			foreach( MasterProjectFolder SubFolder in Folder.SubFolders )
 			{
 				// Recurse
 				EliminateRedundantMasterProjectSubFolders( SubFolder, MasterProjectFolderPath );
 			}
 
-			var SubFoldersToAdd = new List<MasterProjectFolder>();
-			var SubFoldersToRemove = new List<MasterProjectFolder>();
-			foreach( var SubFolder in Folder.SubFolders )
+			List<MasterProjectFolder> SubFoldersToAdd = new List<MasterProjectFolder>();
+			List<MasterProjectFolder> SubFoldersToRemove = new List<MasterProjectFolder>();
+			foreach( MasterProjectFolder SubFolder in Folder.SubFolders )
 			{
 				bool CanCollapseFolder = false;
 
@@ -1051,7 +1050,7 @@ namespace UnrealBuildTool
 				}
 			}
 
-			foreach( var SubFolderToRemove in SubFoldersToRemove )
+			foreach( MasterProjectFolder SubFolderToRemove in SubFoldersToRemove )
 			{
 				Folder.SubFolders.Remove( SubFolderToRemove );
 			}
@@ -1069,9 +1068,9 @@ namespace UnrealBuildTool
 		/// <param name="MasterProjectFolderPath">Parent master project folder path.</param>
 		protected virtual void Validate(MasterProjectFolder Folder, string MasterProjectFolderPath)
 		{
-			foreach (var CurChildProject in Folder.ChildProjects)
+			foreach (ProjectFile CurChildProject in Folder.ChildProjects)
 			{
-				foreach (var OtherChildProject in Folder.ChildProjects)
+				foreach (ProjectFile OtherChildProject in Folder.ChildProjects)
 				{
 					if (CurChildProject != OtherChildProject)
 					{
@@ -1083,26 +1082,26 @@ namespace UnrealBuildTool
 				}
 			}
 
-			foreach (var SubFolder in Folder.SubFolders)
+			foreach (MasterProjectFolder SubFolder in Folder.SubFolders)
 			{
 				// If the parent folder already has a child project or file item with the same name as this sub-folder, then
 				// that's considered an error (it should never have been allowed to have a folder name that collided
 				// with project file names or file items, as that's not supported in Visual Studio.)
-				foreach (var CurChildProject in Folder.ChildProjects)
+				foreach (ProjectFile CurChildProject in Folder.ChildProjects)
 				{
 					if (CurChildProject.ProjectFilePath.GetFileNameWithoutAnyExtensions().Equals(SubFolder.FolderName, StringComparison.InvariantCultureIgnoreCase))
 					{
 						throw new BuildException("Detected collision between a master project sub-folder " + SubFolder.FolderName + " and a project within the outer folder " + CurChildProject.ProjectFilePath + " (master project folder: " + MasterProjectFolderPath + ")");
 					}
 				}
-				foreach (var CurFile in Folder.Files)
+				foreach (string CurFile in Folder.Files)
 				{
 					if (Path.GetFileName(CurFile).Equals(SubFolder.FolderName, StringComparison.InvariantCultureIgnoreCase))
 					{
 						throw new BuildException("Detected collision between a master project sub-folder " + SubFolder.FolderName + " and a file within the outer folder " + CurFile + " (master project folder: " + MasterProjectFolderPath + ")");
 					}
 				}
-				foreach (var CurFolder in Folder.SubFolders)
+				foreach (MasterProjectFolder CurFolder in Folder.SubFolders)
 				{
 					if (CurFolder != SubFolder)
 					{
@@ -1120,11 +1119,11 @@ namespace UnrealBuildTool
 		/// </summary>
 		private void AddUnrealBuildToolProject( MasterProjectFolder ProgramsFolder, IEnumerable<ProjectFile> Dependencies )
 		{
-			var ProjectFileName = FileReference.Combine(UnrealBuildTool.EngineSourceDirectory, "Programs", "UnrealBuildTool", "UnrealBuildTool.csproj" );
-			var UnrealBuildToolProject = new VCSharpProjectFile( ProjectFileName );
+			FileReference ProjectFileName = FileReference.Combine(UnrealBuildTool.EngineSourceDirectory, "Programs", "UnrealBuildTool", "UnrealBuildTool.csproj" );
+			VCSharpProjectFile UnrealBuildToolProject = new VCSharpProjectFile( ProjectFileName );
 			UnrealBuildToolProject.ShouldBuildForAllSolutionTargets = true;
 
-			foreach (var Dependent in Dependencies)
+			foreach (ProjectFile Dependent in Dependencies)
 			{
 				UnrealBuildToolProject.AddDependsOnProject( Dependent );
 			}
@@ -1148,7 +1147,7 @@ namespace UnrealBuildTool
 		{
 			VCSharpProjectFile Project = null;
 
-			var ProjectFileName = FileReference.Combine( UnrealBuildTool.EngineSourceDirectory, "Programs", ProjectName, Path.GetFileName( ProjectName ) + ".csproj" );
+			FileReference ProjectFileName = FileReference.Combine( UnrealBuildTool.EngineSourceDirectory, "Programs", ProjectName, Path.GetFileName( ProjectName ) + ".csproj" );
 			FileInfo Info = new FileInfo( ProjectFileName.FullName );
 			if( Info.Exists )
 			{
@@ -1290,17 +1289,17 @@ namespace UnrealBuildTool
 		{
 			if( IncludeConfigFiles )
 			{
-				foreach( var ProjectFolderAndFile in ProgramProjects )
+				foreach( KeyValuePair<string, ProjectFile> ProjectFolderAndFile in ProgramProjects )
 				{
-					var ProgramFolder = ProjectFolderAndFile.Key;
-					var ProgramProjectFile = ProjectFolderAndFile.Value;
+					string ProgramFolder = ProjectFolderAndFile.Key;
+					ProjectFile ProgramProjectFile = ProjectFolderAndFile.Value;
 
-					var ProgramName = ProgramFolder;
+					string ProgramName = ProgramFolder;
 
 					// @todo projectfiles: The config folder for programs is kind of weird -- you end up going UP a few directories to get to it.  This stuff is not great.
 					// @todo projectfiles: Fragile assumption here about Programs always being under /Engine/Programs
-					var ProgramDirectory = DirectoryReference.Combine( UnrealBuildTool.EngineDirectory, "Programs", ProgramName );
-					var ProgramConfigDirectory = DirectoryReference.Combine( ProgramDirectory, "Config" );
+					DirectoryReference ProgramDirectory = DirectoryReference.Combine( UnrealBuildTool.EngineDirectory, "Programs", ProgramName );
+					DirectoryReference ProgramConfigDirectory = DirectoryReference.Combine( ProgramDirectory, "Config" );
 					if( ProgramConfigDirectory.Exists( ) )
 					{
 						ProgramProjectFile.AddFilesToProject( SourceFileSearch.FindFiles( ProgramConfigDirectory ), ProgramDirectory );
@@ -1318,7 +1317,7 @@ namespace UnrealBuildTool
 		/// <return>Whether the process was successful or not</return>
 		private bool GenerateIntelliSenseData( String[] Arguments, List<Tuple<ProjectFile, FileReference>> TargetFiles )
 		{
-			var bSuccess = true;
+			bool bSuccess = true;
 			if( ShouldGenerateIntelliSenseData() && TargetFiles.Count > 0 )
 			{
 				string ProgressInfoText = Utils.IsRunningOnMono ? "Generating data for project indexing..." : "Binding IntelliSense data...";
@@ -1326,14 +1325,14 @@ namespace UnrealBuildTool
 				{
 					for(int TargetIndex = 0; TargetIndex < TargetFiles.Count; ++TargetIndex)
 					{
-						var TargetProjectFile = TargetFiles[ TargetIndex ].Item1;
-						var CurTarget = TargetFiles[ TargetIndex ].Item2;
+						ProjectFile TargetProjectFile = TargetFiles[ TargetIndex ].Item1;
+						FileReference CurTarget = TargetFiles[ TargetIndex ].Item2;
 
-						var TargetName = CurTarget.GetFileNameWithoutAnyExtensions();
+						string TargetName = CurTarget.GetFileNameWithoutAnyExtensions();
 
 						Log.TraceVerbose( "Found target: " + TargetName );
 
-						var ArgumentsCopy = new string[ Arguments.Length + 1 ];
+						string[] ArgumentsCopy = new string[ Arguments.Length + 1 ];
 						ArgumentsCopy[ 0 ] = TargetName;
 						Array.Copy(Arguments, 0, ArgumentsCopy, 1, Arguments.Length);
 
@@ -1373,7 +1372,7 @@ namespace UnrealBuildTool
 		/// <param name="SupportedPlatformNames">Output string for supported platforms, returned as comma-separated values.</param>
 		protected virtual void SetupSupportedPlatformsAndConfigurations(bool IncludeAllPlatforms, out string SupportedPlatformNames)
 		{
-			var SupportedPlatformsString = new StringBuilder();
+			StringBuilder SupportedPlatformsString = new StringBuilder();
 
 			System.Array PlatformEnums = Enum.GetValues(typeof(UnrealTargetPlatform));
 			foreach (UnrealTargetPlatform Platform in PlatformEnums)
@@ -1388,7 +1387,7 @@ namespace UnrealBuildTool
 				if (bInProjectPlatformsList && (IncludeAllPlatforms || IsRequiredPlatform))
 				{
 					// If there is a build platform present, add it to the SupportedPlatforms list
-					var BuildPlatform = UEBuildPlatform.GetBuildPlatform( Platform, true );
+					UEBuildPlatform BuildPlatform = UEBuildPlatform.GetBuildPlatform( Platform, true );
 					if( BuildPlatform != null )
 					{
 						if (UnrealBuildTool.IsValidPlatform(Platform))
@@ -1431,7 +1430,7 @@ namespace UnrealBuildTool
 		/// <param name="File">Full path of the file to search for</param>
 		protected UProjectInfo FindGameContainingFile(List<UProjectInfo> AllGames, FileReference File)
 		{
-			foreach (var Game in AllGames)
+			foreach (UProjectInfo Game in AllGames)
 			{
 				if (File.IsUnderDirectory(Game.Folder))
 				{
@@ -1453,13 +1452,13 @@ namespace UnrealBuildTool
 			DirectoryReference EngineSourceThirdPartyDirectory = DirectoryReference.Combine(UnrealBuildTool.EngineSourceDirectory, "ThirdParty");
 
 			HashSet<ProjectFile> ProjectsWithPlugins = new HashSet<ProjectFile>();
-			foreach( var CurModuleFile in AllModuleFiles )
+			foreach( FileReference CurModuleFile in AllModuleFiles )
 			{
 				Log.TraceVerbose("AddProjectsForAllModules " + CurModuleFile);
 
 				// The module's "base directory" is simply the directory where its xxx.Build.cs file is stored.  We'll always
 				// harvest source files for this module in this base directory directory and all of its sub-directories.
-				var ModuleName = CurModuleFile.GetFileNameWithoutAnyExtensions();		// Remove both ".cs" and ".Build"
+				string ModuleName = CurModuleFile.GetFileNameWithoutAnyExtensions();		// Remove both ".cs" and ".Build"
 
 				bool WantProjectFileForModule = true;
 
@@ -1506,7 +1505,7 @@ namespace UnrealBuildTool
 					// Setup a project file entry for this module's project.  Remember, some projects may host multiple modules!
 					FileReference ProjectFileName = FileReference.Combine( IntermediateProjectFilesPath, ProjectFileNameBase + ProjectFileExtension );
 					bool bProjectAlreadyExisted;
-					var ProjectFile = FindOrAddProject( ProjectFileName, IncludeInGeneratedProjects:true, bAlreadyExisted:out bProjectAlreadyExisted );
+					ProjectFile ProjectFile = FindOrAddProject( ProjectFileName, IncludeInGeneratedProjects:true, bAlreadyExisted:out bProjectAlreadyExisted );
 
 					// Update our module map
 					ModuleToProjectFileMap[ ModuleName ] = ProjectFile;
@@ -1514,7 +1513,7 @@ namespace UnrealBuildTool
 
 					// Only search subdirectories for non-external modules.  We don't want to add all of the source and header files
 					// for every third-party module, unless we were configured to do so.
-					var SearchSubdirectories = !IsThirdPartyModule || bGatherThirdPartySource;
+					bool SearchSubdirectories = !IsThirdPartyModule || bGatherThirdPartySource;
 
 					if( bGatherThirdPartySource )
 					{
@@ -1523,7 +1522,7 @@ namespace UnrealBuildTool
 
 
 					// Find all of the source files (and other files) and add them to the project
-					var FoundFiles = SourceFileSearch.FindModuleSourceFiles( CurModuleFile, SearchSubdirectories:SearchSubdirectories );
+					List<FileReference> FoundFiles = SourceFileSearch.FindModuleSourceFiles( CurModuleFile, SearchSubdirectories:SearchSubdirectories );
 					ProjectFile.AddFilesToProject( FoundFiles, BaseFolder );
 
 					// Check if there's a plugin directory here
@@ -1575,10 +1574,10 @@ namespace UnrealBuildTool
 
 			// Find all of the target files.  This will filter out any modules or targets that don't
 			// belong to platforms we're generating project files for.
-			var AllTargetFiles = DiscoverTargets(AllGames);
-			foreach( var TargetFilePath in AllTargetFiles )
+			List<FileReference> AllTargetFiles = DiscoverTargets(AllGames);
+			foreach( FileReference TargetFilePath in AllTargetFiles )
 			{
-				var TargetName = TargetFilePath.GetFileNameWithoutAnyExtensions();		// Remove both ".cs" and ".Target"
+				string TargetName = TargetFilePath.GetFileNameWithoutAnyExtensions();		// Remove both ".cs" and ".Target"
 
 				// Check to see if this is an Engine target.  That is, the target is located under the "Engine" folder
 				bool IsEngineTarget = false;
@@ -1614,7 +1613,7 @@ namespace UnrealBuildTool
 
 					// Create target rules for all of the platforms and configuration combinations that we want to enable support for.
 					// Just use the current platform as we only need to recover the target type and both should be supported for all targets...
-					var TargetRulesObject = RulesAssembly.CreateTargetRules(TargetName, new TargetInfo(BuildHostPlatform.Current.Platform, UnrealTargetConfiguration.Development, ""), false);
+					TargetRules TargetRulesObject = RulesAssembly.CreateTargetRules(TargetName, new TargetInfo(BuildHostPlatform.Current.Platform, UnrealTargetConfiguration.Development, ""), false);
 
 					bool IsProgramTarget = false;
 
@@ -1659,7 +1658,7 @@ namespace UnrealBuildTool
 					}
 
 					bool bProjectAlreadyExisted;
-					var ProjectFile = FindOrAddProject(ProjectFilePath, IncludeInGeneratedProjects: true, bAlreadyExisted: out bProjectAlreadyExisted);
+					ProjectFile ProjectFile = FindOrAddProject(ProjectFilePath, IncludeInGeneratedProjects: true, bAlreadyExisted: out bProjectAlreadyExisted);
 					ProjectFile.IsForeignProject = bGeneratingGameProjectFiles && OnlyGameProject != null && TargetFilePath.IsUnderDirectory(OnlyGameProject.Directory);
 					ProjectFile.IsGeneratedProject = true;
 					ProjectFile.IsStubProject = false;
@@ -1697,7 +1696,7 @@ namespace UnrealBuildTool
 						if (!bProjectAlreadyExisted)
 						{
 							// Add the .uproject file for this game/template
-							var UProjectFilePath = FileReference.Combine(BaseFolder, ProjectFileNameBase + ".uproject");
+							FileReference UProjectFilePath = FileReference.Combine(BaseFolder, ProjectFileNameBase + ".uproject");
 							if (UProjectFilePath.Exists())
 							{
 								ProjectFile.AddFileToProject(UProjectFilePath, BaseFolder);
@@ -1710,7 +1709,7 @@ namespace UnrealBuildTool
 
 					}
 
-					foreach (var ExistingProjectTarget in ProjectFile.ProjectTargets)
+					foreach (ProjectTarget ExistingProjectTarget in ProjectFile.ProjectTargets)
 					{
 						if (ExistingProjectTarget.TargetRules.ConfigurationName.Equals(TargetRulesObject.ConfigurationName, StringComparison.InvariantCultureIgnoreCase))
 						{
@@ -1725,7 +1724,7 @@ namespace UnrealBuildTool
 						}
 					}
 
-					var ProjectTarget = new ProjectTarget()
+					ProjectTarget ProjectTarget = new ProjectTarget()
 						{
 							TargetRules = TargetRulesObject,
 							TargetFilePath = TargetFilePath,
@@ -1759,9 +1758,9 @@ namespace UnrealBuildTool
 		protected void AddEngineShaderSource( ProjectFile EngineProject )
 		{
 			// Setup a project file entry for this module's project.  Remember, some projects may host multiple modules!
-			var ShadersDirectory = DirectoryReference.Combine( UnrealBuildTool.EngineDirectory, "Shaders" );
+			DirectoryReference ShadersDirectory = DirectoryReference.Combine( UnrealBuildTool.EngineDirectory, "Shaders" );
 
-			var SubdirectoryNamesToExclude = new List<string>();
+			List<string> SubdirectoryNamesToExclude = new List<string>();
 			{
                 // Don't include binary shaders in the project file.
                 SubdirectoryNamesToExclude.Add( "Binaries" );
@@ -1777,9 +1776,9 @@ namespace UnrealBuildTool
 		/// Adds engine build infrastructure files to the specified project
 		protected void AddEngineBuildFiles( ProjectFile EngineProject )
 		{
-			var BuildDirectory = DirectoryReference.Combine( UnrealBuildTool.EngineDirectory, "Build" );
+			DirectoryReference BuildDirectory = DirectoryReference.Combine( UnrealBuildTool.EngineDirectory, "Build" );
 
-			var SubdirectoryNamesToExclude = new List<string>();
+			List<string> SubdirectoryNamesToExclude = new List<string>();
 			SubdirectoryNamesToExclude.Add("Receipts");
 
 			EngineProject.AddFilesToProject( SourceFileSearch.FindFiles( BuildDirectory, SubdirectoryNamesToExclude ), UnrealBuildTool.EngineDirectory );
@@ -1790,14 +1789,14 @@ namespace UnrealBuildTool
 		protected void AddEngineDocumentation( ProjectFile EngineProject )
 		{
 			// NOTE: The project folder added here will actually be collapsed away later if not needed
-			var DocumentationProjectDirectory = DirectoryReference.Combine( UnrealBuildTool.EngineDirectory, "Documentation" );
-			var DocumentationSourceDirectory = DirectoryReference.Combine( UnrealBuildTool.EngineDirectory, "Documentation", "Source" );
+			DirectoryReference DocumentationProjectDirectory = DirectoryReference.Combine( UnrealBuildTool.EngineDirectory, "Documentation" );
+			DirectoryReference DocumentationSourceDirectory = DirectoryReference.Combine( UnrealBuildTool.EngineDirectory, "Documentation", "Source" );
 			DirectoryInfo DirInfo = new DirectoryInfo( DocumentationProjectDirectory.FullName );
 			if( DirInfo.Exists && DocumentationSourceDirectory.Exists() )
 			{
 				Log.TraceVerbose( "Adding documentation files..." );
 
-				var SubdirectoryNamesToExclude = new List<string>();
+				List<string> SubdirectoryNamesToExclude = new List<string>();
 				{
 					// We never want any of the images or attachment files included in our generated project
 					SubdirectoryNamesToExclude.Add( "Images" );
@@ -1810,18 +1809,18 @@ namespace UnrealBuildTool
 					SubdirectoryNamesToExclude.Add( "Javascript" );
 				}
 
-				var DocumentationFiles = SourceFileSearch.FindFiles( DocumentationSourceDirectory, SubdirectoryNamesToExclude );
+				List<FileReference> DocumentationFiles = SourceFileSearch.FindFiles( DocumentationSourceDirectory, SubdirectoryNamesToExclude );
 
 				// Filter out non-English documentation files if we were configured to do so
 				if( !bAllDocumentationLanguages )
 				{
-					var FilteredDocumentationFiles = new List<FileReference>();
-					foreach( var DocumentationFile in DocumentationFiles )
+					List<FileReference> FilteredDocumentationFiles = new List<FileReference>();
+					foreach( FileReference DocumentationFile in DocumentationFiles )
 					{
 						bool bPassesFilter = true;
 						if( DocumentationFile.FullName.EndsWith( ".udn", StringComparison.InvariantCultureIgnoreCase ) )
 						{
-							var LanguageSuffix = Path.GetExtension( Path.GetFileNameWithoutExtension( DocumentationFile.FullName ) );
+							string LanguageSuffix = Path.GetExtension( Path.GetFileNameWithoutExtension( DocumentationFile.FullName ) );
 							if( !String.IsNullOrEmpty( LanguageSuffix ) &&
 								!LanguageSuffix.Equals( ".int", StringComparison.InvariantCultureIgnoreCase ) )
 							{
@@ -1869,7 +1868,7 @@ namespace UnrealBuildTool
 			}
 
 			// Add a new project file for the specified path
-			var NewProjectFile = AllocateProjectFile( FilePath );
+			ProjectFile NewProjectFile = AllocateProjectFile( FilePath );
 			ProjectFileMap[ FilePath ] = NewProjectFile;
 
 			if( IncludeInGeneratedProjects )
@@ -1908,7 +1907,7 @@ namespace UnrealBuildTool
 		{
 			get
 			{
-				var CombinedList = new List<ProjectFile>();
+				List<ProjectFile> CombinedList = new List<ProjectFile>();
 				CombinedList.AddRange( GeneratedProjectFiles );
 				CombinedList.AddRange( OtherProjectFiles );
 				return CombinedList;
@@ -1924,11 +1923,11 @@ namespace UnrealBuildTool
 		{
             using(ProgressWriter Progress = new ProgressWriter("Writing project files...", true))
 			{
-				var TotalProjectFileCount = GeneratedProjectFiles.Count + 1;	// +1 for the master project file, which we'll save next
+				int TotalProjectFileCount = GeneratedProjectFiles.Count + 1;	// +1 for the master project file, which we'll save next
 
 				for(int ProjectFileIndex = 0 ; ProjectFileIndex < GeneratedProjectFiles.Count; ++ProjectFileIndex )
 				{
-					var CurProject = GeneratedProjectFiles[ ProjectFileIndex ];
+					ProjectFile CurProject = GeneratedProjectFiles[ ProjectFileIndex ];
 					if( !CurProject.WriteProjectFile(
 								InPlatforms: SupportedPlatforms,
 								InConfigurations: SupportedConfigurations ) )
@@ -1946,9 +1945,9 @@ namespace UnrealBuildTool
             // Write AutomationReferences file
             if (AutomationProjectFiles.Any())
             {
-                var NS = XNamespace.Get("http://schemas.microsoft.com/developer/msbuild/2003");
+                XNamespace NS = XNamespace.Get("http://schemas.microsoft.com/developer/msbuild/2003");
 
-                var AutomationToolDir = DirectoryReference.Combine(UnrealBuildTool.EngineSourceDirectory, "Programs", "AutomationTool");
+                DirectoryReference AutomationToolDir = DirectoryReference.Combine(UnrealBuildTool.EngineSourceDirectory, "Programs", "AutomationTool");
                 new XDocument(
                     new XElement(NS + "Project",
                         new XAttribute("ToolsVersion", VCProjectFileGenerator.ProjectFileToolVersionString),
@@ -1976,9 +1975,9 @@ namespace UnrealBuildTool
 		protected virtual bool WriteProjectFileManifest()
 		{
 			BuildManifest Manifest = new BuildManifest();
-			foreach( var CurProject in GeneratedProjectFiles )
+			foreach( ProjectFile CurProject in GeneratedProjectFiles )
 			{
-				foreach( var SourceFile in CurProject.SourceFiles )
+				foreach( ProjectFile.SourceFile SourceFile in CurProject.SourceFiles )
 				{
 					Manifest.AddBuildProduct( SourceFile.Reference.FullName );
 				}
@@ -2007,7 +2006,7 @@ namespace UnrealBuildTool
 		{
 			// Check to see if the file already exists, and if so, load it up
 			string LoadedFileContent = null;
-			var FileAlreadyExists = File.Exists( FileName );
+			bool FileAlreadyExists = File.Exists( FileName );
 			if( FileAlreadyExists )
 			{
 				try
@@ -2025,7 +2024,7 @@ namespace UnrealBuildTool
 			bool FileNeedsSave = true;
 			if( LoadedFileContent != null )
 			{
-				var bIgnoreProjectFileWhitespaces = true;
+				bool bIgnoreProjectFileWhitespaces = true;
 				if (ProjectFileComparer.CompareOrdinalIgnoreCase(LoadedFileContent, NewFileContents, bIgnoreProjectFileWhitespaces) == 0)
 				{
 					// Exact match!
@@ -2082,7 +2081,7 @@ namespace UnrealBuildTool
 				throw new BuildException( "Expecting existing project to not have any ProjectTargets defined yet." );
 			}
 			
-			var ProjectTarget = new ProjectTarget();
+			ProjectTarget ProjectTarget = new ProjectTarget();
 			if( bForceDevelopmentConfiguration )
 			{
 				ProjectTarget.ForceDevelopmentConfiguration = true;
@@ -2092,14 +2091,14 @@ namespace UnrealBuildTool
 			if (bNeedsAllPlatformAndConfigurations)
 			{
 				// Add all platforms
-				var AllPlatforms = Enum.GetValues(typeof(UnrealTargetPlatform));
+				Array AllPlatforms = Enum.GetValues(typeof(UnrealTargetPlatform));
 				foreach (UnrealTargetPlatform CurPlatfrom in AllPlatforms)
 				{
 					ProjectTarget.ExtraSupportedPlatforms.Add(CurPlatfrom);
 				}
 
 				// Add all configurations
-				var AllConfigurations = Enum.GetValues(typeof(UnrealTargetConfiguration));
+				Array AllConfigurations = Enum.GetValues(typeof(UnrealTargetConfiguration));
 				foreach (UnrealTargetConfiguration CurConfiguration in AllConfigurations)
 				{
 					ProjectTarget.ExtraSupportedConfigurations.Add( CurConfiguration );
@@ -2118,7 +2117,7 @@ namespace UnrealBuildTool
 				else
 				{
 					// Otherwise, add all platforms
-					var AllPlatforms = Enum.GetValues(typeof(UnrealTargetPlatform));
+					Array AllPlatforms = Enum.GetValues(typeof(UnrealTargetPlatform));
 					foreach (UnrealTargetPlatform CurPlatfrom in AllPlatforms)
 					{
 						ProjectTarget.ExtraSupportedPlatforms.Add(CurPlatfrom);
@@ -2136,7 +2135,7 @@ namespace UnrealBuildTool
 				else
 				{
 					// Otherwise, add all configurations
-					var AllConfigurations = Enum.GetValues(typeof(UnrealTargetConfiguration));
+					Array AllConfigurations = Enum.GetValues(typeof(UnrealTargetConfiguration));
 					foreach (UnrealTargetConfiguration CurConfiguration in AllConfigurations)
 					{
 						ProjectTarget.ExtraSupportedConfigurations.Add(CurConfiguration);
@@ -2163,7 +2162,7 @@ namespace UnrealBuildTool
 		/// The default project to be built for the solution.
 		protected ProjectFile DefaultProject;
 
-		/// The project for UnrealBuildTool.  Note that when generating Rocket project files, we won't have
+		/// The project for UnrealBuildTool.  Note that when generating project files for installed builds, we won't have
 		/// an UnrealBuildTool project at all.
 		protected ProjectFile UBTProject;
 
diff --git a/Engine/Source/Programs/UnrealBuildTool/System/RulesCompiler.cs b/Engine/Source/Programs/UnrealBuildTool/System/RulesCompiler.cs
--- a/Engine/Source/Programs/UnrealBuildTool/System/RulesCompiler.cs
+++ b/Engine/Source/Programs/UnrealBuildTool/System/RulesCompiler.cs
@@ -976,7 +976,7 @@ namespace UnrealBuildTool
 		/// <returns>a list of target platforms for the monolithic</returns>        
 		public virtual List<UnrealTargetPlatform> GUBP_GetPlatforms_MonolithicOnly(UnrealTargetPlatform HostPlatform)
 		{
-			var Result = new List<UnrealTargetPlatform> { HostPlatform };
+			List<UnrealTargetPlatform> Result = new List<UnrealTargetPlatform> { HostPlatform };
 			// hack to set up the templates without adding anything to their .targets.cs files
 			if (!String.IsNullOrEmpty(TargetName) && TargetName.StartsWith("TP_"))
 			{
@@ -1000,7 +1000,7 @@ namespace UnrealBuildTool
 		/// <returns>a list of target platforms for the monolithic without cook</returns>        
 		public virtual List<UnrealTargetPlatform> GUBP_GetBuildOnlyPlatforms_MonolithicOnly(UnrealTargetPlatform HostPlatform)
 		{
-			var Result = new List<UnrealTargetPlatform> { };
+			List<UnrealTargetPlatform> Result = new List<UnrealTargetPlatform> { };
 			return Result;
 		}
 		/// <summary>
@@ -1071,7 +1071,7 @@ namespace UnrealBuildTool
 		}
 		public virtual GUBPProjectOptions GUBP_IncludeProjectInPromotedBuild_EditorTypeOnly(UnrealTargetPlatform HostPlatform)
 		{
-			var Result = new GUBPProjectOptions();
+			GUBPProjectOptions Result = new GUBPProjectOptions();
 			// hack to set up the templates without adding anything to their .targets.cs files
 			// tweaked to include FP_ folders too - which are temporary
 			if (!String.IsNullOrEmpty(TargetName) && (TargetName.StartsWith("TP_") || TargetName.StartsWith("FP_")))
@@ -1112,8 +1112,8 @@ namespace UnrealBuildTool
 		/// </summary>
 		public virtual Dictionary<string, string> GUBP_GetEditorTests_EditorTypeOnly(UnrealTargetPlatform HostPlatform)
 		{
-			var MacOption = HostPlatform == UnrealTargetPlatform.Mac ? " -Mac" : "";
-			var Result = new Dictionary<string, string>();
+			string MacOption = HostPlatform == UnrealTargetPlatform.Mac ? " -Mac" : "";
+			Dictionary<string, string> Result = new Dictionary<string, string>();
 			Result.Add("EditorTest", "BuildCookRun -run -editortest -unattended -nullrhi -NoP4" + MacOption);
 			Result.Add("GameTest", "BuildCookRun -run -unattended -nullrhi -NoP4" + MacOption);
 			Result.Add("EditorAutomationTest", "BuildCookRun -run -editortest -RunAutomationTests -unattended -nullrhi -NoP4" + MacOption);
@@ -1145,7 +1145,7 @@ namespace UnrealBuildTool
 		/// </summary>
 		public virtual Dictionary<string, string> GUBP_GetGameTests_MonolithicOnly(UnrealTargetPlatform HostPlatform, UnrealTargetPlatform AltHostPlatform, UnrealTargetPlatform Platform)
 		{
-			var Result = new Dictionary<string, string>();
+			Dictionary<string, string> Result = new Dictionary<string, string>();
 			if ((Platform == HostPlatform || Platform == AltHostPlatform) && Type == TargetType.Game)  // for now, we will only run these for the dev config of the host platform
 			{
 				Result.Add("CookedGameTest", "BuildCookRun -run -skipcook -stage -pak -deploy -unattended -nullrhi -NoP4 -platform=" + Platform.ToString());
@@ -1158,7 +1158,7 @@ namespace UnrealBuildTool
 		/// </summary>
 		public virtual Dictionary<string, string> GUBP_GetClientServerTests_MonolithicOnly(UnrealTargetPlatform HostPlatform, UnrealTargetPlatform AltHostPlatform, UnrealTargetPlatform ServerPlatform, UnrealTargetPlatform ClientPlatform)
 		{
-			var Result = new Dictionary<string, string>();
+			Dictionary<string, string> Result = new Dictionary<string, string>();
 #if false // needs work
             if ((ServerPlatform == HostPlatform || ServerPlatform == AltHostPlatform) &&
                 (ClientPlatform == HostPlatform || ClientPlatform == AltHostPlatform) && 
@@ -1389,7 +1389,7 @@ namespace UnrealBuildTool
 		public ModuleRules CreateModuleRules(string ModuleName, TargetInfo Target, out FileReference ModuleFileName)
 		{
 			// Currently, we expect the user's rules object type name to be the same as the module name
-			var ModuleTypeName = ModuleName;
+			string ModuleTypeName = ModuleName;
 
 			// Make sure the module file is known to us
 			if (!ModuleNameToModuleFile.TryGetValue(ModuleName, out ModuleFileName))
@@ -1409,7 +1409,7 @@ namespace UnrealBuildTool
 			TargetInfo LocalTarget = new TargetInfo(LocalPlatform, LocalConfiguration, Target.Architecture, Target.Type.Value, Target.bIsMonolithic.Value);
 
 			// The build module must define a type named 'Rules' that derives from our 'ModuleRules' type.  
-			var RulesObjectType = CompiledAssembly.GetType(ModuleName);
+			Type RulesObjectType = CompiledAssembly.GetType(ModuleName);
 
 			if (RulesObjectType == null)
 			{
@@ -1435,7 +1435,7 @@ namespace UnrealBuildTool
 			}
 
 			// Update the run-time dependencies path to remove $(PluginDir) and replace with a full path. When the receipt is saved it'll be converted to a $(ProjectDir) or $(EngineDir) equivalent.
-			foreach (var Dependency in RulesObject.RuntimeDependencies)
+			foreach (RuntimeDependency Dependency in RulesObject.RuntimeDependencies)
 			{
 				const string PluginDirVariable = "$(PluginDir)";
 				if (Dependency.Path.StartsWith(PluginDirVariable, StringComparison.InvariantCultureIgnoreCase))
@@ -1478,7 +1478,7 @@ namespace UnrealBuildTool
 			}
 			catch (Exception Ex)
 			{
-				var AssemblyFileName = Path.GetFileNameWithoutExtension(CompiledAssembly.Location);
+				string AssemblyFileName = Path.GetFileNameWithoutExtension(CompiledAssembly.Location);
 				throw new BuildException(Ex,
 					"Unable to instantiate instance of '{0}' object type from compiled assembly '{1}'.  Unreal Build Tool creates an instance of your module's 'Rules' object in order to find out about your module's requirements.  The CLR exception details may provide more information:  {2}",
 					InTargetName, AssemblyFileName, Ex.ToString());
@@ -1749,7 +1749,7 @@ namespace UnrealBuildTool
 			// Process the additional search path, if sent in
 			if (AdditionalSearchPaths != null)
 			{
-				foreach (var AdditionalSearchPath in AdditionalSearchPaths)
+				foreach (DirectoryReference AdditionalSearchPath in AdditionalSearchPaths)
 				{
 					if (AdditionalSearchPath != null)
 					{
diff --git a/Engine/Source/Programs/UnrealBuildTool/System/SNDBS.cs b/Engine/Source/Programs/UnrealBuildTool/System/SNDBS.cs
--- a/Engine/Source/Programs/UnrealBuildTool/System/SNDBS.cs
+++ b/Engine/Source/Programs/UnrealBuildTool/System/SNDBS.cs
@@ -7,6 +7,7 @@ using System.Diagnostics;
 using System.Xml;
 using System.Text.RegularExpressions;
 using System.Linq;
+using System.Management;
 using System.Reflection;
 using System.Threading;
 
@@ -35,7 +36,7 @@ namespace UnrealBuildTool
 		/// <param name="e">  Event arguments (In this case, the line of string output)</param>
 		static protected void ActionDebugOutput(object sender, DataReceivedEventArgs e)
 		{
-			var Output = e.Data;
+			string Output = e.Data;
 			if (Output == null)
 			{
 				return;
@@ -311,10 +312,10 @@ namespace UnrealBuildTool
 				{
 					try
 					{
-						using (var Mos = new System.Management.ManagementObjectSearcher("Select * from Win32_Processor"))
+						using (ManagementObjectSearcher Mos = new System.Management.ManagementObjectSearcher("Select * from Win32_Processor"))
 						{
-							var MosCollection = Mos.Get();
-							foreach (var Item in MosCollection)
+							ManagementObjectCollection MosCollection = Mos.Get();
+							foreach (ManagementBaseObject Item in MosCollection)
 							{
 								NumCores += int.Parse(Item["NumberOfCores"].ToString());
 							}
diff --git a/Engine/Source/Programs/UnrealBuildTool/System/UProjectInfo.cs b/Engine/Source/Programs/UnrealBuildTool/System/UProjectInfo.cs
--- a/Engine/Source/Programs/UnrealBuildTool/System/UProjectInfo.cs
+++ b/Engine/Source/Programs/UnrealBuildTool/System/UProjectInfo.cs
@@ -56,7 +56,7 @@ namespace UnrealBuildTool
 			{
 				Files = Directory.GetFiles(InTargetFolder.FullName, "*.Target.cs", SearchOption.TopDirectoryOnly).AsEnumerable();
 			}
-			foreach (var TargetFilename in Files)
+			foreach (string TargetFilename in Files)
 			{
 				bFoundTargetFiles = true;
 				foreach (KeyValuePair<FileReference, UProjectInfo> Entry in ProjectInfoDictionary)
@@ -87,7 +87,7 @@ namespace UnrealBuildTool
 			bOutFoundTargetFiles |= FindTargetFilesInFolder(CurrentTopDirectory);
 			if (bOutFoundTargetFiles == false)
 			{
-				foreach (var TargetFolder in Directory.EnumerateDirectories(CurrentTopDirectory.FullName, "*", SearchOption.TopDirectoryOnly).Select(x => new DirectoryReference(x)))
+				foreach (DirectoryReference TargetFolder in Directory.EnumerateDirectories(CurrentTopDirectory.FullName, "*", SearchOption.TopDirectoryOnly).Select(x => new DirectoryReference(x)))
 				{
 					SubFolderList.Add(TargetFolder);
 					bOutFoundTargetFiles |= FindTargetFilesInFolder(TargetFolder);
@@ -97,7 +97,7 @@ namespace UnrealBuildTool
 			if (bOutFoundTargetFiles == false)
 			{
 				// Recurse each folders folders
-				foreach (var SubFolder in SubFolderList)
+				foreach (DirectoryReference SubFolder in SubFolderList)
 				{
 					FindTargetFiles(SubFolder, ref bOutFoundTargetFiles);
 				}
@@ -127,7 +127,7 @@ namespace UnrealBuildTool
 				UProjectInfo NewProjectInfo = new UProjectInfo(ProjectFile, bIsCodeProject);
 				if (ShortProjectNameDictionary.ContainsKey(NewProjectInfo.GameName))
 				{
-					var FirstProject = ProjectInfoDictionary[ShortProjectNameDictionary[NewProjectInfo.GameName]];
+					UProjectInfo FirstProject = ProjectInfoDictionary[ShortProjectNameDictionary[NewProjectInfo.GameName]];
 					throw new BuildException("There are multiple projects with name {0}\n\t* {1}\n\t* {2}\nThis is not currently supported.", NewProjectInfo.GameName, FirstProject.FilePath.FullName, NewProjectInfo.FilePath.FullName);
 				}
 
@@ -164,7 +164,7 @@ namespace UnrealBuildTool
 			string RootDirectory = Path.Combine(Path.GetDirectoryName(Assembly.GetExecutingAssembly().GetOriginalLocation()), "..", "..", "..");
 			string EngineSourceDirectory = Path.GetFullPath(Path.Combine(RootDirectory, "Engine", "Source"));
 
-			foreach (var File in Directory.EnumerateFiles(RootDirectory, "*.uprojectdirs", SearchOption.TopDirectoryOnly))
+			foreach (string File in Directory.EnumerateFiles(RootDirectory, "*.uprojectdirs", SearchOption.TopDirectoryOnly))
 			{
 				string FilePath = Path.GetFullPath(File);
 				Log.TraceVerbose("\tFound uprojectdirs file {0}", FilePath);
diff --git a/Engine/Source/Programs/UnrealBuildTool/System/Unity.cs b/Engine/Source/Programs/UnrealBuildTool/System/Unity.cs
--- a/Engine/Source/Programs/UnrealBuildTool/System/Unity.cs
+++ b/Engine/Source/Programs/UnrealBuildTool/System/Unity.cs
@@ -129,7 +129,7 @@ namespace UnrealBuildTool
 			{
 				EndCurrentUnityFile();
 
-				var Result = UnityFiles;
+				List<FileCollection> Result = UnityFiles;
 
 				// Null everything to ensure that failure will occur if you accidentally reuse this object.
 				CurrentCollection = null;
@@ -157,9 +157,9 @@ namespace UnrealBuildTool
 			string BaseName
 			)
 		{
-			var NewCPPFiles = new List<FileItem>();
+			List<FileItem> NewCPPFiles = new List<FileItem>();
 
-			var BuildPlatform = UEBuildPlatform.GetBuildPlatformForCPPTargetPlatform(CompileEnvironment.Config.Target.Platform);
+			UEBuildPlatform BuildPlatform = UEBuildPlatform.GetBuildPlatformForCPPTargetPlatform(CompileEnvironment.Config.Target.Platform);
 
 			// Figure out size of all input files combined. We use this to determine whether to use larger unity threshold or not.
 			long TotalBytesInCPPFiles = CPPFiles.Sum(F => F.Info.Length);
@@ -183,7 +183,7 @@ namespace UnrealBuildTool
 				// Sort the incoming file paths alphabetically, so there will be consistency in unity blobs across multiple machines.
 				// Note that we're relying on this not only sorting files within each directory, but also the directories
 				// themselves, so the whole list of file paths is the same across computers.
-				var SortedCPPFiles = CPPFiles.GetRange(0, CPPFiles.Count);
+				List<FileItem> SortedCPPFiles = CPPFiles.GetRange(0, CPPFiles.Count);
 				{
 					// Case-insensitive file path compare, because you never know what is going on with local file systems
 					Comparison<FileItem> FileItemComparer = (FileA, FileB) => { return FileA.AbsolutePath.ToLowerInvariant().CompareTo(FileB.AbsolutePath.ToLowerInvariant()); };
@@ -197,7 +197,7 @@ namespace UnrealBuildTool
 				{
 					int CandidateWorkingSetSourceFileCount = 0;
 					int WorkingSetSourceFileCount = 0;
-					foreach (var CPPFile in SortedCPPFiles)
+					foreach (FileItem CPPFile in SortedCPPFiles)
 					{
 						// Don't include writable source files into unity blobs
 						if (!CPPFile.Reference.IsUnderDirectory(Target.EngineIntermediateDirectory) &&
@@ -225,9 +225,9 @@ namespace UnrealBuildTool
 					}
 				}
 
-				var CPPUnityFileBuilder = new UnityFileBuilder(bForceIntoSingleUnityFile ? -1 : BuildConfiguration.NumIncludedBytesPerUnityCPP);
-				var AdaptiveUnityBuildInfoString = new StringBuilder();
-				foreach (var CPPFile in SortedCPPFiles)
+				UnityFileBuilder CPPUnityFileBuilder = new UnityFileBuilder(bForceIntoSingleUnityFile ? -1 : BuildConfiguration.NumIncludedBytesPerUnityCPP);
+				StringBuilder AdaptiveUnityBuildInfoString = new StringBuilder();
+				foreach (FileItem CPPFile in SortedCPPFiles)
 				{
 					if (!bForceIntoSingleUnityFile && CPPFile.AbsolutePath.IndexOf(".GeneratedWrapper.", StringComparison.InvariantCultureIgnoreCase) != -1)
 					{
@@ -246,7 +246,7 @@ namespace UnrealBuildTool
 						// longer compile times the first time you build after your working file set changes.
 						CPPUnityFileBuilder.AddVirtualFile(CPPFile);
 
-						var CPPFileName = Path.GetFileName(CPPFile.AbsolutePath);
+						string CPPFileName = Path.GetFileName(CPPFile.AbsolutePath);
 						if (AdaptiveUnityBuildInfoString.Length == 0)
 						{
 							AdaptiveUnityBuildInfoString.Append(String.Format("[Adaptive unity build] Excluded from {0} unity file: {1}", BaseName, CPPFileName));
@@ -292,7 +292,7 @@ namespace UnrealBuildTool
 			// Create a set of CPP files that combine smaller CPP files into larger compilation units, along with the corresponding 
 			// actions to compile them.
 			int CurrentUnityFileCount = 0;
-			foreach (var UnityFile in AllUnityFiles)
+			foreach (FileCollection UnityFile in AllUnityFiles)
 			{
 				++CurrentUnityFileCount;
 
@@ -319,7 +319,7 @@ namespace UnrealBuildTool
 				}
 
 				// Add source files to the unity file
-				foreach (var CPPFile in UnityFile.Files)
+				foreach (FileItem CPPFile in UnityFile.Files)
 				{
 					OutputUnityCPPWriter.WriteLine("#include \"{0}\"", ToolChain.ConvertPath(CPPFile.AbsolutePath));
 					if (OutputUnityCPPWriterExtra != null)
