Copying //UE4/Dev-Build to //UE4/Main

diff --git a/Engine/Source/Programs/UnrealBuildTool/System/UnrealBuildTool.cs b/Engine/Source/Programs/UnrealBuildTool/System/UnrealBuildTool.cs
--- a/Engine/Source/Programs/UnrealBuildTool/System/UnrealBuildTool.cs
+++ b/Engine/Source/Programs/UnrealBuildTool/System/UnrealBuildTool.cs
@@ -95,6 +95,31 @@ namespace UnrealBuildTool
 		/// </summary>
 		public static readonly DirectoryReference EngineSourceDirectory = DirectoryReference.Combine(EngineDirectory, "Source");
 
+		/// <summary>
+		/// Full path to the Engine/Source/Runtime directory
+		/// </summary>
+		public static readonly DirectoryReference EngineSourceRuntimeDirectory = DirectoryReference.Combine(EngineSourceDirectory, "Runtime");
+
+		/// <summary>
+		/// Full path to the Engine/Source/Developer directory
+		/// </summary>
+		public static readonly DirectoryReference EngineSourceDeveloperDirectory = DirectoryReference.Combine(EngineSourceDirectory, "Developer");
+
+		/// <summary>
+		/// Full path to the Engine/Source/Editor directory
+		/// </summary>
+		public static readonly DirectoryReference EngineSourceEditorDirectory = DirectoryReference.Combine(EngineSourceDirectory, "Editor");
+
+		/// <summary>
+		/// Full path to the Engine/Source/Programs directory
+		/// </summary>
+		public static readonly DirectoryReference EngineSourceProgramsDirectory = DirectoryReference.Combine(EngineSourceDirectory, "Programs");
+
+		/// <summary>
+		/// Full path to the Engine/Source/ThirdParty directory
+		/// </summary>
+		public static readonly DirectoryReference EngineSourceThirdPartyDirectory = DirectoryReference.Combine(EngineSourceDirectory, "ThirdParty");
+
 		/// <summary>
 		/// The Remote Ini directory.  This should always be valid when compiling using a remote server.
 		/// </summary>
@@ -307,7 +332,7 @@ namespace UnrealBuildTool
 			{
 				if (Predicate(Platform))
 				{
-					var BuildPlatform = UEBuildPlatform.GetBuildPlatform(Platform, true);
+					UEBuildPlatform BuildPlatform = UEBuildPlatform.GetBuildPlatform(Platform, true);
 					if ((bCheckValidity == false) || (BuildPlatform != null))
 					{
 						OutPlatforms.Add(Platform);
@@ -384,7 +409,7 @@ namespace UnrealBuildTool
 
 
 		/// <summary>
-		/// Is this a valid platform. Used primarily for Rocket vs non-Rocket.
+		/// Is this a valid platform. Used primarily for Installed vs non-Installed builds.
 		/// </summary>
 		/// <param name="InPlatform"></param>
 		/// <returns>true if valid, false if not</returns>
@@ -432,25 +457,25 @@ namespace UnrealBuildTool
 				Log.TraceVerbose("Searching for ToolChains, BuildPlatforms, BuildDeploys and ProjectGenerators...");
 
 				List<System.Type> ProjectGeneratorList = new List<System.Type>();
-				var AllTypes = UBTAssembly.GetTypes();
+				Type[] AllTypes = UBTAssembly.GetTypes();
 				if (!bSkipBuildPlatforms)
 				{
 					// register all build platforms first, since they implement SDK-switching logic that can set environment variables
-					foreach (var CheckType in AllTypes)
+					foreach (Type CheckType in AllTypes)
 					{
 						if (CheckType.IsClass && !CheckType.IsAbstract)
 						{
 							if (CheckType.IsSubclassOf(typeof(UEBuildPlatformFactory)))
 							{
 								Log.TraceVerbose("    Registering build platform: {0}", CheckType.ToString());
-								var TempInst = (UEBuildPlatformFactory)(UBTAssembly.CreateInstance(CheckType.FullName, true));
+								UEBuildPlatformFactory TempInst = (UEBuildPlatformFactory)(UBTAssembly.CreateInstance(CheckType.FullName, true));
 								TempInst.TryRegisterBuildPlatforms();
 							}
 						}
 					}
 				}
 				// register all the other classes
-				foreach (var CheckType in AllTypes)
+				foreach (Type CheckType in AllTypes)
 				{
 					if (CheckType.IsClass && !CheckType.IsAbstract)
 					{
@@ -463,7 +488,7 @@ namespace UnrealBuildTool
 
 				// We need to process the ProjectGenerators AFTER all BuildPlatforms have been 
 				// registered as they use the BuildPlatform to verify they are legitimate.
-				foreach (var ProjType in ProjectGeneratorList)
+				foreach (Type ProjType in ProjectGeneratorList)
 				{
 					Log.TraceVerbose("    Registering project generator: {0}", ProjType.ToString());
 					UEPlatformProjectGenerator TempInst = (UEPlatformProjectGenerator)(UBTAssembly.CreateInstance(ProjType.FullName, true));
@@ -622,7 +647,7 @@ namespace UnrealBuildTool
 			// Then let the command lines override any configs necessary.
 			const string UCAUsageThresholdString = "-ucausagethreshold=";
 			const string UCAModuleToAnalyzeString = "-ucamoduletoanalyze=";
-			foreach (var LowercaseArg in Arguments.Select(arg => arg.ToLowerInvariant()))
+			foreach (string LowercaseArg in Arguments.Select(arg => arg.ToLowerInvariant()))
 			{
 				if (LowercaseArg == "-verbose")
 				{
@@ -863,7 +888,7 @@ namespace UnrealBuildTool
 				}
 			}
 
-			// Build the list of game projects that we know about. When building from the editor (for hot-reload) or for Rocket projects, we require the 
+			// Build the list of game projects that we know about. When building from the editor (for hot-reload) or for projects from installed builds, we require the 
 			// project file to be passed in. Otherwise we scan for projects in directories named in UE4Games.uprojectdirs.
 			if (ProjectFile != null)
 			{
@@ -874,7 +899,7 @@ namespace UnrealBuildTool
 				UProjectInfo.FillProjectInfo();
 			}
 
-			var BasicInitStartTime = DateTime.UtcNow;
+			DateTime BasicInitStartTime = DateTime.UtcNow;
 
 			ECompilationResult Result = ECompilationResult.Succeeded;
 
@@ -975,15 +1000,15 @@ namespace UnrealBuildTool
 				try
 				{
 					string GameName = null;
-					var bGenerateVCProjectFiles = false;
-					var bGenerateXcodeProjectFiles = false;
-					var bGenerateMakefiles = false;
-					var bGenerateCMakefiles = false;
-					var bGenerateQMakefiles = false;
-					var bGenerateKDevelopFiles = false;
-					var bGenerateCodeLiteFiles = false;
-					var bValidPlatformsOnly = false;
-					var bSpecificModulesOnly = false;
+					bool bGenerateVCProjectFiles = false;
+					bool bGenerateXcodeProjectFiles = false;
+					bool bGenerateMakefiles = false;
+					bool bGenerateCMakefiles = false;
+					bool bGenerateQMakefiles = false;
+					bool bGenerateKDevelopFiles = false;
+					bool bGenerateCodeLiteFiles = false;
+					bool bValidPlatformsOnly = false;
+					bool bSpecificModulesOnly = false;
 
 					// We need to be able to identify if one of the arguments is the platform...
 					// Leverage the existing parser function in UEBuildTarget to get this information.
@@ -1005,7 +1030,7 @@ namespace UnrealBuildTool
 						bIsGatheringBuild_Unsafe = false;
 					}
 
-					foreach (var Arg in Arguments)
+					foreach (string Arg in Arguments)
 					{
 						string LowercaseArg = Arg.ToLowerInvariant();
 						const string OnlyPlatformSpecificForArg = "-onlyplatformspecificfor=";
@@ -1229,7 +1254,7 @@ namespace UnrealBuildTool
 
 					if (BuildConfiguration.bPrintPerformanceInfo)
 					{
-						var BasicInitTime = (DateTime.UtcNow - BasicInitStartTime).TotalSeconds;
+						double BasicInitTime = (DateTime.UtcNow - BasicInitStartTime).TotalSeconds;
 						Log.TraceInformation("Basic UBT initialization took " + BasicInitTime + "s");
 					}
 
@@ -1308,7 +1333,7 @@ namespace UnrealBuildTool
 							}
 							else
 							{
-								var BuildPlatform = UEBuildPlatform.GetBuildPlatform(CheckPlatform, true);
+								UEBuildPlatform BuildPlatform = UEBuildPlatform.GetBuildPlatform(CheckPlatform, true);
 								if (BuildPlatform != null)
 								{
 									// Setup environment wasn't called, so set the flag
@@ -1322,7 +1347,7 @@ namespace UnrealBuildTool
 							if ((Result == ECompilationResult.Succeeded) && (BuildConfiguration.bDeployAfterCompile == true) && (BuildConfiguration.bXGEExport == false) &&
 								(UEBuildConfiguration.bGenerateManifest == false) && (UEBuildConfiguration.bGenerateExternalFileList == false) && (UEBuildConfiguration.bCleanProject == false))
 							{
-								var TargetDescs = UEBuildTarget.ParseTargetCommandLine(Arguments, ref ProjectFile);
+								List<TargetDescriptor> TargetDescs = UEBuildTarget.ParseTargetCommandLine(Arguments, ref ProjectFile);
 
 								UEBuildPlatform BuildPlatform = UEBuildPlatform.GetBuildPlatform(CheckPlatform);
 								UEBuildPlatformContext PlatformContext = BuildPlatform.CreateContext(TargetDescs[0].ProjectFile);
@@ -1347,7 +1372,7 @@ namespace UnrealBuildTool
 						}
 					}
 					// Print some performance info
-					var BuildDuration = (DateTime.UtcNow - StartTime).TotalSeconds;
+					double BuildDuration = (DateTime.UtcNow - StartTime).TotalSeconds;
 					if (BuildConfiguration.bPrintPerformanceInfo)
 					{
 						Log.TraceInformation("GetIncludes time: " + CPPEnvironment.TotalTimeSpentGettingIncludes + "s (" + CPPEnvironment.TotalIncludesRequested + " includes)");
@@ -1457,7 +1482,7 @@ namespace UnrealBuildTool
 		public static void ValidatePlatforms(string[] Arguments)
 		{
 			List<UnrealTargetPlatform> Platforms = new List<UnrealTargetPlatform>();
-			foreach (var CurArgument in Arguments)
+			foreach (string CurArgument in Arguments)
 			{
 				if (CurArgument.StartsWith("-"))
 				{
@@ -1524,7 +1549,7 @@ namespace UnrealBuildTool
 
 			foreach (UnrealTargetPlatform platform in Platforms)
 			{
-				var BuildPlatform = UEBuildPlatform.GetBuildPlatform(platform, true);
+				UEBuildPlatform BuildPlatform = UEBuildPlatform.GetBuildPlatform(platform, true);
 				if (BuildPlatform != null && BuildPlatform.HasRequiredSDKsInstalled() == SDKStatus.Valid)
 				{
 					Console.WriteLine("##PlatformValidate: {0} VALID", platform.ToString());
@@ -1542,7 +1567,7 @@ namespace UnrealBuildTool
 			bool bSuccess = true;
 
 
-			var RunUBTInitStartTime = DateTime.UtcNow;
+			DateTime RunUBTInitStartTime = DateTime.UtcNow;
 
 
 			// Reset global configurations
@@ -1552,8 +1577,8 @@ namespace UnrealBuildTool
 			UnrealTargetPlatform ResetPlatform = UnrealTargetPlatform.Unknown;
 			UnrealTargetConfiguration ResetConfiguration;
 			UEBuildTarget.ParsePlatformAndConfiguration(Arguments, out ResetPlatform, out ResetConfiguration);
-			var BuildPlatform = UEBuildPlatform.GetBuildPlatform(ResetPlatform);
-			var BuildPlatformContext = BuildPlatform.CreateContext(ProjectFile);
+			UEBuildPlatform BuildPlatform = UEBuildPlatform.GetBuildPlatform(ResetPlatform);
+			UEBuildPlatformContext BuildPlatformContext = BuildPlatform.CreateContext(ProjectFile);
 			BuildPlatformContext.ResetBuildConfiguration(ResetConfiguration);
 
 			// now that we have the platform, we can set the intermediate path to include the platform/architecture name
@@ -1562,7 +1587,7 @@ namespace UnrealBuildTool
 			string ExecutorName = "Unknown";
 			ECompilationResult BuildResult = ECompilationResult.Succeeded;
 
-			var ToolChain = BuildPlatformContext.CreateToolChainForDefaultCppPlatform();
+			UEToolChain ToolChain = BuildPlatformContext.CreateToolChainForDefaultCppPlatform();
 
 			string EULAViolationWarning = null;
 			Thread CPPIncludesThread = null;
@@ -1588,13 +1613,13 @@ namespace UnrealBuildTool
 
 				if (BuildConfiguration.bPrintPerformanceInfo)
 				{
-					var RunUBTInitTime = (DateTime.UtcNow - RunUBTInitStartTime).TotalSeconds;
+					double RunUBTInitTime = (DateTime.UtcNow - RunUBTInitStartTime).TotalSeconds;
 					Log.TraceInformation("RunUBT initialization took " + RunUBTInitTime + "s");
 				}
 
-				var TargetDescs = new List<TargetDescriptor>();
+				List<TargetDescriptor> TargetDescs = new List<TargetDescriptor>();
 				{
-					var TargetDescstStartTime = DateTime.UtcNow;
+					DateTime TargetDescstStartTime = DateTime.UtcNow;
 
 					foreach (string[] TargetSetting in TargetSettings)
 					{
@@ -1603,7 +1628,7 @@ namespace UnrealBuildTool
 
 					if (BuildConfiguration.bPrintPerformanceInfo)
 					{
-						var TargetDescsTime = (DateTime.UtcNow - TargetDescstStartTime).TotalSeconds;
+						double TargetDescsTime = (DateTime.UtcNow - TargetDescstStartTime).TotalSeconds;
 						Log.TraceInformation("Target descriptors took " + TargetDescsTime + "s");
 					}
 				}
@@ -1771,12 +1796,12 @@ namespace UnrealBuildTool
 				}
 				else
 				{
-					var TargetInitStartTime = DateTime.UtcNow;
+					DateTime TargetInitStartTime = DateTime.UtcNow;
 
 					Targets = new List<UEBuildTarget>();
-					foreach (var TargetDesc in TargetDescs)
+					foreach (TargetDescriptor TargetDesc in TargetDescs)
 					{
-						var Target = UEBuildTarget.CreateTarget(TargetDesc);
+						UEBuildTarget Target = UEBuildTarget.CreateTarget(TargetDesc);
 						if ((Target == null) && (UEBuildConfiguration.bCleanProject))
 						{
 							continue;
@@ -1786,17 +1811,17 @@ namespace UnrealBuildTool
 
 					if (BuildConfiguration.bPrintPerformanceInfo)
 					{
-						var TargetInitTime = (DateTime.UtcNow - TargetInitStartTime).TotalSeconds;
+						double TargetInitTime = (DateTime.UtcNow - TargetInitStartTime).TotalSeconds;
 						Log.TraceInformation("Target init took " + TargetInitTime + "s");
 					}
 				}
 
 				// Build action lists for all passed in targets.
-				var OutputItemsForAllTargets = new List<FileItem>();
-				var TargetNameToUObjectModules = new Dictionary<string, List<UHTModuleInfo>>(StringComparer.InvariantCultureIgnoreCase);
-				foreach (var Target in Targets)
+				List<FileItem> OutputItemsForAllTargets = new List<FileItem>();
+				Dictionary<string, List<UHTModuleInfo>> TargetNameToUObjectModules = new Dictionary<string, List<UHTModuleInfo>>(StringComparer.InvariantCultureIgnoreCase);
+				foreach (UEBuildTarget Target in Targets)
 				{
-					var TargetStartTime = DateTime.UtcNow;
+					DateTime TargetStartTime = DateTime.UtcNow;
 
 					if (bIsHotReload)
 					{
@@ -1849,14 +1874,14 @@ namespace UnrealBuildTool
 									throw new BuildException("ERROR: The '-graph' option only works with a single target at a time");
 								}
 								ActionGraph.FinalizeActionGraph();
-								var ActionsToExecute = ActionGraph.AllActions;
+								List<Action> ActionsToExecute = ActionGraph.AllActions;
 
-								var VisualizationType = ActionGraph.ActionGraphVisualizationType.OnlyCPlusPlusFilesAndHeaders;
+								ActionGraph.ActionGraphVisualizationType VisualizationType = ActionGraph.ActionGraphVisualizationType.OnlyCPlusPlusFilesAndHeaders;
 								ActionGraph.SaveActionGraphVisualization(Target, Path.Combine(BuildConfiguration.BaseIntermediatePath, Target.GetTargetName() + ".gexf"), Target.GetTargetName(), VisualizationType, ActionsToExecute);
 							}
 						}
 
-						var TargetBuildTime = (DateTime.UtcNow - TargetStartTime).TotalSeconds;
+						double TargetBuildTime = (DateTime.UtcNow - TargetStartTime).TotalSeconds;
 
 						// Send out telemetry for this target
 						Telemetry.SendEvent("TargetBuildStats.2",
@@ -1888,8 +1913,8 @@ namespace UnrealBuildTool
 
 						// For now simply treat all object files as the root target.
 						{
-							var PrerequisiteActionsSet = new HashSet<Action>();
-							foreach (var OutputItem in OutputItemsForAllTargets)
+							HashSet<Action> PrerequisiteActionsSet = new HashSet<Action>();
+							foreach (FileItem OutputItem in OutputItemsForAllTargets)
 							{
 								ActionGraph.GatherPrerequisiteActions(OutputItem, ref PrerequisiteActionsSet);
 							}
@@ -1929,7 +1954,7 @@ namespace UnrealBuildTool
 							}
 
 
-							foreach (var EnvironmentVariable in UBTMakefile.EnvironmentVariables)
+							foreach (Tuple<string, string> EnvironmentVariable in UBTMakefile.EnvironmentVariables)
 							{
 								// @todo ubtmake: There may be some variables we do NOT want to clobber.
 								Environment.SetEnvironmentVariable(EnvironmentVariable.Item1, EnvironmentVariable.Item2);
@@ -1948,7 +1973,7 @@ namespace UnrealBuildTool
 							// If any of the targets need UHT to be run, we'll go ahead and do that now
 							if(BuildResult.Succeeded())
 							{
-								foreach (var Target in Targets)
+								foreach (UEBuildTarget Target in Targets)
 								{
 									List<UHTModuleInfo> TargetUObjectModules;
 									if (UBTMakefile.TargetNameToUObjectModules.TryGetValue(Target.GetTargetName(), out TargetUObjectModules))
@@ -2077,13 +2102,13 @@ namespace UnrealBuildTool
 				//    are successfully built.  If we didn't save our dependency cache after build failures, source files for those build products that were successsfully
 				//    built before the failure would not be considered out of date on the next run, so this is our only chance to cache C++ includes for those files!
 
-				foreach (var DependencyCache in CPPEnvironment.IncludeDependencyCache.Values)
+				foreach (DependencyCache DependencyCache in CPPEnvironment.IncludeDependencyCache.Values)
 				{
 					DependencyCache.Save();
 				}
 				CPPEnvironment.IncludeDependencyCache.Clear();
 
-				foreach (var FlatCPPIncludeDependencyCache in CPPEnvironment.FlatCPPIncludeDependencyCache.Values)
+				foreach (FlatCPPIncludeDependencyCache FlatCPPIncludeDependencyCache in CPPEnvironment.FlatCPPIncludeDependencyCache.Values)
 				{
 					FlatCPPIncludeDependencyCache.Save();
 				}
@@ -2152,12 +2177,12 @@ namespace UnrealBuildTool
 		/// <param name="Target">Target</param>
 		private static void InvalidateMakefiles(TargetDescriptor Target)
 		{
-			var MakefileNames = new string[] { "HotReloadMakefile.ubt", "Makefile.ubt" };
-			var BaseDir = GetUBTMakefileDirectoryPathForSingleTarget(Target);
+			string[] MakefileNames = new string[] { "HotReloadMakefile.ubt", "Makefile.ubt" };
+			DirectoryReference BaseDir = GetUBTMakefileDirectoryPathForSingleTarget(Target);
 
-			foreach (var MakefileName in MakefileNames)
+			foreach (string MakefileName in MakefileNames)
 			{
-				var MakefileRef = FileReference.Combine(BaseDir, MakefileName);
+				FileReference MakefileRef = FileReference.Combine(BaseDir, MakefileName);
 				if (MakefileRef.Exists())
 				{
 					MakefileRef.Delete();
@@ -2187,23 +2212,23 @@ namespace UnrealBuildTool
 
 			if (!ProjectFileGenerator.bGenerateProjectFiles && !UEBuildConfiguration.bGenerateManifest && bIsEditorTarget)
 			{
-				var EditorProcessFilenames = UEBuildTarget.MakeExecutablePaths(UnrealBuildTool.EngineDirectory, "UE4Editor", TargetDesc.Platform, TargetDesc.Configuration, "", UnrealTargetConfiguration.Development, false, null, null);
+				List<FileReference> EditorProcessFilenames = UEBuildTarget.MakeExecutablePaths(UnrealBuildTool.EngineDirectory, "UE4Editor", TargetDesc.Platform, TargetDesc.Configuration, "", UnrealTargetConfiguration.Development, false, null, null);
 				if (EditorProcessFilenames.Count != 1)
 				{
 					throw new BuildException("ShouldDoHotReload cannot handle multiple binaries returning from UEBuildTarget.MakeExecutablePaths");
 				}
 
-				var Processes = BuildHostPlatform.Current.GetProcesses();
-				var EditorRunsDir = Path.Combine(UnrealBuildTool.EngineDirectory.FullName, "Intermediate", "EditorRuns");
+				BuildHostPlatform.ProcessInfo[] Processes = BuildHostPlatform.Current.GetProcesses();
+				string EditorRunsDir = Path.Combine(UnrealBuildTool.EngineDirectory.FullName, "Intermediate", "EditorRuns");
 
 				if (!Directory.Exists(EditorRunsDir))
 				{
 					return false;
 				}
 
-				var EditorRunsFiles = new DirectoryInfo(EditorRunsDir).GetFiles();
+				FileInfo[] EditorRunsFiles = new DirectoryInfo(EditorRunsDir).GetFiles();
 
-				foreach(var File in EditorRunsFiles)
+				foreach(FileInfo File in EditorRunsFiles)
 				{
 					int PID;
 					BuildHostPlatform.ProcessInfo Proc = null;
@@ -2232,7 +2257,7 @@ namespace UnrealBuildTool
 		/// <returns>List of build target settings</returns>
 		private static List<string[]> ParseCommandLineFlags(string[] Arguments)
 		{
-			var TargetSettings = new List<string[]>();
+			List<string[]> TargetSettings = new List<string[]>();
 			int ArgumentIndex = 0;
 
 			if (Utils.ParseCommandLineFlag(Arguments, "-targets", out ArgumentIndex))
@@ -2242,7 +2267,7 @@ namespace UnrealBuildTool
 					throw new BuildException("Expected filename after -targets argument, but found nothing.");
 				}
 				// Parse lines from the referenced file into target settings.
-				var Lines = File.ReadAllLines(Arguments[ArgumentIndex + 1]);
+				string[] Lines = File.ReadAllLines(Arguments[ArgumentIndex + 1]);
 				foreach (string Line in Lines)
 				{
 					if (Line != "" && Line[0] != ';')
@@ -2273,16 +2298,16 @@ namespace UnrealBuildTool
 					// @todo ubtmake: This thread will access data structures that are also used on the main UBT thread, but during this time UBT
 					// is only invoking the build executor, so should not be touching this stuff.  However, we need to at some guards to make sure.
 
-					foreach (var TargetAndOutdatedPrerequisites in TargetToOutdatedPrerequisitesMap)
+					foreach (KeyValuePair<UEBuildTarget, List<FileItem>> TargetAndOutdatedPrerequisites in TargetToOutdatedPrerequisitesMap)
 					{
-						var Target = TargetAndOutdatedPrerequisites.Key;
-						var OutdatedPrerequisites = TargetAndOutdatedPrerequisites.Value;
+						UEBuildTarget Target = TargetAndOutdatedPrerequisites.Key;
+						List<FileItem> OutdatedPrerequisites = TargetAndOutdatedPrerequisites.Value;
 
-						foreach (var PrerequisiteItem in OutdatedPrerequisites)
+						foreach (FileItem PrerequisiteItem in OutdatedPrerequisites)
 						{
 							// Invoke our deep include scanner to figure out whether any of the files included by this source file have
 							// changed since the build product was built
-							var FileBuildPlatform = UEBuildPlatform.GetBuildPlatform(Target.Platform);
+							UEBuildPlatform FileBuildPlatform = UEBuildPlatform.GetBuildPlatform(Target.Platform);
 							CPPEnvironment.FindAndCacheAllIncludedFiles(Target, PrerequisiteItem, FileBuildPlatform, PrerequisiteItem.CachedCPPIncludeInfo, bOnlyCachedDependencies: false);
 						}
 					}
@@ -2392,9 +2417,9 @@ namespace UnrealBuildTool
 				throw new BuildException("Can't save a makefile that has invalid contents.  See UBTMakefile.IsValidMakefile()");
 			}
 
-			var TimerStartTime = DateTime.UtcNow;
+			DateTime TimerStartTime = DateTime.UtcNow;
 
-			var UBTMakefileItem = FileItem.GetItemByFileReference(GetUBTMakefilePath(TargetDescs));
+			FileItem UBTMakefileItem = FileItem.GetItemByFileReference(GetUBTMakefilePath(TargetDescs));
 
 			// @todo ubtmake: Optimization: The UBTMakefile saved for game projects is upwards of 9 MB.  We should try to shrink its content if possible
 			// @todo ubtmake: Optimization: C# Serialization may be too slow for these big Makefiles.  Loading these files often shows up as the slower part of the assembling phase.
@@ -2416,7 +2441,7 @@ namespace UnrealBuildTool
 
 			if (BuildConfiguration.bPrintPerformanceInfo)
 			{
-				var TimerDuration = DateTime.UtcNow - TimerStartTime;
+				TimeSpan TimerDuration = DateTime.UtcNow - TimerStartTime;
 				Log.TraceInformation("Saving makefile took " + TimerDuration.TotalSeconds + "s");
 			}
 		}
@@ -2432,7 +2457,7 @@ namespace UnrealBuildTool
 		{
 			// Check the directory timestamp on the project files directory.  If the user has generated project files more
 			// recently than the UBTMakefile, then we need to consider the file to be out of date
-			var UBTMakefileInfo = new FileInfo(MakefilePath.FullName);
+			FileInfo UBTMakefileInfo = new FileInfo(MakefilePath.FullName);
 			if (!UBTMakefileInfo.Exists)
 			{
 				// UBTMakefile doesn't even exist, so we won't bother loading it
@@ -2449,7 +2474,7 @@ namespace UnrealBuildTool
 			{
 				if (ProjectFileGenerator.IntermediateProjectFilesPath.Exists())
 				{
-					var EngineProjectFilesLastUpdateTime = new FileInfo(ProjectFileGenerator.ProjectTimestampFile).LastWriteTime;
+					DateTime EngineProjectFilesLastUpdateTime = new FileInfo(ProjectFileGenerator.ProjectTimestampFile).LastWriteTime;
 					if (UBTMakefileInfo.LastWriteTime.CompareTo(EngineProjectFilesLastUpdateTime) < 0)
 					{
 						// Engine project files are newer than UBTMakefile
@@ -2464,7 +2489,7 @@ namespace UnrealBuildTool
 			if (ProjectFile != null)
 			{
 				string ProjectFilename = ProjectFile.FullName;
-				var ProjectFileInfo = new FileInfo(ProjectFilename);
+				FileInfo ProjectFileInfo = new FileInfo(ProjectFilename);
 				if (!ProjectFileInfo.Exists || UBTMakefileInfo.LastWriteTime.CompareTo(ProjectFileInfo.LastWriteTime) < 0)
 				{
 					// .uproject file is newer than UBTMakefile
@@ -2473,11 +2498,11 @@ namespace UnrealBuildTool
 					return null;
 				}
 
-				var MasterProjectRelativePath = ProjectFile.Directory;
-				var GameIntermediateProjectFilesPath = Path.Combine(MasterProjectRelativePath.FullName, "Intermediate", "ProjectFiles");
+				DirectoryReference MasterProjectRelativePath = ProjectFile.Directory;
+				string GameIntermediateProjectFilesPath = Path.Combine(MasterProjectRelativePath.FullName, "Intermediate", "ProjectFiles");
 				if (Directory.Exists(GameIntermediateProjectFilesPath))
 				{
-					var GameProjectFilesLastUpdateTime = new DirectoryInfo(GameIntermediateProjectFilesPath).LastWriteTime;
+					DateTime GameProjectFilesLastUpdateTime = new DirectoryInfo(GameIntermediateProjectFilesPath).LastWriteTime;
 					if (UBTMakefileInfo.LastWriteTime.CompareTo(GameProjectFilesLastUpdateTime) < 0)
 					{
 						// Game project files are newer than UBTMakefile
@@ -2489,7 +2514,7 @@ namespace UnrealBuildTool
 			}
 
 			// Check to see if UnrealBuildTool.exe was compiled more recently than the UBTMakefile
-			var UnrealBuildToolTimestamp = new FileInfo(Assembly.GetExecutingAssembly().Location).LastWriteTime;
+			DateTime UnrealBuildToolTimestamp = new FileInfo(Assembly.GetExecutingAssembly().Location).LastWriteTime;
 			if (UBTMakefileInfo.LastWriteTime.CompareTo(UnrealBuildToolTimestamp) < 0)
 			{
 				// UnrealBuildTool.exe was compiled more recently than the UBTMakefile
@@ -2510,7 +2535,7 @@ namespace UnrealBuildTool
 
 			try
 			{
-				var LoadUBTMakefileStartTime = DateTime.UtcNow;
+				DateTime LoadUBTMakefileStartTime = DateTime.UtcNow;
 
 				using (FileStream Stream = new FileStream(UBTMakefileInfo.FullName, FileMode.Open, FileAccess.Read))
 				{
@@ -2520,7 +2545,7 @@ namespace UnrealBuildTool
 
 				if (BuildConfiguration.bPrintPerformanceInfo)
 				{
-					var LoadUBTMakefileTime = (DateTime.UtcNow - LoadUBTMakefileStartTime).TotalSeconds;
+					double LoadUBTMakefileTime = (DateTime.UtcNow - LoadUBTMakefileStartTime).TotalSeconds;
 					Log.TraceInformation("LoadUBTMakefile took " + LoadUBTMakefileTime + "s");
 				}
 			}
@@ -2539,12 +2564,12 @@ namespace UnrealBuildTool
 			}
 
 			// Check if any of the target's Build.cs files are newer than the makefile
-			foreach (var Target in LoadedUBTMakefile.Targets)
+			foreach (UEBuildTarget Target in LoadedUBTMakefile.Targets)
 			{
 				string TargetCsFilename = Target.TargetCsFilename.FullName;
 				if (TargetCsFilename != null)
 				{
-					var TargetCsFile = new FileInfo(TargetCsFilename);
+					FileInfo TargetCsFile = new FileInfo(TargetCsFilename);
 					bool bTargetCsFileExists = TargetCsFile.Exists;
 					if (!bTargetCsFileExists || TargetCsFile.LastWriteTime > UBTMakefileInfo.LastWriteTime)
 					{
@@ -2555,11 +2580,11 @@ namespace UnrealBuildTool
 				}
 
 				IEnumerable<string> BuildCsFilenames = Target.GetAllModuleBuildCsFilenames();
-				foreach (var BuildCsFilename in BuildCsFilenames)
+				foreach (string BuildCsFilename in BuildCsFilenames)
 				{
 					if (BuildCsFilename != null)
 					{
-						var BuildCsFile = new FileInfo(BuildCsFilename);
+						FileInfo BuildCsFile = new FileInfo(BuildCsFilename);
 						bool bBuildCsFileExists = BuildCsFile.Exists;
 						if (!bBuildCsFileExists || BuildCsFile.LastWriteTime > UBTMakefileInfo.LastWriteTime)
 						{
@@ -2574,10 +2599,10 @@ namespace UnrealBuildTool
 			// We do a check to see if any modules' headers have changed which have
 			// acquired or lost UHT types.  If so, which should be rare,
 			// we'll just invalidate the entire makefile and force it to be rebuilt.
-			foreach (var Target in LoadedUBTMakefile.Targets)
+			foreach (UEBuildTarget Target in LoadedUBTMakefile.Targets)
 			{
 				// Get all H files in processed modules newer than the makefile itself
-				var HFilesNewerThanMakefile =
+				HashSet<string> HFilesNewerThanMakefile =
 					new HashSet<string>(
 						Target.FlatModuleCsData
 						.SelectMany(x => x.Value.ModuleSourceFolder != null ? Directory.EnumerateFiles(x.Value.ModuleSourceFolder.FullName, "*.h", SearchOption.AllDirectories) : Enumerable.Empty<string>())
@@ -2586,7 +2611,7 @@ namespace UnrealBuildTool
 					);
 
 				// Get all H files in all modules processed in the last makefile build
-				var AllUHTHeaders = new HashSet<string>(Target.FlatModuleCsData.Select(x => x.Value).SelectMany(x => x.UHTHeaderNames));
+				HashSet<string> AllUHTHeaders = new HashSet<string>(Target.FlatModuleCsData.Select(x => x.Value).SelectMany(x => x.UHTHeaderNames));
 
 				// Makefile is invalid if:
 				// * There are any newer files which contain no UHT data, but were previously in the makefile
@@ -2612,9 +2637,9 @@ namespace UnrealBuildTool
 			if (BuildConfiguration.bUseAdaptiveUnityBuild)
 			{
 				// Get all .cpp files in processed modules newer than the makefile itself
-				foreach (var Target in LoadedUBTMakefile.Targets)
+				foreach (UEBuildTarget Target in LoadedUBTMakefile.Targets)
 				{
-					foreach (var FlatModule in Target.FlatModuleCsData.Values)
+					foreach (FlatModuleCsDataType FlatModule in Target.FlatModuleCsData.Values)
 					{
 						if (FlatModule != null && FlatModule.ModuleSourceFolder != null)
 						{
@@ -2622,7 +2647,7 @@ namespace UnrealBuildTool
 							// Makefile was created. 
 							//if( Directory.GetLastWriteTimeUtc( FlatModule.ModuleSourceFolder ) > UBTMakefileInfo.LastWriteTimeUtc )
 							{
-								foreach (var SourceFilePath in FlatModule.ModuleSourceFolder.EnumerateFileReferences("*.cpp", SearchOption.AllDirectories))
+								foreach (FileReference SourceFilePath in FlatModule.ModuleSourceFolder.EnumerateFileReferences("*.cpp", SearchOption.AllDirectories))
 								{
 									// Was the file written to since the Makefile was created?  We're not able to do a fully exhaustive
 									// check for working sets when in assembler mode, so we might get it wrong occasionally.  This is
@@ -2662,7 +2687,7 @@ namespace UnrealBuildTool
 
 			if (TargetDescs.Count == 1)
 			{
-				var TargetDesc = TargetDescs[0];
+				TargetDescriptor TargetDesc = TargetDescs[0];
 
 				bool bIsHotReload = (UEBuildConfiguration.bHotReloadFromIDE || (TargetDesc.OnlyModules != null && TargetDesc.OnlyModules.Count > 0));
 				string UBTMakefileName = bIsHotReload ? "HotReloadMakefile.ubt" : "Makefile.ubt";
@@ -2673,7 +2698,7 @@ namespace UnrealBuildTool
 			{
 				// For Makefiles that contain multiple targets, we'll make up a file name that contains all of the targets, their
 				// configurations and platforms, and save it into the base intermediate folder
-				var TargetCollectionName = MakeTargetCollectionName(TargetDescs);
+				string TargetCollectionName = MakeTargetCollectionName(TargetDescs);
 
 				TargetDescriptor DescriptorWithProject = TargetDescs.FirstOrDefault(x => x.ProjectFile != null);
 
@@ -2738,13 +2763,13 @@ namespace UnrealBuildTool
 				throw new BuildException("Expecting at least one Target to be passed to MakeTargetCollectionName");
 			}
 
-			var SortedTargets = new List<TargetDescriptor>();
+			List<TargetDescriptor> SortedTargets = new List<TargetDescriptor>();
 			SortedTargets.AddRange(TargetDescs);
 			SortedTargets.Sort((x, y) => { return x.TargetName.CompareTo(y.TargetName); });
 
 			// Figure out what to call our action graph based on everything we're building
-			var TargetCollectionName = new StringBuilder();
-			foreach (var Target in SortedTargets)
+			StringBuilder TargetCollectionName = new StringBuilder();
+			foreach (TargetDescriptor Target in SortedTargets)
 			{
 				if (TargetCollectionName.Length > 0)
 				{
@@ -2787,24 +2812,24 @@ namespace UnrealBuildTool
 		{
 			// Gather all of the response files for link actions.  We're going to need to patch 'em up after we figure out new
 			// names for all of the output files and import libraries
-			var ResponseFilePaths = new List<string>();
+			List<string> ResponseFilePaths = new List<string>();
 
 			// Keep a map of the original file names and their new file names, so we can fix up response files after
-			var OriginalFileNameAndNewFileNameList_NoExtensions = new List<Tuple<string, string>>();
+			List<Tuple<string, string>> OriginalFileNameAndNewFileNameList_NoExtensions = new List<Tuple<string, string>>();
 
 			// Finally, we'll keep track of any file items that we had to create counterparts for change file names, so we can fix those up too
-			var AffectedOriginalFileItemAndNewFileItemMap = new Dictionary<FileItem, FileItem>();
+			Dictionary<FileItem, FileItem> AffectedOriginalFileItemAndNewFileItemMap = new Dictionary<FileItem, FileItem>();
 
-			foreach (var Action in ActionGraph.AllActions)
+			foreach (Action Action in ActionGraph.AllActions)
 			{
 				if (Action.ActionType == ActionType.Link)
 				{
 					// Assume that the first produced item (with no extension) is our output file name
-					var OriginalFileNameWithoutExtension = Utils.GetFilenameWithoutAnyExtensions(Action.ProducedItems[0].AbsolutePath);
+					string OriginalFileNameWithoutExtension = Utils.GetFilenameWithoutAnyExtensions(Action.ProducedItems[0].AbsolutePath);
 
 					// Remove the numbered suffix from the end of the file
-					var IndexOfLastHyphen = OriginalFileNameWithoutExtension.LastIndexOf('-');
-					var OriginalNumberSuffix = OriginalFileNameWithoutExtension.Substring(IndexOfLastHyphen);
+					int IndexOfLastHyphen = OriginalFileNameWithoutExtension.LastIndexOf('-');
+					string OriginalNumberSuffix = OriginalFileNameWithoutExtension.Substring(IndexOfLastHyphen);
 					int NumberSuffix;
 					bool bHasNumberSuffix = int.TryParse(OriginalNumberSuffix, out NumberSuffix);
 
@@ -2818,8 +2843,8 @@ namespace UnrealBuildTool
 					else
 					{
 						// Remove "-####-Platform-Configuration" suffix in Debug configuration
-						var IndexOfSecondLastHyphen = OriginalFileNameWithoutExtension.LastIndexOf('-', IndexOfLastHyphen - 1, IndexOfLastHyphen);
-						var IndexOfThirdLastHyphen = OriginalFileNameWithoutExtension.LastIndexOf('-', IndexOfSecondLastHyphen - 1, IndexOfSecondLastHyphen);
+						int IndexOfSecondLastHyphen = OriginalFileNameWithoutExtension.LastIndexOf('-', IndexOfLastHyphen - 1, IndexOfLastHyphen);
+						int IndexOfThirdLastHyphen = OriginalFileNameWithoutExtension.LastIndexOf('-', IndexOfSecondLastHyphen - 1, IndexOfSecondLastHyphen);
 
 						if (!int.TryParse(OriginalFileNameWithoutExtension.Substring(IndexOfThirdLastHyphen + 1, IndexOfSecondLastHyphen - IndexOfThirdLastHyphen - 1), out NumberSuffix))
 						{
@@ -2846,7 +2871,7 @@ namespace UnrealBuildTool
 						// Add a new random suffix
 						if (OnlyModules != null)
 						{
-							foreach (var OnlyModule in OnlyModules)
+							foreach (OnlyModule OnlyModule in OnlyModules)
 							{
 								if (OnlyModule.OnlyModuleName.Equals(ModuleName, StringComparison.InvariantCultureIgnoreCase))
 								{
@@ -2861,24 +2886,24 @@ namespace UnrealBuildTool
 							UniqueSuffix = "-" + (new Random((int)(DateTime.Now.Ticks % Int32.MaxValue)).Next(10000)).ToString();
 						}
 					}
-					var NewFileNameWithoutExtension = OriginalFileNameWithoutNumberSuffix + UniqueSuffix + PlatformConfigSuffix;
+					string NewFileNameWithoutExtension = OriginalFileNameWithoutNumberSuffix + UniqueSuffix + PlatformConfigSuffix;
 
 					// Find the response file in the command line.  We'll need to make a copy of it with our new file name.
 					{
 						string ResponseFileExtension = ".response";
-						var ResponseExtensionIndex = Action.CommandArguments.IndexOf(ResponseFileExtension, StringComparison.InvariantCultureIgnoreCase);
+						int ResponseExtensionIndex = Action.CommandArguments.IndexOf(ResponseFileExtension, StringComparison.InvariantCultureIgnoreCase);
 						if (ResponseExtensionIndex != -1)
 						{
-							var ResponseFilePathIndex = Action.CommandArguments.LastIndexOf("@\"", ResponseExtensionIndex);
+							int ResponseFilePathIndex = Action.CommandArguments.LastIndexOf("@\"", ResponseExtensionIndex);
 							if (ResponseFilePathIndex == -1)
 							{
 								throw new BuildException("Couldn't find response file path in action's command arguments when hot reloading");
 							}
 
-							var OriginalResponseFilePathWithoutExtension = Action.CommandArguments.Substring(ResponseFilePathIndex + 2, (ResponseExtensionIndex - ResponseFilePathIndex) - 2);
-							var OriginalResponseFilePath = OriginalResponseFilePathWithoutExtension + ResponseFileExtension;
+							string OriginalResponseFilePathWithoutExtension = Action.CommandArguments.Substring(ResponseFilePathIndex + 2, (ResponseExtensionIndex - ResponseFilePathIndex) - 2);
+							string OriginalResponseFilePath = OriginalResponseFilePathWithoutExtension + ResponseFileExtension;
 
-							var NewResponseFilePath = OriginalResponseFilePath.Replace(OriginalFileNameWithoutExtension, NewFileNameWithoutExtension);
+							string NewResponseFilePath = OriginalResponseFilePath.Replace(OriginalFileNameWithoutExtension, NewFileNameWithoutExtension);
 
 							// Copy the old response file to the new path
 							File.Copy(OriginalResponseFilePath, NewResponseFilePath, overwrite: true);
@@ -2897,15 +2922,15 @@ namespace UnrealBuildTool
 					{
 						for (int ItemIndex = 0; ItemIndex < Action.ProducedItems.Count; ++ItemIndex)
 						{
-							var OriginalProducedItem = Action.ProducedItems[ItemIndex];
+							FileItem OriginalProducedItem = Action.ProducedItems[ItemIndex];
 
-							var OriginalProducedItemFilePath = OriginalProducedItem.AbsolutePath;
-							var NewProducedItemFilePath = OriginalProducedItemFilePath.Replace(OriginalFileNameWithoutExtension, NewFileNameWithoutExtension);
+							string OriginalProducedItemFilePath = OriginalProducedItem.AbsolutePath;
+							string NewProducedItemFilePath = OriginalProducedItemFilePath.Replace(OriginalFileNameWithoutExtension, NewFileNameWithoutExtension);
 
 							if (OriginalProducedItemFilePath != NewProducedItemFilePath)
 							{
 								// OK, the produced item's file name changed so we'll update it to point to our new file
-								var NewProducedItem = FileItem.GetItemByPath(NewProducedItemFilePath);
+								FileItem NewProducedItem = FileItem.GetItemByPath(NewProducedItemFilePath);
 								Action.ProducedItems[ItemIndex] = NewProducedItem;
 
 								// Copy the other important settings from the original file item
@@ -2930,11 +2955,11 @@ namespace UnrealBuildTool
 
 
 			// Do another pass and update any actions that depended on the original file names that we changed
-			foreach (var Action in ActionGraph.AllActions)
+			foreach (Action Action in ActionGraph.AllActions)
 			{
 				for (int ItemIndex = 0; ItemIndex < Action.PrerequisiteItems.Count; ++ItemIndex)
 				{
-					var OriginalFileItem = Action.PrerequisiteItems[ItemIndex];
+					FileItem OriginalFileItem = Action.PrerequisiteItems[ItemIndex];
 
 					FileItem NewFileItem;
 					if (AffectedOriginalFileItemAndNewFileItemMap.TryGetValue(OriginalFileItem, out NewFileItem))
@@ -2948,22 +2973,22 @@ namespace UnrealBuildTool
 
 			if (OriginalFileNameAndNewFileNameList_NoExtensions.Count > 0)
 			{
-				foreach (var ResponseFilePath in ResponseFilePaths)
+				foreach (string ResponseFilePath in ResponseFilePaths)
 				{
 					// Load the file up
 					StringBuilder FileContents = new StringBuilder(Utils.ReadAllText(ResponseFilePath));
 
 					// Replace all of the old file names with new ones
-					foreach (var FileNameTuple in OriginalFileNameAndNewFileNameList_NoExtensions)
+					foreach (Tuple<string, string> FileNameTuple in OriginalFileNameAndNewFileNameList_NoExtensions)
 					{
-						var OriginalFileNameWithoutExtension = FileNameTuple.Item1;
-						var NewFileNameWithoutExtension = FileNameTuple.Item2;
+						string OriginalFileNameWithoutExtension = FileNameTuple.Item1;
+						string NewFileNameWithoutExtension = FileNameTuple.Item2;
 
 						FileContents.Replace(OriginalFileNameWithoutExtension, NewFileNameWithoutExtension);
 					}
 
 					// Overwrite the original file
-					var FileContentsString = FileContents.ToString();
+					string FileContentsString = FileContents.ToString();
 					File.WriteAllText(ResponseFilePath, FileContentsString, new System.Text.UTF8Encoding(false));
 				}
 			}
diff --git a/Engine/Source/Programs/UnrealBuildTool/System/Utils.cs b/Engine/Source/Programs/UnrealBuildTool/System/Utils.cs
--- a/Engine/Source/Programs/UnrealBuildTool/System/Utils.cs
+++ b/Engine/Source/Programs/UnrealBuildTool/System/Utils.cs
@@ -11,6 +11,7 @@ using System.Xml.Serialization;
 using System.Runtime.CompilerServices;
 using System.Runtime.InteropServices;
 using System.Linq;
+using System.Management;
 using Tools.DotNETCommon.CaselessDictionary;
 using Tools.DotNETCommon.HarvestEnvVars;
 using System.Web.Script.Serialization;
@@ -307,7 +308,7 @@ namespace UnrealBuildTool
 				throw new BuildException(Ex, "Failed to harvest environment variables");
 			}
 
-			foreach (var EnvVar in EnvVars)
+			foreach (KeyValuePair<string, string> EnvVar in EnvVars)
 			{
 				Environment.SetEnvironmentVariable(EnvVar.Key, EnvVar.Value);
 			}
@@ -367,13 +368,13 @@ namespace UnrealBuildTool
 		/// <param name="Args">Arguments to Command</param>
 		public static string RunLocalProcessAndReturnStdOut(string Command, string Args)
 		{
-			var StartInfo = new ProcessStartInfo(Command, Args);
+			ProcessStartInfo StartInfo = new ProcessStartInfo(Command, Args);
 			StartInfo.UseShellExecute = false;
 			StartInfo.RedirectStandardOutput = true;
 			StartInfo.CreateNoWindow = true;
 
 			string FullOutput = "";
-			using (var LocalProcess = Process.Start(StartInfo))
+			using (Process LocalProcess = Process.Start(StartInfo))
 			{
 				StreamReader OutputReader = LocalProcess.StandardOutput;
 				// trim off any extraneous new lines, helpful for those one-line outputs
@@ -392,7 +393,7 @@ namespace UnrealBuildTool
 		{
 			// Make a list of all platform name strings that we're *not* currently compiling, to speed
 			// up file path comparisons later on
-			var OtherPlatformNameStrings = new List<string>();
+			List<string> OtherPlatformNameStrings = new List<string>();
 			{
 				// look at each group to see if any supported platforms are in it
 				List<UnrealPlatformGroup> SupportedGroups = new List<UnrealPlatformGroup>();
@@ -525,20 +526,20 @@ namespace UnrealBuildTool
 				RelativeToDirectory = ".";
 			}
 
-			var AbsolutePath = SourcePath;
+			string AbsolutePath = SourcePath;
 			if (!Path.IsPathRooted(AbsolutePath))
 			{
 				AbsolutePath = Path.GetFullPath(SourcePath);
 			}
-			var SourcePathEndsWithDirectorySeparator = AbsolutePath.EndsWith(Path.DirectorySeparatorChar.ToString()) || AbsolutePath.EndsWith(Path.AltDirectorySeparatorChar.ToString());
+			bool SourcePathEndsWithDirectorySeparator = AbsolutePath.EndsWith(Path.DirectorySeparatorChar.ToString()) || AbsolutePath.EndsWith(Path.AltDirectorySeparatorChar.ToString());
 			if (AlwaysTreatSourceAsDirectory && !SourcePathEndsWithDirectorySeparator)
 			{
 				AbsolutePath += Path.DirectorySeparatorChar;
 			}
 
-			var AbsolutePathUri = new Uri(AbsolutePath);
+			Uri AbsolutePathUri = new Uri(AbsolutePath);
 
-			var AbsoluteRelativeDirectory = RelativeToDirectory;
+			string AbsoluteRelativeDirectory = RelativeToDirectory;
 			if (!Path.IsPathRooted(AbsoluteRelativeDirectory))
 			{
 				AbsoluteRelativeDirectory = Path.GetFullPath(AbsoluteRelativeDirectory);
@@ -552,11 +553,11 @@ namespace UnrealBuildTool
 			}
 
 			// Convert to URI form which is where we can make the relative conversion happen
-			var AbsoluteRelativeDirectoryUri = new Uri(AbsoluteRelativeDirectory);
+			Uri AbsoluteRelativeDirectoryUri = new Uri(AbsoluteRelativeDirectory);
 
 			// Ask the URI system to convert to a nicely formed relative path, then convert it back to a regular path string
-			var UriRelativePath = AbsoluteRelativeDirectoryUri.MakeRelativeUri(AbsolutePathUri);
-			var RelativePath = Uri.UnescapeDataString(UriRelativePath.ToString()).Replace('/', Path.DirectorySeparatorChar);
+			Uri UriRelativePath = AbsoluteRelativeDirectoryUri.MakeRelativeUri(AbsolutePathUri);
+			string RelativePath = Uri.UnescapeDataString(UriRelativePath.ToString()).Replace('/', Path.DirectorySeparatorChar);
 
 			// If we added a directory separator character earlier on, remove it now
 			if (!SourcePathEndsWithDirectorySeparator && AlwaysTreatSourceAsDirectory && RelativePath.EndsWith(Path.DirectorySeparatorChar.ToString()))
@@ -600,7 +601,7 @@ namespace UnrealBuildTool
 
 			// Display updated progress string and keep track of how long it was
 			float ProgressValue = Denominator > 0 ? ((float)Numerator / (float)Denominator) : 1.0f;
-			var ProgressString = String.Format("{0}%", Math.Round(ProgressValue * 100.0f));
+			string ProgressString = String.Format("{0}%", Math.Round(ProgressValue * 100.0f));
 			NumCharsToBackspaceOver = ProgressString.Length;
 			Console.Write(ProgressString);
 		}
@@ -731,7 +732,7 @@ namespace UnrealBuildTool
 		{
 			Filename = Path.GetFileName(Filename);
 
-			var DotIndex = Filename.IndexOf('.');
+			int DotIndex = Filename.IndexOf('.');
 			if (DotIndex == -1)
 			{
 				return Filename; // No need to copy string
@@ -768,7 +769,7 @@ namespace UnrealBuildTool
 		/// <returns></returns>
 		public static bool IsFileUnderDirectory(string FilePath, string Directory)
 		{
-			var DirectoryPathPlusSeparator = Path.GetFullPath(Directory);
+			string DirectoryPathPlusSeparator = Path.GetFullPath(Directory);
 			if (!DirectoryPathPlusSeparator.EndsWith(Path.DirectorySeparatorChar.ToString()))
 			{
 				DirectoryPathPlusSeparator += Path.DirectorySeparatorChar;
@@ -798,8 +799,8 @@ namespace UnrealBuildTool
 		/// <returns></returns>
 		public static string MakeRerootedFilePath(string FilePath, string BaseDirectory, string NewBaseDirectory)
 		{
-			var RelativeFile = StripBaseDirectory(FilePath, BaseDirectory);
-			var DestFile = Path.Combine(NewBaseDirectory, RelativeFile);
+			string RelativeFile = StripBaseDirectory(FilePath, BaseDirectory);
+			string DestFile = Path.Combine(NewBaseDirectory, RelativeFile);
 			return DestFile;
 		}
 
@@ -937,10 +938,10 @@ namespace UnrealBuildTool
 			try
 			{
 				int NumCores = 0;
-				using (var Mos = new System.Management.ManagementObjectSearcher("Select * from Win32_Processor"))
+				using (ManagementObjectSearcher Mos = new System.Management.ManagementObjectSearcher("Select * from Win32_Processor"))
 				{
-					var MosCollection = Mos.Get();
-					foreach (var Item in MosCollection)
+					ManagementObjectCollection MosCollection = Mos.Get();
+					foreach (ManagementBaseObject Item in MosCollection)
 					{
 						NumCores += int.Parse(Item["NumberOfCores"].ToString());
 					}
@@ -1506,8 +1507,8 @@ namespace UnrealBuildTool
 		{
 			totalBytesRead = 0;
 			const int BytesToRead = 4096;
-			var buf = new byte[BytesToRead];
-			var bytesReadThisTime = 0;
+			byte[] buf = new byte[BytesToRead];
+			int bytesReadThisTime = 0;
 			do
 			{
 				bytesReadThisTime = input.Read(buf, 0, BytesToRead);
@@ -1524,7 +1525,7 @@ namespace UnrealBuildTool
 		/// <returns>memory stream that contains the stream contents.</returns>
 		public static MemoryStream ReadIntoMemoryStream(this Stream input)
 		{
-			var data = new MemoryStream(4096);
+			MemoryStream data = new MemoryStream(4096);
 			data.ReadFrom(input);
 			return data;
 		}
@@ -1560,7 +1561,7 @@ namespace UnrealBuildTool
 		/// <returns></returns>
 		public static int GetHashCodeFromValues<T1>(T1 Value1)
 		{
-			var Result = BasePrimeNumber;
+			int Result = BasePrimeNumber;
 			if (Value1 != null)
 			{
 				Result *= Value1.GetHashCode() + FieldMutatorPrimeNumber;
@@ -1578,7 +1579,7 @@ namespace UnrealBuildTool
 		/// <returns></returns>
 		public static int GetHashCodeFromValues<T1, T2>(T1 Value1, T2 Value2)
 		{
-			var Result = BasePrimeNumber;
+			int Result = BasePrimeNumber;
 			if (Value1 != null)
 			{
 				Result *= Value1.GetHashCode() + FieldMutatorPrimeNumber;
@@ -1602,7 +1603,7 @@ namespace UnrealBuildTool
 		/// <returns></returns>
 		public static int GetHashCodeFromValues<T1, T2, T3>(T1 Value1, T2 Value2, T3 Value3)
 		{
-			var Result = BasePrimeNumber;
+			int Result = BasePrimeNumber;
 			if (Value1 != null)
 			{
 				Result *= Value1.GetHashCode() + FieldMutatorPrimeNumber;
@@ -1632,7 +1633,7 @@ namespace UnrealBuildTool
 		/// <returns></returns>
 		public static int GetHashCodeFromValues<T1, T2, T3, T4>(T1 Value1, T2 Value2, T3 Value3, T4 Value4)
 		{
-			var Result = BasePrimeNumber;
+			int Result = BasePrimeNumber;
 			if (Value1 != null)
 			{
 				Result *= Value1.GetHashCode() + FieldMutatorPrimeNumber;
@@ -1694,13 +1695,13 @@ namespace UnrealBuildTool
 
 		public void Dispose()
 		{
-			var TotalSeconds = (DateTime.UtcNow - StartTime).TotalSeconds;
-			var LogIndent = 0;
+			double TotalSeconds = (DateTime.UtcNow - StartTime).TotalSeconds;
+			int LogIndent = 0;
 			lock (IndentLock)
 			{
 				LogIndent = --Indent;
 			}
-			var IndentText = new StringBuilder(LogIndent * 2);
+			StringBuilder IndentText = new StringBuilder(LogIndent * 2);
 			IndentText.Append(' ', LogIndent * 2);
 
 			Log.WriteLine(Verbosity, "{0}{1} took {2}s", IndentText.ToString(), TimerName, TotalSeconds);
@@ -1750,7 +1751,7 @@ namespace UnrealBuildTool
 		/// </summary>
 		public void Dispose()
 		{
-			var TotalSeconds = (DateTime.UtcNow - StartTime).TotalSeconds;
+			double TotalSeconds = (DateTime.UtcNow - StartTime).TotalSeconds;
 			ScopedAccumulator.Time += TotalSeconds;
 		}
 
@@ -1781,7 +1782,7 @@ namespace UnrealBuildTool
 		/// </summary>
 		static public void DumpAllCounters()
 		{
-			foreach (var Counter in Accumulators)
+			foreach (KeyValuePair<string, Accumulator> Counter in Accumulators)
 			{
 				LogAccumulator(Counter.Key, Counter.Value);
 			}
diff --git a/Engine/Source/Programs/UnrealBuildTool/System/VCProject.cs b/Engine/Source/Programs/UnrealBuildTool/System/VCProject.cs
--- a/Engine/Source/Programs/UnrealBuildTool/System/VCProject.cs
+++ b/Engine/Source/Programs/UnrealBuildTool/System/VCProject.cs
@@ -81,7 +81,7 @@ namespace UnrealBuildTool
 			// Only load GUIDs if we're in project generation mode.  Regular builds don't need GUIDs for anything.
 			if (ProjectFileGenerator.bGenerateProjectFiles)
 			{
-				var Doc = new XmlDocument();
+				XmlDocument Doc = new XmlDocument();
 				Doc.Load(ProjectFilePath.FullName);
 
 				// @todo projectfiles: Ideally we could do a better job about preserving GUIDs when only minor changes are made
@@ -89,7 +89,7 @@ namespace UnrealBuildTool
 
 				// @todo projectfiles: Can we "seed" a GUID based off the project path and generate consistent GUIDs each time?
 
-				var Elements = Doc.GetElementsByTagName("ProjectGuid");
+				XmlNodeList Elements = Doc.GetElementsByTagName("ProjectGuid");
 				foreach (XmlElement Element in Elements)
 				{
 					ProjectGUID = Guid.ParseExact(Element.InnerText.Trim("{}".ToCharArray()), "D");
@@ -129,7 +129,7 @@ namespace UnrealBuildTool
 				}
 				else
 				{
-					var PlatformProjectGenerator = UEPlatformProjectGenerator.GetPlatformProjectGenerator(Platform, bInAllowFailure: true);
+					UEPlatformProjectGenerator PlatformProjectGenerator = UEPlatformProjectGenerator.GetPlatformProjectGenerator(Platform, bInAllowFailure: true);
 
 					// Check to see if this platform is supported directly by Visual Studio projects.
 					bool HasActualVSPlatform = (PlatformProjectGenerator != null) ? PlatformProjectGenerator.HasVisualStudioSupport(Platform, Configuration) : false;
@@ -172,13 +172,13 @@ namespace UnrealBuildTool
 		/// <returns>True if this is a valid combination for this project, otherwise false</returns>
 		public static bool IsValidProjectPlatformAndConfiguration(ProjectTarget ProjectTarget, UnrealTargetPlatform Platform, UnrealTargetConfiguration Configuration)
 		{
-			var PlatformProjectGenerator = UEPlatformProjectGenerator.GetPlatformProjectGenerator(Platform, true);
+			UEPlatformProjectGenerator PlatformProjectGenerator = UEPlatformProjectGenerator.GetPlatformProjectGenerator(Platform, true);
 			if (PlatformProjectGenerator == null)
 			{
 				return false;
 			}
 
-			var BuildPlatform = UEBuildPlatform.GetBuildPlatform(Platform, true);
+			UEBuildPlatform BuildPlatform = UEBuildPlatform.GetBuildPlatform(Platform, true);
 			if (BuildPlatform == null)
 			{
 				return false;
@@ -190,8 +190,8 @@ namespace UnrealBuildTool
 			}
 
 
-			var SupportedConfigurations = new List<UnrealTargetConfiguration>();
-			var SupportedPlatforms = new List<UnrealTargetPlatform>();
+			List<UnrealTargetConfiguration> SupportedConfigurations = new List<UnrealTargetConfiguration>();
+			List<UnrealTargetPlatform> SupportedPlatforms = new List<UnrealTargetPlatform>();
 			if (!ProjectFileGenerator.bCreateDummyConfigsForUnsupportedPlatforms)
 			{
 				if (ProjectTarget.TargetRules != null)
@@ -211,14 +211,14 @@ namespace UnrealBuildTool
 
 			// Add all of the extra platforms/configurations for this target
 			{
-				foreach (var ExtraPlatform in ProjectTarget.ExtraSupportedPlatforms)
+				foreach (UnrealTargetPlatform ExtraPlatform in ProjectTarget.ExtraSupportedPlatforms)
 				{
 					if (!SupportedPlatforms.Contains(ExtraPlatform))
 					{
 						SupportedPlatforms.Add(ExtraPlatform);
 					}
 				}
-				foreach (var ExtraConfig in ProjectTarget.ExtraSupportedConfigurations)
+				foreach (UnrealTargetConfiguration ExtraConfig in ProjectTarget.ExtraSupportedConfigurations)
 				{
 					if (ProjectFileGenerator.bIncludeTestAndShippingConfigs || (ExtraConfig != UnrealTargetConfiguration.Shipping && ExtraConfig != UnrealTargetConfiguration.Test))
 					{
@@ -322,13 +322,13 @@ namespace UnrealBuildTool
 			bool bSuccess = true;
 
 			// Build up the new include search path string
-			var VCIncludeSearchPaths = new StringBuilder();
+			StringBuilder VCIncludeSearchPaths = new StringBuilder();
 			{
-				foreach (var CurPath in IntelliSenseIncludeSearchPaths)
+				foreach (string CurPath in IntelliSenseIncludeSearchPaths)
 				{
 					VCIncludeSearchPaths.Append(CurPath + ";");
 				}
-				foreach (var CurPath in IntelliSenseSystemIncludeSearchPaths)
+				foreach (string CurPath in IntelliSenseSystemIncludeSearchPaths)
 				{
 					VCIncludeSearchPaths.Append(CurPath + ";");
 				}
@@ -346,8 +346,8 @@ namespace UnrealBuildTool
 				}
 			}
 
-			var VCPreprocessorDefinitions = new StringBuilder();
-			foreach (var CurDef in IntelliSensePreprocessorDefinitions)
+			StringBuilder VCPreprocessorDefinitions = new StringBuilder();
+			foreach (string CurDef in IntelliSensePreprocessorDefinitions)
 			{
 				if (VCPreprocessorDefinitions.Length > 0)
 				{
@@ -357,13 +357,13 @@ namespace UnrealBuildTool
 			}
 
 			// Setup VC project file content
-			var VCProjectFileContent = new StringBuilder();
-			var VCFiltersFileContent = new StringBuilder();
-			var VCUserFileContent = new StringBuilder();
+			StringBuilder VCProjectFileContent = new StringBuilder();
+			StringBuilder VCFiltersFileContent = new StringBuilder();
+			StringBuilder VCUserFileContent = new StringBuilder();
 
 			// Visual Studio doesn't require a *.vcxproj.filters file to even exist alongside the project unless 
 			// it actually has something of substance in it.  We'll avoid saving it out unless we need to.
-			var FiltersFileIsNeeded = false;
+			bool FiltersFileIsNeeded = false;
 
 			// Project file header
 			VCProjectFileContent.Append(
@@ -383,7 +383,7 @@ namespace UnrealBuildTool
 
 			// Build up a list of platforms and configurations this project will support.  In this list, Unknown simply
 			// means that we should use the default "stub" project platform and configuration name.
-			var ProjectConfigAndTargetCombinations = new List<ProjectConfigAndTargetCombination>();
+			List<ProjectConfigAndTargetCombination> ProjectConfigAndTargetCombinations = new List<ProjectConfigAndTargetCombination>();
 
 			// If this is a "stub" project, then only add a single configuration to the project
 			if (IsStubProject)
@@ -394,20 +394,20 @@ namespace UnrealBuildTool
 			else
 			{
 				// Figure out all the desired configurations
-				foreach (var Configuration in InConfigurations)
+				foreach (UnrealTargetConfiguration Configuration in InConfigurations)
 				{
 					[REDACTED_EMAIL]: Put this in a commonly accessible place?
 					if (UnrealBuildTool.IsValidConfiguration(Configuration) == false)
 					{
 						continue;
 					}
-					foreach (var Platform in InPlatforms)
+					foreach (UnrealTargetPlatform Platform in InPlatforms)
 					{
 						if (UnrealBuildTool.IsValidPlatform(Platform) == false)
 						{
 							continue;
 						}
-						var BuildPlatform = UEBuildPlatform.GetBuildPlatform(Platform, true);
+						UEBuildPlatform BuildPlatform = UEBuildPlatform.GetBuildPlatform(Platform, true);
 						if ((BuildPlatform != null) && (BuildPlatform.HasRequiredSDKsInstalled() == SDKStatus.Valid))
 						{
 							// Now go through all of the target types for this project
@@ -416,7 +416,7 @@ namespace UnrealBuildTool
 								throw new BuildException("Expecting at least one ProjectTarget to be associated with project '{0}' in the TargetProjects list ", ProjectFilePath);
 							}
 
-							foreach (var ProjectTarget in ProjectTargets)
+							foreach (ProjectTarget ProjectTarget in ProjectTargets)
 							{
 								if (IsValidProjectPlatformAndConfiguration(ProjectTarget, Platform, Configuration))
 								{
@@ -436,10 +436,10 @@ namespace UnrealBuildTool
 				"	<ItemGroup Label=\"ProjectConfigurations\">" + ProjectFileGenerator.NewLine);
 
 			// Make a list of the platforms and configs as project-format names
-			var ProjectPlatforms = new List<UnrealTargetPlatform>();
-			var ProjectPlatformNameAndPlatforms = new List<Tuple<string, UnrealTargetPlatform>>();	// ProjectPlatformName, Platform
-			var ProjectConfigurationNameAndConfigurations = new List<Tuple<string, UnrealTargetConfiguration>>();	// ProjectConfigurationName, Configuration
-			foreach (var Combination in ProjectConfigAndTargetCombinations)
+			List<UnrealTargetPlatform> ProjectPlatforms = new List<UnrealTargetPlatform>();
+			List<Tuple<string, UnrealTargetPlatform>> ProjectPlatformNameAndPlatforms = new List<Tuple<string, UnrealTargetPlatform>>();	// ProjectPlatformName, Platform
+			List<Tuple<string, UnrealTargetConfiguration>> ProjectConfigurationNameAndConfigurations = new List<Tuple<string, UnrealTargetConfiguration>>();	// ProjectConfigurationName, Configuration
+			foreach (ProjectConfigAndTargetCombination Combination in ProjectConfigAndTargetCombinations)
 			{
 				if (!ProjectPlatforms.Contains(Combination.Platform))
 				{
@@ -456,12 +456,12 @@ namespace UnrealBuildTool
 			}
 
 			// Output ALL the project's config-platform permutations (project files MUST do this)
-			foreach (var ConfigurationTuple in ProjectConfigurationNameAndConfigurations)
+			foreach (Tuple<string, UnrealTargetConfiguration> ConfigurationTuple in ProjectConfigurationNameAndConfigurations)
 			{
-				var ProjectConfigurationName = ConfigurationTuple.Item1;
-				foreach (var PlatformTuple in ProjectPlatformNameAndPlatforms)
+				string ProjectConfigurationName = ConfigurationTuple.Item1;
+				foreach (Tuple<string, UnrealTargetPlatform> PlatformTuple in ProjectPlatformNameAndPlatforms)
 				{
-					var ProjectPlatformName = PlatformTuple.Item1;
+					string ProjectPlatformName = PlatformTuple.Item1;
 					VCProjectFileContent.Append(
 							"		<ProjectConfiguration Include=\"" + ProjectConfigurationName + "|" + ProjectPlatformName + "\">" + ProjectFileGenerator.NewLine +
 							"			<Configuration>" + ProjectConfigurationName + "</Configuration>" + ProjectFileGenerator.NewLine +
@@ -482,7 +482,7 @@ namespace UnrealBuildTool
 			StringBuilder AdditionalPropertyGroups = new StringBuilder();
 			if (!IsStubProject)
 			{
-				foreach (var Platform in ProjectPlatforms)
+				foreach (UnrealTargetPlatform Platform in ProjectPlatforms)
 				{
 					UEPlatformProjectGenerator ProjGenerator = UEPlatformProjectGenerator.GetPlatformProjectGenerator(Platform, true);
 					if (ProjGenerator != null && ProjGenerator.HasVisualStudioSupport(Platform, UnrealTargetConfiguration.Development))
@@ -508,14 +508,14 @@ namespace UnrealBuildTool
 			}
 
 			// Write each project configuration PreDefaultProps section
-			foreach (var ConfigurationTuple in ProjectConfigurationNameAndConfigurations)
+			foreach (Tuple<string, UnrealTargetConfiguration> ConfigurationTuple in ProjectConfigurationNameAndConfigurations)
 			{
-				var ProjectConfigurationName = ConfigurationTuple.Item1;
-				var TargetConfiguration = ConfigurationTuple.Item2;
-				foreach (var PlatformTuple in ProjectPlatformNameAndPlatforms)
+				string ProjectConfigurationName = ConfigurationTuple.Item1;
+				UnrealTargetConfiguration TargetConfiguration = ConfigurationTuple.Item2;
+				foreach (Tuple<string, UnrealTargetPlatform> PlatformTuple in ProjectPlatformNameAndPlatforms)
 				{
-					var ProjectPlatformName = PlatformTuple.Item1;
-					var TargetPlatform = PlatformTuple.Item2;
+					string ProjectPlatformName = PlatformTuple.Item1;
+					UnrealTargetPlatform TargetPlatform = PlatformTuple.Item2;
 					WritePreDefaultPropsConfiguration(TargetPlatform, TargetConfiguration, ProjectPlatformName, ProjectConfigurationName, VCProjectFileContent);
 				}
 			}
@@ -524,14 +524,14 @@ namespace UnrealBuildTool
 				"	<Import Project=\"$(VCTargetsPath)\\Microsoft.Cpp.Default.props\" />" + ProjectFileGenerator.NewLine);
 
 			// Write each project configuration PreDefaultProps section
-			foreach (var ConfigurationTuple in ProjectConfigurationNameAndConfigurations)
+			foreach (Tuple<string, UnrealTargetConfiguration> ConfigurationTuple in ProjectConfigurationNameAndConfigurations)
 			{
-				var ProjectConfigurationName = ConfigurationTuple.Item1;
-				var TargetConfiguration = ConfigurationTuple.Item2;
-				foreach (var PlatformTuple in ProjectPlatformNameAndPlatforms)
+				string ProjectConfigurationName = ConfigurationTuple.Item1;
+				UnrealTargetConfiguration TargetConfiguration = ConfigurationTuple.Item2;
+				foreach (Tuple<string, UnrealTargetPlatform> PlatformTuple in ProjectPlatformNameAndPlatforms)
 				{
-					var ProjectPlatformName = PlatformTuple.Item1;
-					var TargetPlatform = PlatformTuple.Item2;
+					string ProjectPlatformName = PlatformTuple.Item1;
+					UnrealTargetPlatform TargetPlatform = PlatformTuple.Item2;
 					WritePostDefaultPropsConfiguration(TargetPlatform, TargetConfiguration, ProjectPlatformName, ProjectConfigurationName, VCProjectFileContent);
 				}
 			}
@@ -543,20 +543,20 @@ namespace UnrealBuildTool
 				);
 
 			// Write each project configuration
-			foreach (var Combination in ProjectConfigAndTargetCombinations)
+			foreach (ProjectConfigAndTargetCombination Combination in ProjectConfigAndTargetCombinations)
 			{
 				WriteConfiguration(ProjectName, Combination, VCProjectFileContent, bGenerateUserFileContent ? VCUserFileContent : null);
 			}
 
 			// Source folders and files
 			{
-				var LocalAliasedFiles = new List<AliasedFile>(AliasedFiles);
+				List<AliasedFile> LocalAliasedFiles = new List<AliasedFile>(AliasedFiles);
 
-				foreach (var CurFile in SourceFiles)
+				foreach (SourceFile CurFile in SourceFiles)
 				{
 					// We want all source file and directory paths in the project files to be relative to the project file's
 					// location on the disk.  Convert the path to be relative to the project file directory
-					var ProjectRelativeSourceFile = CurFile.Reference.MakeRelativeTo(ProjectFilePath.Directory);
+					string ProjectRelativeSourceFile = CurFile.Reference.MakeRelativeTo(ProjectFilePath.Directory);
 
 					// By default, files will appear relative to the project file in the solution.  This is kind of the normal Visual
 					// Studio way to do things, but because our generated project files are emitted to intermediate folders, if we always
@@ -592,8 +592,8 @@ namespace UnrealBuildTool
 					"	<ItemGroup>" + ProjectFileGenerator.NewLine);
 
 				// Add all file directories to the filters file as solution filters
-				var FilterDirectories = new HashSet<string>();
-				foreach (var AliasedFile in LocalAliasedFiles)
+				HashSet<string> FilterDirectories = new HashSet<string>();
+				foreach (AliasedFile AliasedFile in LocalAliasedFiles)
 				{
 					// No need to add the root directory relative to the project (it would just be an empty string!)
 					if (!String.IsNullOrWhiteSpace(AliasedFile.ProjectPath))
@@ -601,7 +601,7 @@ namespace UnrealBuildTool
 						FiltersFileIsNeeded = EnsureFilterPathExists(AliasedFile.ProjectPath, VCFiltersFileContent, FilterDirectories);
 					}
 
-					var VCFileType = GetVCFileType(AliasedFile.FileSystemPath);
+					string VCFileType = GetVCFileType(AliasedFile.FileSystemPath);
 
 					VCProjectFileContent.Append(
 						"		<" + VCFileType + " Include=\"" + EscapeFileName(AliasedFile.FileSystemPath) + "\" />" + ProjectFileGenerator.NewLine);
@@ -666,7 +666,7 @@ namespace UnrealBuildTool
 			StringBuilder AdditionalImportSettings = new StringBuilder();
 			if (!IsStubProject)
 			{
-				foreach (var Platform in ProjectPlatforms)
+				foreach (UnrealTargetPlatform Platform in ProjectPlatforms)
 				{
 					UEPlatformProjectGenerator ProjGenerator = UEPlatformProjectGenerator.GetPlatformProjectGenerator(Platform, true);
 					if (ProjGenerator != null && ProjGenerator.HasVisualStudioSupport(Platform, UnrealTargetConfiguration.Development))
@@ -679,7 +679,7 @@ namespace UnrealBuildTool
 			string OutputManifestString = "";
 			if (!IsStubProject)
 			{
-				foreach (var Platform in ProjectPlatforms)
+				foreach (UnrealTargetPlatform Platform in ProjectPlatforms)
 				{
 					UEPlatformProjectGenerator ProjGenerator = UEPlatformProjectGenerator.GetPlatformProjectGenerator(Platform, true);
 					if (ProjGenerator != null && ProjGenerator.HasVisualStudioSupport(Platform, UnrealTargetConfiguration.Development))
@@ -687,9 +687,9 @@ namespace UnrealBuildTool
 						// @todo projectfiles: Serious hacks here because we are trying to emit one-time platform-specific sections that need information
 						//    about a target type, but the project file may contain many types of targets!  Some of this logic will need to move into
 						//    the per-target configuration writing code.
-						var HackTargetType = TargetRules.TargetType.Game;
+						TargetRules.TargetType HackTargetType = TargetRules.TargetType.Game;
 						FileReference HackTargetFilePath = null;
-						foreach (var Combination in ProjectConfigAndTargetCombinations)
+						foreach (ProjectConfigAndTargetCombination Combination in ProjectConfigAndTargetCombinations)
 						{
 							if (Combination.Platform == Platform &&
 								Combination.ProjectTarget.TargetRules != null &&
@@ -739,7 +739,7 @@ namespace UnrealBuildTool
 			if (bSuccess)
 			{
 				// Create a path to the project file's filters file
-				var VCFiltersFilePath = ProjectFilePath.FullName + ".filters";
+				string VCFiltersFilePath = ProjectFilePath.FullName + ".filters";
 				if (FiltersFileIsNeeded)
 				{
 					bSuccess = ProjectFileGenerator.WriteFileIfChanged(VCFiltersFilePath, VCFiltersFileContent.ToString());
@@ -764,7 +764,7 @@ namespace UnrealBuildTool
 			if (VCUserFileContent.Length > 0)
 			{
 				// Create a path to the project file's user file
-				var VCUserFilePath = ProjectFilePath.FullName + ".user";
+				string VCUserFilePath = ProjectFilePath.FullName + ".user";
 				// Never overwrite the existing user path as it will cause them to lose their settings
 				if (File.Exists(VCUserFilePath) == false)
 				{
@@ -778,18 +778,18 @@ namespace UnrealBuildTool
 		private static bool EnsureFilterPathExists(string FilterRelativeSourceDirectory, StringBuilder VCFiltersFileContent, HashSet<string> FilterDirectories)
 		{
 			// We only want each directory to appear once in the filters file
-			var PathRemaining = Utils.CleanDirectorySeparators(FilterRelativeSourceDirectory);
-			var FiltersFileIsNeeded = false;
+			string PathRemaining = Utils.CleanDirectorySeparators(FilterRelativeSourceDirectory);
+			bool FiltersFileIsNeeded = false;
 			if (!FilterDirectories.Contains(PathRemaining))
 			{
 				// Make sure all subdirectories leading up to this directory each have their own filter, too!
-				var AllDirectoriesInPath = new List<string>();
-				var PathSoFar = "";
+				List<string> AllDirectoriesInPath = new List<string>();
+				string PathSoFar = "";
 				for (; ; )
 				{
 					if (PathRemaining.Length > 0)
 					{
-						var SlashIndex = PathRemaining.IndexOf(Path.DirectorySeparatorChar);
+						int SlashIndex = PathRemaining.IndexOf(Path.DirectorySeparatorChar);
 						string SplitDirectory;
 						if (SlashIndex != -1)
 						{
@@ -815,7 +815,7 @@ namespace UnrealBuildTool
 					}
 				}
 
-				foreach (var LeadingDirectory in AllDirectoriesInPath)
+				foreach (string LeadingDirectory in AllDirectoriesInPath)
 				{
 					if (!FilterDirectories.Contains(LeadingDirectory))
 					{
@@ -824,7 +824,7 @@ namespace UnrealBuildTool
 						// Generate a unique GUID for this folder
 						// NOTE: When saving generated project files, we ignore differences in GUIDs if every other part of the file
 						//       matches identically with the pre-existing file
-						var FilterGUID = Guid.NewGuid().ToString("B").ToUpperInvariant();
+						string FilterGUID = Guid.NewGuid().ToString("B").ToUpperInvariant();
 
 						VCFiltersFileContent.Append(
 							"		<Filter Include=\"" + LeadingDirectory + "\">" + ProjectFileGenerator.NewLine +
@@ -879,7 +879,7 @@ namespace UnrealBuildTool
                 return;
             }
 
-            var ProjectConfigurationAndPlatformName = ProjectConfigurationName + "|" + ProjectPlatformName;
+            string ProjectConfigurationAndPlatformName = ProjectConfigurationName + "|" + ProjectPlatformName;
             string ConditionString = "Condition=\"'$(Configuration)|$(Platform)'=='" + ProjectConfigurationAndPlatformName + "'\"";
 
             string PlatformToolsetString = (ProjGenerator != null) ? ProjGenerator.GetVisualStudioPreDefaultString(TargetPlatform, TargetConfiguration, this) : "";
@@ -903,7 +903,7 @@ namespace UnrealBuildTool
 				return;
 			}
 
-			var ProjectConfigurationAndPlatformName = ProjectConfigurationName + "|" + ProjectPlatformName;
+			string ProjectConfigurationAndPlatformName = ProjectConfigurationName + "|" + ProjectPlatformName;
 			string ConditionString = "Condition=\"'$(Configuration)|$(Platform)'=='" + ProjectConfigurationAndPlatformName + "'\"";
 
 			string PlatformToolsetString = (ProjGenerator != null) ? ProjGenerator.GetVisualStudioPlatformToolsetString(TargetPlatform, TargetConfiguration, this) : "";
@@ -985,12 +985,12 @@ namespace UnrealBuildTool
 					TargetRules TargetRulesObject = Combination.ProjectTarget.TargetRules;
 					FileReference TargetFilePath = Combination.ProjectTarget.TargetFilePath;
 					string TargetName = TargetFilePath.GetFileNameWithoutAnyExtensions();
-					var UBTPlatformName = IsStubProject ? StubProjectPlatformName : Platform.ToString();
-					var UBTConfigurationName = IsStubProject ? StubProjectConfigurationName : Configuration.ToString();
+					string UBTPlatformName = IsStubProject ? StubProjectPlatformName : Platform.ToString();
+					string UBTConfigurationName = IsStubProject ? StubProjectConfigurationName : Configuration.ToString();
 
 					// Setup output path
-					var BuildPlatform = UEBuildPlatform.GetBuildPlatform(Platform);
-					var BuildPlatformContext = BuildPlatform.CreateContext(Combination.ProjectTarget.ProjectFilePath);
+					UEBuildPlatform BuildPlatform = UEBuildPlatform.GetBuildPlatform(Platform);
+					UEBuildPlatformContext BuildPlatformContext = BuildPlatform.CreateContext(Combination.ProjectTarget.ProjectFilePath);
 
 					// Figure out if this is a monolithic build
 					bool bShouldCompileMonolithic = BuildPlatform.ShouldCompileMonolithicBinary(Platform);
@@ -1205,14 +1205,14 @@ namespace UnrealBuildTool
 		/// </summary>
 		public List<string> GetCSharpDependencies()
 		{
-			var RelativeFilePaths = new List<string>();
-			var Doc = new XmlDocument();
+			List<string> RelativeFilePaths = new List<string>();
+			XmlDocument Doc = new XmlDocument();
 			Doc.Load(ProjectFilePath.FullName);
 
-			var Tags = new string[] { "Compile", "Page", "Resource" };
-			foreach (var Tag in Tags)
+			string[] Tags = new string[] { "Compile", "Page", "Resource" };
+			foreach (string Tag in Tags)
 			{
-				var Elements = Doc.GetElementsByTagName(Tag);
+				XmlNodeList Elements = Doc.GetElementsByTagName(Tag);
 				foreach (XmlElement Element in Elements)
 				{
 					RelativeFilePaths.Add(Element.GetAttribute("Include"));
@@ -1254,7 +1254,7 @@ namespace UnrealBuildTool
 			bool bSuccess = true;
 
 			// Setup C# project file content.
-			var ProjectFileContent = new StringBuilder();
+			StringBuilder ProjectFileContent = new StringBuilder();
 
 			// Project file header.
 			ProjectFileContent.Append(
@@ -1288,7 +1288,7 @@ namespace UnrealBuildTool
 			if (DotNetAssemblyReferences.Count > 0)
 			{
 				ProjectFileContent.Append("<ItemGroup>" + ProjectFileGenerator.NewLine);
-				foreach (var CurReference in DotNetAssemblyReferences)
+				foreach (string CurReference in DotNetAssemblyReferences)
 				{
 					ProjectFileContent.Append("\t<Reference Include=\"" + EscapeFileName(CurReference) + "\" />" + ProjectFileGenerator.NewLine);
 				}
@@ -1298,7 +1298,7 @@ namespace UnrealBuildTool
 			if (AssemblyReferences.Count > 0)
 			{
 				ProjectFileContent.Append("<ItemGroup>" + ProjectFileGenerator.NewLine);
-				foreach (var CurReference in AssemblyReferences)
+				foreach (FileReference CurReference in AssemblyReferences)
 				{
 					ProjectFileContent.Append("\t<Reference Include=\"" + EscapeFileName(CurReference.GetFileNameWithoutExtension()) + "\" >" + ProjectFileGenerator.NewLine);
 					ProjectFileContent.Append("\t\t<HintPath>" + CurReference.MakeRelativeTo(ProjectFilePath.Directory) + "</HintPath>" + ProjectFileGenerator.NewLine);
@@ -1308,9 +1308,9 @@ namespace UnrealBuildTool
 			}
 
 			// Other references (note it's assumed all references here are at least of MSBuildProjectFile type.
-			foreach (var Project in DependsOnProjects)
+			foreach (ProjectFile Project in DependsOnProjects)
 			{
-				var RelativePath = Project.ProjectFilePath.Directory.MakeRelativeTo(ProjectFilePath.Directory);
+				string RelativePath = Project.ProjectFilePath.Directory.MakeRelativeTo(ProjectFilePath.Directory);
 				RelativePath = Path.Combine(RelativePath, Project.ProjectFilePath.GetFileName());
 				ProjectFileContent.Append(
 					"<ItemGroup>" + ProjectFileGenerator.NewLine +
@@ -1325,9 +1325,9 @@ namespace UnrealBuildTool
 			ProjectFileContent.Append(
 				"	<ItemGroup>" + ProjectFileGenerator.NewLine);
 			// Add all files to the project.
-			foreach (var CurFile in SourceFiles)
+			foreach (SourceFile CurFile in SourceFiles)
 			{
-				var ProjectRelativeSourceFile = CurFile.Reference.MakeRelativeTo(ProjectFilePath.Directory);
+				string ProjectRelativeSourceFile = CurFile.Reference.MakeRelativeTo(ProjectFilePath.Directory);
 				ProjectFileContent.Append(
 					"		<Compile Include=\"" + EscapeFileName(ProjectRelativeSourceFile) + "\" />" + ProjectFileGenerator.NewLine);
 			}
diff --git a/Engine/Source/Programs/UnrealBuildTool/System/VCProjectFileGenerator.cs b/Engine/Source/Programs/UnrealBuildTool/System/VCProjectFileGenerator.cs
--- a/Engine/Source/Programs/UnrealBuildTool/System/VCProjectFileGenerator.cs
+++ b/Engine/Source/Programs/UnrealBuildTool/System/VCProjectFileGenerator.cs
@@ -198,7 +198,7 @@ namespace UnrealBuildTool
 				}
 			}
 
-			foreach (var CurArgument in Arguments)
+			foreach (string CurArgument in Arguments)
 			{
 				switch (CurArgument.ToUpperInvariant())
 				{
@@ -245,9 +245,9 @@ namespace UnrealBuildTool
 			// ...but only if the user didn't override this via the command-line.
 			if (!UnrealBuildTool.CommandLineContains("-2015") && !UnrealBuildTool.CommandLineContains("-2013"))
 			{
-				foreach (var SupportedPlatform in SupportedPlatforms)
+				foreach (UnrealTargetPlatform SupportedPlatform in SupportedPlatforms)
 				{
-					var BuildPlatform = UEBuildPlatform.GetBuildPlatform(SupportedPlatform, true);
+					UEBuildPlatform BuildPlatform = UEBuildPlatform.GetBuildPlatform(SupportedPlatform, true);
 					if (BuildPlatform != null)
 					{
 						// Don't worry about platforms that we're missing SDKs for
@@ -312,7 +312,7 @@ namespace UnrealBuildTool
 		/// <returns>The generated solution configuration name</returns>
 		string MakeSolutionConfigurationName(UnrealTargetConfiguration Configuration, string TargetConfigurationName)
 		{
-			var SolutionConfigName = Configuration.ToString();
+			string SolutionConfigName = Configuration.ToString();
 
 			// Don't bother postfixing "Game" or "Program" -- that will be the default when using "Debug", "Development", etc.
 			// Also don't postfix "RocketGame" when we're building Rocket game projects.  That's the only type of game there is in that case!
@@ -331,10 +331,10 @@ namespace UnrealBuildTool
 		{
 			bool bSuccess = true;
 
-			var SolutionFileName = MasterProjectName + ".sln";
+			string SolutionFileName = MasterProjectName + ".sln";
 
 			// Setup solution file content
-			var VCSolutionFileContent = new StringBuilder();
+			StringBuilder VCSolutionFileContent = new StringBuilder();
 
 			const string VersionTag = "# UnrealEngineGeneratedSolutionVersion=1.0";
 
@@ -392,16 +392,16 @@ namespace UnrealBuildTool
 			// Solution folders, files and project entries
 			{
 				// This the GUID that Visual Studio uses to identify a solution folder
-				var SolutionFolderEntryGUID = "{2150E333-8FDC-42A3-9474-1A3956D46DE8}";
+				string SolutionFolderEntryGUID = "{2150E333-8FDC-42A3-9474-1A3956D46DE8}";
 
 				// Solution folders
 				{
-					var AllSolutionFolders = new List<MasterProjectFolder>();
+					List<MasterProjectFolder> AllSolutionFolders = new List<MasterProjectFolder>();
 					System.Action<List<MasterProjectFolder> /* Folders */ > GatherFoldersFunction = null;
 					GatherFoldersFunction = FolderList =>
 						{
 							AllSolutionFolders.AddRange(FolderList);
-							foreach (var CurSubFolder in FolderList)
+							foreach (MasterProjectFolder CurSubFolder in FolderList)
 							{
 								GatherFoldersFunction(CurSubFolder.SubFolders);
 							}
@@ -410,7 +410,7 @@ namespace UnrealBuildTool
 
 					foreach (VisualStudioSolutionFolder CurFolder in AllSolutionFolders)
 					{
-						var FolderGUIDString = CurFolder.FolderGUID.ToString("B").ToUpperInvariant();
+						string FolderGUIDString = CurFolder.FolderGUID.ToString("B").ToUpperInvariant();
 						VCSolutionFileContent.Append(
 								"Project(\"" + SolutionFolderEntryGUID + "\") = \"" + CurFolder.FolderName + "\", \"" + CurFolder.FolderName + "\", \"" + FolderGUIDString + "\"" + ProjectFileGenerator.NewLine);
 
@@ -419,7 +419,7 @@ namespace UnrealBuildTool
 						{
 							VCSolutionFileContent.Append(
 									"	ProjectSection(SolutionItems) = preProject" + ProjectFileGenerator.NewLine);
-							foreach (var CurFile in CurFolder.Files)
+							foreach (string CurFile in CurFolder.Files)
 							{
 								// Syntax is:  <relative file path> = <relative file path>
 								VCSolutionFileContent.Append(
@@ -444,10 +444,10 @@ namespace UnrealBuildTool
 
 					// NOTE: The project name in the solution doesn't actually *have* to match the project file name on disk.  However,
 					//       we prefer it when it does match so we use the actual file name here.
-					var ProjectNameInSolution = CurProject.ProjectFilePath.GetFileNameWithoutExtension();
+					string ProjectNameInSolution = CurProject.ProjectFilePath.GetFileNameWithoutExtension();
 
 					// Use the existing project's GUID that's already known to us
-					var ProjectGUID = CurProject.ProjectGUID.ToString("B").ToUpperInvariant();
+					string ProjectGUID = CurProject.ProjectGUID.ToString("B").ToUpperInvariant();
 
 					VCSolutionFileContent.Append(
 							"Project(\"" + ProjectTypeGUID + "\") = \"" + ProjectNameInSolution + "\", \"" + CurProject.ProjectFilePath.MakeRelativeTo(ProjectFileGenerator.MasterProjectPath) + "\", \"" + ProjectGUID + "\"" + ProjectFileGenerator.NewLine);
@@ -456,7 +456,7 @@ namespace UnrealBuildTool
 					// freshly compiled before kicking off any build operations on this target project
 					if (!CurProject.IsStubProject)
 					{
-						var Dependencies = new List<ProjectFile>();
+						List<ProjectFile> Dependencies = new List<ProjectFile>();
 						if (CurProject.IsGeneratedProject && UBTProject != null && CurProject != UBTProject)
 						{
 							Dependencies.Add(UBTProject);
@@ -473,9 +473,9 @@ namespace UnrealBuildTool
 							VCSolutionFileContent.Append("\tProjectSection(ProjectDependencies) = postProject" + ProjectFileGenerator.NewLine);
 
 							// Setup any addition dependencies this project has...
-							foreach (var DependsOnProject in Dependencies)
+							foreach (ProjectFile DependsOnProject in Dependencies)
 							{
-								var DependsOnProjectGUID = ((MSBuildProjectFile)DependsOnProject).ProjectGUID.ToString("B").ToUpperInvariant();
+								string DependsOnProjectGUID = ((MSBuildProjectFile)DependsOnProject).ProjectGUID.ToString("B").ToUpperInvariant();
 								VCSolutionFileContent.Append("\t\t" + DependsOnProjectGUID + " = " + DependsOnProjectGUID + ProjectFileGenerator.NewLine);
 							}
 
@@ -491,7 +491,7 @@ namespace UnrealBuildTool
 
 			// Solution configuration platforms.  This is just a list of all of the platforms and configurations that
 			// appear in Visual Studio's build configuration selector.
-			var SolutionConfigCombinations = new List<VCSolutionConfigCombination>();
+			List<VCSolutionConfigCombination> SolutionConfigCombinations = new List<VCSolutionConfigCombination>();
 
 			// The "Global" section has source control, solution configurations, project configurations,
 			// preferences, and project hierarchy data
@@ -503,18 +503,18 @@ namespace UnrealBuildTool
 						VCSolutionFileContent.Append(
 							"	GlobalSection(SolutionConfigurationPlatforms) = preSolution" + ProjectFileGenerator.NewLine);
 
-						var SolutionConfigurationsValidForProjects = new Dictionary<string, Tuple<UnrealTargetConfiguration, string>>();
-						var PlatformsValidForProjects = new HashSet<UnrealTargetPlatform>();
+						Dictionary<string, Tuple<UnrealTargetConfiguration, string>> SolutionConfigurationsValidForProjects = new Dictionary<string, Tuple<UnrealTargetConfiguration, string>>();
+						HashSet<UnrealTargetPlatform> PlatformsValidForProjects = new HashSet<UnrealTargetPlatform>();
 
-						foreach (var CurConfiguration in SupportedConfigurations)
+						foreach (UnrealTargetConfiguration CurConfiguration in SupportedConfigurations)
 						{
 							if (UnrealBuildTool.IsValidConfiguration(CurConfiguration))
 							{
-								foreach (var CurPlatform in SupportedPlatforms)
+								foreach (UnrealTargetPlatform CurPlatform in SupportedPlatforms)
 								{
 									if (UnrealBuildTool.IsValidPlatform(CurPlatform))
 									{
-										foreach (var CurProject in AllProjectFiles)
+										foreach (ProjectFile CurProject in AllProjectFiles)
 										{
 											if (!CurProject.IsStubProject)
 											{
@@ -524,20 +524,20 @@ namespace UnrealBuildTool
 												}
 
 												// Figure out the set of valid target configuration names
-												foreach (var ProjectTarget in CurProject.ProjectTargets)
+												foreach (ProjectTarget ProjectTarget in CurProject.ProjectTargets)
 												{
 													if (VCProjectFile.IsValidProjectPlatformAndConfiguration(ProjectTarget, CurPlatform, CurConfiguration))
 													{
 														PlatformsValidForProjects.Add(CurPlatform);
 
 														// Default to a target configuration name of "Game", since that will collapse down to an empty string
-														var TargetConfigurationName = TargetRules.TargetType.Game.ToString();
+														string TargetConfigurationName = TargetRules.TargetType.Game.ToString();
 														if (ProjectTarget.TargetRules != null)
 														{
 															TargetConfigurationName = ProjectTarget.TargetRules.ConfigurationName;
 														}
 
-														var SolutionConfigName = MakeSolutionConfigurationName(CurConfiguration, TargetConfigurationName);
+														string SolutionConfigName = MakeSolutionConfigurationName(CurConfiguration, TargetConfigurationName);
 														SolutionConfigurationsValidForProjects[SolutionConfigName] = new Tuple<UnrealTargetConfiguration, string>(CurConfiguration, TargetConfigurationName);
 													}
 												}
@@ -548,18 +548,18 @@ namespace UnrealBuildTool
 							}
 						}
 
-						foreach (var CurPlatform in PlatformsValidForProjects)
+						foreach (UnrealTargetPlatform CurPlatform in PlatformsValidForProjects)
 						{
-							foreach (var SolutionConfigKeyValue in SolutionConfigurationsValidForProjects)
+							foreach (KeyValuePair<string, Tuple<UnrealTargetConfiguration, string>> SolutionConfigKeyValue in SolutionConfigurationsValidForProjects)
 							{
 								// e.g.  "Development|Win64 = Development|Win64"
-								var SolutionConfigName = SolutionConfigKeyValue.Key;
-								var Configuration = SolutionConfigKeyValue.Value.Item1;
-								var TargetConfigurationName = SolutionConfigKeyValue.Value.Item2;
+								string SolutionConfigName = SolutionConfigKeyValue.Key;
+								UnrealTargetConfiguration Configuration = SolutionConfigKeyValue.Value.Item1;
+								string TargetConfigurationName = SolutionConfigKeyValue.Value.Item2;
 
-								var SolutionPlatformName = CurPlatform.ToString();
+								string SolutionPlatformName = CurPlatform.ToString();
 
-								var SolutionConfigAndPlatformPair = SolutionConfigName + "|" + SolutionPlatformName;
+								string SolutionConfigAndPlatformPair = SolutionConfigName + "|" + SolutionPlatformName;
 								SolutionConfigCombinations.Add(
 										new VCSolutionConfigCombination
 										{
@@ -579,8 +579,8 @@ namespace UnrealBuildTool
 								)
 							);
 
-						var AppendedSolutionConfigAndPlatformNames = new HashSet<string>(StringComparer.InvariantCultureIgnoreCase);
-						foreach (var SolutionConfigCombination in SolutionConfigCombinations)
+						HashSet<string> AppendedSolutionConfigAndPlatformNames = new HashSet<string>(StringComparer.InvariantCultureIgnoreCase);
+						foreach (VCSolutionConfigCombination SolutionConfigCombination in SolutionConfigCombinations)
 						{
 							// We alias "Game" and "Program" to both have the same solution configuration, so we're careful not to add the same combination twice.
 							if (!AppendedSolutionConfigAndPlatformNames.Contains(SolutionConfigCombination.VCSolutionConfigAndPlatformName))
@@ -602,7 +602,7 @@ namespace UnrealBuildTool
 						VCSolutionFileContent.Append(
 							"	GlobalSection(ProjectConfigurationPlatforms) = postSolution" + ProjectFileGenerator.NewLine);
 
-						var CombinationsThatWereMatchedToProjects = new List<VCSolutionConfigCombination>();
+						List<VCSolutionConfigCombination> CombinationsThatWereMatchedToProjects = new List<VCSolutionConfigCombination>();
 
 						foreach (MSBuildProjectFile CurProject in AllProjectFiles)
 						{
@@ -618,17 +618,17 @@ namespace UnrealBuildTool
 								{
 									throw new BuildException("Expecting project '" + CurProject.ProjectFilePath + "' to have at least one ProjectTarget associated with it!");
 								}
-								var IsProgramProject = CurProject.ProjectTargets[0].TargetRules != null && CurProject.ProjectTargets[0].TargetRules.Type == TargetRules.TargetType.Program;
+								bool IsProgramProject = CurProject.ProjectTargets[0].TargetRules != null && CurProject.ProjectTargets[0].TargetRules.Type == TargetRules.TargetType.Program;
 
-								var GameOrProgramConfigsAlreadyMapped = new HashSet<string>();
-								foreach (var SolutionConfigCombination in SolutionConfigCombinations)
+								HashSet<string> GameOrProgramConfigsAlreadyMapped = new HashSet<string>();
+								foreach (VCSolutionConfigCombination SolutionConfigCombination in SolutionConfigCombinations)
 								{
 									// Handle aliasing of Program and Game target configuration names
 									if ((IsProgramProject && GameOrProgramConfigsAlreadyMapped.Add(SolutionConfigCombination.VCSolutionConfigAndPlatformName)) ||
 										IsProgramProject && SolutionConfigCombination.TargetConfigurationName != TargetRules.TargetType.Game.ToString() ||
 										!IsProgramProject && SolutionConfigCombination.TargetConfigurationName != TargetRules.TargetType.Program.ToString())
 									{
-										var TargetConfigurationName = SolutionConfigCombination.TargetConfigurationName;
+										string TargetConfigurationName = SolutionConfigCombination.TargetConfigurationName;
 										if (IsProgramProject && TargetConfigurationName == TargetRules.TargetType.Game.ToString())
 										{
 											TargetConfigurationName = TargetRules.TargetType.Program.ToString();
@@ -637,7 +637,7 @@ namespace UnrealBuildTool
 										// Now, we want to find a target in this project that maps to the current solution config combination.  Only up to one target should
 										// and every solution config combination should map to at least one target in one project (otherwise we shouldn't have added it!).
 										ProjectTarget MatchingProjectTarget = null;
-										foreach (var ProjectTarget in CurProject.ProjectTargets)
+										foreach (ProjectTarget ProjectTarget in CurProject.ProjectTargets)
 										{
 											bool IsMatchingCombination = VCProjectFile.IsValidProjectPlatformAndConfiguration(ProjectTarget, SolutionConfigCombination.Platform, SolutionConfigCombination.Configuration);
 											if (ProjectTarget.TargetRules != null)
@@ -675,8 +675,8 @@ namespace UnrealBuildTool
 											}
 										}
 
-										var SolutionConfiguration = SolutionConfigCombination.Configuration;
-										var SolutionPlatform = SolutionConfigCombination.Platform;
+										UnrealTargetConfiguration SolutionConfiguration = SolutionConfigCombination.Configuration;
+										UnrealTargetPlatform SolutionPlatform = SolutionConfigCombination.Platform;
 
 
 										if (MatchingProjectTarget == null)
@@ -690,7 +690,7 @@ namespace UnrealBuildTool
 											SolutionPlatform = UnrealTargetPlatform.Win64;
 											if (CurProject.ProjectTargets[0].TargetRules != null)
 											{
-												var ProjectSupportedPlatforms = new List<UnrealTargetPlatform>();
+												List<UnrealTargetPlatform> ProjectSupportedPlatforms = new List<UnrealTargetPlatform>();
 												CurProject.ProjectTargets[0].TargetRules.GetSupportedPlatforms(ref ProjectSupportedPlatforms);
 												if (!ProjectSupportedPlatforms.Contains(SolutionPlatform))
 												{
@@ -738,10 +738,10 @@ namespace UnrealBuildTool
 										string ProjectPlatformName;
 										CurProject.MakeProjectPlatformAndConfigurationNames(SolutionPlatform, SolutionConfiguration, TargetConfigurationName, out ProjectPlatformName, out ProjectConfigName);
 
-										var ProjectConfigAndPlatformPair = ProjectConfigName.ToString() + "|" + ProjectPlatformName.ToString();
+										string ProjectConfigAndPlatformPair = ProjectConfigName.ToString() + "|" + ProjectPlatformName.ToString();
 
 										// e.g.  "{4232C52C-680F-4850-8855-DC39419B5E9B}.Debug|iOS.ActiveCfg = iOS_Debug|Win32"
-										var CurProjectGUID = CurProject.ProjectGUID.ToString("B").ToUpperInvariant();
+										string CurProjectGUID = CurProject.ProjectGUID.ToString("B").ToUpperInvariant();
 										VCSolutionFileContent.Append(
 											"		" + CurProjectGUID + "." + SolutionConfigCombination.VCSolutionConfigAndPlatformName + ".ActiveCfg = " + ProjectConfigAndPlatformPair + ProjectFileGenerator.NewLine);
 
@@ -761,7 +761,7 @@ namespace UnrealBuildTool
 												VCSolutionFileContent.Append(
 														"		" + CurProjectGUID + "." + SolutionConfigCombination.VCSolutionConfigAndPlatformName + ".Build.0 = " + ProjectConfigAndPlatformPair + ProjectFileGenerator.NewLine);
 
-												var ProjGen = UEPlatformProjectGenerator.GetPlatformProjectGenerator(SolutionConfigCombination.Platform, true);
+												UEPlatformProjectGenerator ProjGen = UEPlatformProjectGenerator.GetPlatformProjectGenerator(SolutionConfigCombination.Platform, true);
 												if (MatchingProjectTarget.ProjectDeploys ||
 													((ProjGen != null) && (ProjGen.GetVisualStudioDeploymentEnabled(SolutionPlatform, SolutionConfiguration) == true)))
 												{
@@ -778,7 +778,7 @@ namespace UnrealBuildTool
 						}
 
 						// Check for problems
-						foreach (var SolutionConfigCombination in SolutionConfigCombinations)
+						foreach (VCSolutionConfigCombination SolutionConfigCombination in SolutionConfigCombinations)
 						{
 							if (!CombinationsThatWereMatchedToProjects.Contains(SolutionConfigCombination))
 							{
@@ -820,7 +820,7 @@ namespace UnrealBuildTool
 							{
 								foreach (VisualStudioSolutionFolder CurFolder in LocalMasterProjectFolders)
 								{
-									var CurFolderGUIDString = CurFolder.FolderGUID.ToString("B").ToUpperInvariant();
+									string CurFolderGUIDString = CurFolder.FolderGUID.ToString("B").ToUpperInvariant();
 
 									foreach (MSBuildProjectFile ChildProject in CurFolder.ChildProjects)
 									{
@@ -855,7 +855,7 @@ namespace UnrealBuildTool
 			// Save the solution file
 			if (bSuccess)
 			{
-				var SolutionFilePath = FileReference.Combine(MasterProjectPath, SolutionFileName).FullName;
+				string SolutionFilePath = FileReference.Combine(MasterProjectPath, SolutionFileName).FullName;
 				bSuccess = WriteFileIfChanged(SolutionFilePath, VCSolutionFileContent.ToString());
 			}
 
@@ -938,7 +938,7 @@ namespace UnrealBuildTool
 		/// <returns>The cleaned up string</returns>
 		public string CleanupStringForSCC(string Str)
 		{
-			var Cleaned = Str;
+			string Cleaned = Str;
 
 			// SCC is expecting paths to contain only double-backslashes for path separators.  It's a bit weird but we need to do it.
 			Cleaned = Cleaned.Replace(Path.DirectorySeparatorChar.ToString(), Path.DirectorySeparatorChar.ToString() + Path.DirectorySeparatorChar.ToString());
diff --git a/Engine/Source/Programs/UnrealBuildTool/System/XGE.cs b/Engine/Source/Programs/UnrealBuildTool/System/XGE.cs
--- a/Engine/Source/Programs/UnrealBuildTool/System/XGE.cs
+++ b/Engine/Source/Programs/UnrealBuildTool/System/XGE.cs
@@ -1,6 +1,7 @@
 // Copyright 1998-2016 Epic Games, Inc. All Rights Reserved.
 
 using System;
+using System.Collections;
 using System.Collections.Generic;
 using System.IO;
 using System.Diagnostics;
@@ -84,7 +85,7 @@ namespace UnrealBuildTool
 						XGEResult = XGE.ExecuteTaskFileWithProgressMarkup(XGETaskFilePath, Actions.Count, (sender, args) =>
 							{
 								// sometimes the args comes in as null
-								var match = XGEDurationRegex.Match(args.Data ?? "");
+								Match match = XGEDurationRegex.Match(args.Data ?? "");
 
 								/*
 																// This is considered fragile and risky parsing of undocumented XGE output, so protect this code from taking down UBT
@@ -140,7 +141,7 @@ namespace UnrealBuildTool
 
 								// XGE outputs the duration info in a format that makes VC think it's an error file/line notation if the full filename is used.
 								// So munge it a bit so we don't confuse VC.
-								var updatedData = match.Success ? args.Data.Replace('(', '[').Replace(')', ']') : args.Data;
+								string updatedData = match.Success ? args.Data.Replace('(', '[').Replace(')', ']') : args.Data;
 
 								// forward the output on to the normal handler or the console like normal
 								if (Actions[0].OutputEventHandler != null)
@@ -198,7 +199,7 @@ namespace UnrealBuildTool
 			List<Action> Actions = InActions;
 			if (BuildConfiguration.bXGEExport)
 			{
-				var CurrentEnvironment = Environment.GetEnvironmentVariables();
+				IDictionary CurrentEnvironment = Environment.GetEnvironmentVariables();
 				foreach (System.Collections.DictionaryEntry Pair in CurrentEnvironment)
 				{
 					if (!UnrealBuildTool.InitialEnvironment.Contains(Pair.Key) || (string)(UnrealBuildTool.InitialEnvironment[Pair.Key]) != (string)(Pair.Value))
@@ -231,7 +232,7 @@ namespace UnrealBuildTool
 				{
 					//Console.WriteLine("The UBT action graph was not sorted. Sorting actions....");
 					Actions = new List<Action>();
-					var UsedActions = new HashSet<int>();
+					HashSet<int> UsedActions = new HashSet<int>();
 					for (int ActionIndex = 0; ActionIndex < InActions.Count; ActionIndex++)
 					{
 						if (UsedActions.Contains(ActionIndex))
@@ -305,7 +306,7 @@ namespace UnrealBuildTool
 				XmlElement VariablesElement = XGETaskDocument.CreateElement("Variables");
 				EnvironmentElement.AppendChild(VariablesElement);
 
-				foreach (var Pair in ExportEnv)
+				foreach (KeyValuePair<string, string> Pair in ExportEnv)
 				{
 					// <Variable>...</Variable>
 					XmlElement VariableElement = XGETaskDocument.CreateElement("Variable");
diff --git a/Engine/Source/Programs/UnrealBuildTool/ToolChain/UEToolChain.cs b/Engine/Source/Programs/UnrealBuildTool/ToolChain/UEToolChain.cs
--- a/Engine/Source/Programs/UnrealBuildTool/ToolChain/UEToolChain.cs
+++ b/Engine/Source/Programs/UnrealBuildTool/ToolChain/UEToolChain.cs
@@ -110,7 +110,7 @@ namespace UnrealBuildTool
 		{
 			PrerequisiteItems.Add(SourceFile);
 
-			var RemoteThis = this as RemoteToolChain;
+			RemoteToolChain RemoteThis = this as RemoteToolChain;
 			bool bAllowUploading = RemoteThis != null && BuildHostPlatform.Current.Platform != UnrealTargetPlatform.Mac;	// Don't use remote features when compiling from a Mac
 			if (bAllowUploading)
 			{
@@ -123,7 +123,7 @@ namespace UnrealBuildTool
 				//		-> Two CASES:
 				//				1) NOT WORKING: Non-unity file went away (SourceFile in this context).  That seems like an existing old use case.  Compile params or Response file should have changed?
 				//				2) WORKING: Indirect file went away (unity'd original source file or include).  This would return a file that no longer exists and adds to the prerequiteitems list
-				var IncludedFileList = CPPEnvironment.FindAndCacheAllIncludedFiles(Target, SourceFile, BuildPlatform, CompileEnvironment.Config.CPPIncludeInfo, bOnlyCachedDependencies: BuildConfiguration.bUseUBTMakefiles);
+				List<FileItem> IncludedFileList = CPPEnvironment.FindAndCacheAllIncludedFiles(Target, SourceFile, BuildPlatform, CompileEnvironment.Config.CPPIncludeInfo, bOnlyCachedDependencies: BuildConfiguration.bUseUBTMakefiles);
 				if (IncludedFileList != null)
 				{
 					foreach (FileItem IncludedFile in IncludedFileList)
