Copying //UE4/Dev-Build to //UE4/Main

diff --git a/Engine/Source/Programs/UnrealBuildTool/Utilities/GraphVisualization.cs b/Engine/Source/Programs/UnrealBuildTool/Utilities/GraphVisualization.cs
--- a/Engine/Source/Programs/UnrealBuildTool/Utilities/GraphVisualization.cs
+++ b/Engine/Source/Programs/UnrealBuildTool/Utilities/GraphVisualization.cs
@@ -96,18 +96,18 @@ namespace UnrealBuildTool
 		/// <param name="GraphEdges">List of all graph edges.  Index order is important and must match with the individual edge Id members!</param>
 		public static void WriteGraphFile(string Filename, string Description, List<GraphNode> GraphNodes, List<GraphEdge> GraphEdges)
 		{
-			var Settings = new XmlWriterSettings();
+			XmlWriterSettings Settings = new XmlWriterSettings();
 			Settings.Indent = true;
 			Settings.IndentChars = "    ";
 
 			// Figure out all of the custom attribute types we're dealing with
-			var AllAttributes = new Dictionary<string, GraphAttribute>(StringComparer.InvariantCultureIgnoreCase);
+			Dictionary<string, GraphAttribute> AllAttributes = new Dictionary<string, GraphAttribute>(StringComparer.InvariantCultureIgnoreCase);
 			int NextAttributeID = 0;
-			foreach (var GraphNode in GraphNodes)
+			foreach (GraphNode GraphNode in GraphNodes)
 			{
-				foreach (var AttributeName in GraphNode.Attributes.Keys)
+				foreach (string AttributeName in GraphNode.Attributes.Keys)
 				{
-					var AttributeValue = GraphNode.Attributes[AttributeName];
+					object AttributeValue = GraphNode.Attributes[AttributeName];
 
 					string AttributeTypeName;
 					if (AttributeValue.GetType() == typeof(int))
@@ -158,13 +158,13 @@ namespace UnrealBuildTool
 			}
 
 
-			using (var Writer = XmlWriter.Create(Filename, Settings))
+			using (XmlWriter Writer = XmlWriter.Create(Filename, Settings))
 			{
 				// NOTE: The GEXF XML format is defined here:  http://gexf.net/1.2draft/gexf-12draft-primer.pdf
 
-				var GEXFNamespace = "http://www.gexf.net/1.2-draft";
-				var SchemaNamespace = "http://www.w3.org/2001/XMLSchema-instance";
-				var VizNamespace = "http://www.gexf.net/1.2draft/viz";
+				string GEXFNamespace = "http://www.gexf.net/1.2-draft";
+				string SchemaNamespace = "http://www.w3.org/2001/XMLSchema-instance";
+				string VizNamespace = "http://www.gexf.net/1.2draft/viz";
 
 				Writer.WriteStartElement("gexf", GEXFNamespace);
 				Writer.WriteAttributeString("xmlns", "xsi", null, SchemaNamespace);
@@ -192,7 +192,7 @@ namespace UnrealBuildTool
 								// @todo: Add support for edge attributes, not just node attributes
 								Writer.WriteAttributeString("class", "node");	// Node attributes, not edges!
 
-								foreach (var Attribute in AllAttributes.Values)
+								foreach (GraphAttribute Attribute in AllAttributes.Values)
 								{
 									Writer.WriteStartElement("attribute");
 									{
@@ -211,7 +211,7 @@ namespace UnrealBuildTool
 
 						Writer.WriteStartElement("nodes");
 						{
-							foreach (var GraphNode in GraphNodes)
+							foreach (GraphNode GraphNode in GraphNodes)
 							{
 								Writer.WriteStartElement("node");
 								{
@@ -244,12 +244,12 @@ namespace UnrealBuildTool
 									{
 										Writer.WriteStartElement("attvalues");
 										{
-											foreach (var AttributeHashEntry in GraphNode.Attributes)
+											foreach (KeyValuePair<string, object> AttributeHashEntry in GraphNode.Attributes)
 											{
-												var AttributeName = AttributeHashEntry.Key;
-												var AttributeValue = AttributeHashEntry.Value;
+												string AttributeName = AttributeHashEntry.Key;
+												object AttributeValue = AttributeHashEntry.Value;
 
-												var Attribute = AllAttributes[AttributeName];
+												GraphAttribute Attribute = AllAttributes[AttributeName];
 
 												Writer.WriteStartElement("attvalue");
 												{
@@ -270,7 +270,7 @@ namespace UnrealBuildTool
 
 						Writer.WriteStartElement("edges");
 						{
-							foreach (var GraphEdge in GraphEdges)
+							foreach (GraphEdge GraphEdge in GraphEdges)
 							{
 								Writer.WriteStartElement("edge");
 								{
diff --git a/Engine/Source/Programs/UnrealBuildTool/Utilities/XmlConfigLoader.cs b/Engine/Source/Programs/UnrealBuildTool/Utilities/XmlConfigLoader.cs
--- a/Engine/Source/Programs/UnrealBuildTool/Utilities/XmlConfigLoader.cs
+++ b/Engine/Source/Programs/UnrealBuildTool/Utilities/XmlConfigLoader.cs
@@ -47,7 +47,7 @@ namespace UnrealBuildTool
 		/// <returns>XML element representing config class.</returns>
 		private static XElement CreateConfigTypeXMLElement(Type ConfigType)
 		{
-			var NS = XNamespace.Get("https://www.unrealengine.com/BuildConfiguration");
+			XNamespace NS = XNamespace.Get("https://www.unrealengine.com/BuildConfiguration");
 
 			return new XElement(NS + ConfigType.Name,
 				from Field in GetConfigurableFields(ConfigType)
@@ -61,7 +61,7 @@ namespace UnrealBuildTool
 		/// <returns>XML representing current config.</returns>
 		private static XDocument GetConfigXML()
 		{
-			var NS = XNamespace.Get("https://www.unrealengine.com/BuildConfiguration");
+			XNamespace NS = XNamespace.Get("https://www.unrealengine.com/BuildConfiguration");
 			return new XDocument(
 					new XElement(NS + "Configuration",
 						from ConfigType in GetAllConfigurationTypes()
@@ -76,7 +76,7 @@ namespace UnrealBuildTool
 		/// <returns>Default config XML.</returns>
 		public static string GetDefaultXML()
 		{
-			var Comment = @"
+			string Comment = @"
 	#########################################################################
 	#																		#
 	#	This is an XML with default UBT configuration. If you want to		#
@@ -121,7 +121,7 @@ namespace UnrealBuildTool
 	priority. Not defined classes and fields are left alone.
 ";
 
-			var DefaultXml = GetConfigXML();
+			XDocument DefaultXml = GetConfigXML();
 
 			DefaultXml.AddFirst(new XComment(Comment));
 
@@ -153,7 +153,7 @@ namespace UnrealBuildTool
 		public static string BuildXSD()
 		{
 			// all elements use this namespace
-			var NS = XNamespace.Get("http://www.w3.org/2001/XMLSchema");
+			XNamespace NS = XNamespace.Get("http://www.w3.org/2001/XMLSchema");
 
 			return WriteToBuffer(new XDocument(
 				// define the root element that declares the schema and target namespace it is validating.
@@ -205,7 +205,7 @@ namespace UnrealBuildTool
 
 			LoadData();
 
-			foreach (var ConfClass in Data.Keys)
+			foreach (Type ConfClass in Data.Keys)
 			{
 				Load(ConfClass);
 			}
@@ -225,7 +225,7 @@ namespace UnrealBuildTool
 			{
 				if (bReadOnlyFile && File.Exists(FilePath))
 				{
-					var Attributes = File.GetAttributes(FilePath);
+					FileAttributes Attributes = File.GetAttributes(FilePath);
 					if (Attributes.HasFlag(FileAttributes.ReadOnly))
 					{
 						File.SetAttributes(FilePath, Attributes & ~FileAttributes.ReadOnly);
@@ -263,7 +263,7 @@ namespace UnrealBuildTool
 		/// </summary>
 		private static void LoadDefaults()
 		{
-			foreach (var ConfigType in GetAllConfigurationTypes())
+			foreach (Type ConfigType in GetAllConfigurationTypes())
 			{
 				InvokeIfExists(ConfigType, "LoadDefaults");
 			}
@@ -279,12 +279,12 @@ namespace UnrealBuildTool
 			/// </summary>
 			public void LoadXmlData()
 			{
-				foreach (var DataPair in DataMap)
+				foreach (KeyValuePair<FieldInfo, object> DataPair in DataMap)
 				{
 					DataPair.Key.SetValue(null, DataPair.Value);
 				}
 
-				foreach (var PropertyPair in PropertyMap)
+				foreach (KeyValuePair<PropertyInfo, object> PropertyPair in PropertyMap)
 				{
 					PropertyPair.Key.SetValue(null, PropertyPair.Value, null);
 				}
@@ -351,11 +351,11 @@ namespace UnrealBuildTool
 				const string ConfigXmlFileName = "BuildConfiguration.xml";
 
 				// Filter out non-existing
-				var ExistingLocations = new List<string>();
+				List<string> ExistingLocations = new List<string>();
 
-				foreach (var PossibleLocation in PossibleLocations)
+				foreach (string PossibleLocation in PossibleLocations)
 				{
-					var FilePath = Path.Combine(PossibleLocation, ConfigXmlFileName);
+					string FilePath = Path.Combine(PossibleLocation, ConfigXmlFileName);
 
 					if (File.Exists(FilePath))
 					{
@@ -430,7 +430,7 @@ namespace UnrealBuildTool
 				{
 					Directory.CreateDirectory(Path.GetDirectoryName(FSLocation));
 
-					var FilePath = Path.Combine(FSLocation);
+					string FilePath = Path.Combine(FSLocation);
 
 					const string TemplateContent =
 						"<?xml version=\"1.0\" encoding=\"utf-8\" ?>" +
@@ -478,7 +478,7 @@ namespace UnrealBuildTool
 				{
 					Directory.CreateDirectory(Path.GetDirectoryName(FSLocation));
 
-					var FilePath = Path.Combine(FSLocation);
+					string FilePath = Path.Combine(FSLocation);
 
 					OverwriteIfDifferent(FilePath, GetDefaultXML(), true);
 					bExists = true;
@@ -516,7 +516,7 @@ namespace UnrealBuildTool
 			 *	priority. Not defined classes and fields are left alone.
 			 */
 
-			var UE4EnginePath = new FileInfo(Path.Combine(Path.GetDirectoryName(Assembly.GetEntryAssembly().GetOriginalLocation()), "..", "..")).FullName;
+			string UE4EnginePath = new FileInfo(Path.Combine(Path.GetDirectoryName(Assembly.GetEntryAssembly().GetOriginalLocation()), "..", "..")).FullName;
 
 			ConfigLocationHierarchy = new XmlConfigLocation[]
 			{
@@ -543,7 +543,7 @@ namespace UnrealBuildTool
 		/// </summary>
 		private static void LoadData()
 		{
-			foreach (var PossibleConfigLocation in ConfigLocationHierarchy)
+			foreach (XmlConfigLocation PossibleConfigLocation in ConfigLocationHierarchy)
 			{
 				if (!PossibleConfigLocation.bExists)
 				{
@@ -567,7 +567,7 @@ namespace UnrealBuildTool
 		/// </summary>
 		private static void CreateUserXmlConfigTemplate()
 		{
-			foreach (var PossibleConfigLocation in ConfigLocationHierarchy)
+			foreach (XmlConfigLocation PossibleConfigLocation in ConfigLocationHierarchy)
 			{
 				if (!PossibleConfigLocation.IfShouldCreateFile())
 				{
@@ -583,12 +583,12 @@ namespace UnrealBuildTool
 		/// </summary>
 		private static XmlSchema ReadConfigSchema()
 		{
-			var Settings = new XmlReaderSettings();
+			XmlReaderSettings Settings = new XmlReaderSettings();
 			Settings.ValidationType = ValidationType.DTD;
 
-			using (var SR = new StringReader(File.ReadAllText(GetXSDPath(), Encoding.UTF8)))
+			using (StringReader SR = new StringReader(File.ReadAllText(GetXSDPath(), Encoding.UTF8)))
 			{
-				using (var XR = XmlReader.Create(SR, Settings))
+				using (XmlReader XR = XmlReader.Create(SR, Settings))
 				{
 					return XmlSchema.Read(XR, (object Sender, System.Xml.Schema.ValidationEventArgs EventArgs) =>
 					{
@@ -622,10 +622,10 @@ namespace UnrealBuildTool
 		/// <param name="ConfigurationXmlPath">The path to the file with values.</param>
 		private static void LoadDataFromFile(string ConfigurationXmlPath)
 		{
-			var ConfigDocument = new XmlDocument();
-			var NS = new XmlNamespaceManager(ConfigDocument.NameTable);
+			XmlDocument ConfigDocument = new XmlDocument();
+			XmlNamespaceManager NS = new XmlNamespaceManager(ConfigDocument.NameTable);
 			NS.AddNamespace("ns", "https://www.unrealengine.com/BuildConfiguration");
-			var ReaderSettings = new XmlReaderSettings();
+			XmlReaderSettings ReaderSettings = new XmlReaderSettings();
 
 			ReaderSettings.ValidationEventHandler += (object Sender, System.Xml.Schema.ValidationEventArgs EventArgs) =>
 			{
@@ -636,13 +636,13 @@ namespace UnrealBuildTool
 			ReaderSettings.ValidationType = ValidationType.Schema;
 			ReaderSettings.Schemas.Add(GetConfigSchema());
 
-			using (var SR = new StringReader(File.ReadAllText(ConfigurationXmlPath, Encoding.UTF8)))
+			using (StringReader SR = new StringReader(File.ReadAllText(ConfigurationXmlPath, Encoding.UTF8)))
 			{
-				var Reader = XmlReader.Create(SR, ReaderSettings);
+				XmlReader Reader = XmlReader.Create(SR, ReaderSettings);
 				ConfigDocument.Load(Reader);
 			}
 
-			var XmlClasses = ConfigDocument.DocumentElement.SelectNodes("/ns:Configuration/*", NS);
+			XmlNodeList XmlClasses = ConfigDocument.DocumentElement.SelectNodes("/ns:Configuration/*", NS);
 
 			if (XmlClasses.Count == 0)
 			{
@@ -650,7 +650,7 @@ namespace UnrealBuildTool
 				{
 					ConfigDocument.DocumentElement.SetAttribute("xmlns", "https://www.unrealengine.com/BuildConfiguration");
 
-					var NSDoc = new XmlDocument();
+					XmlDocument NSDoc = new XmlDocument();
 
 					NSDoc.LoadXml(ConfigDocument.OuterXml);
 
@@ -669,7 +669,7 @@ namespace UnrealBuildTool
 
 			foreach (XmlNode XmlClass in XmlClasses)
 			{
-				var ClassType = Type.GetType("UnrealBuildTool." + XmlClass.Name);
+				Type ClassType = Type.GetType("UnrealBuildTool." + XmlClass.Name);
 
 				if (ClassType == null)
 				{
@@ -691,7 +691,7 @@ namespace UnrealBuildTool
 					Data.Add(ClassType, ClassData);
 				}
 
-				var XmlFields = XmlClass.SelectNodes("*");
+				XmlNodeList XmlFields = XmlClass.SelectNodes("*");
 
 				foreach (XmlNode XmlField in XmlFields)
 				{
@@ -722,13 +722,13 @@ namespace UnrealBuildTool
 					if (Field.FieldType.IsArray)
 					{
 						// If the type is an array type get items for it.
-						var XmlItems = XmlField.SelectNodes("ns:Item", NS);
+						XmlNodeList XmlItems = XmlField.SelectNodes("ns:Item", NS);
 
 						// Get the C# type of the array.
-						var ItemType = Field.FieldType.GetElementType();
+						Type ItemType = Field.FieldType.GetElementType();
 
 						// Create the array according to the ItemType.
-						var OutputArray = Array.CreateInstance(ItemType, XmlItems.Count);
+						Array OutputArray = Array.CreateInstance(ItemType, XmlItems.Count);
 
 						int Id = 0;
 						foreach (XmlNode XmlItem in XmlItems)
@@ -792,7 +792,7 @@ namespace UnrealBuildTool
 			// InvariantCulture.
 
 			bool bWithCulture = true;
-			var ParseMethod = ParsingType.GetMethod("Parse", new Type[] { typeof(System.String), typeof(IFormatProvider) });
+			MethodInfo ParseMethod = ParsingType.GetMethod("Parse", new Type[] { typeof(System.String), typeof(IFormatProvider) });
 
 			if (ParseMethod == null)
 			{
@@ -805,7 +805,7 @@ namespace UnrealBuildTool
 				throw new BuildException("BuildConfiguration Loading: Parsing of the type {0} is not supported.", ParsingType.Name);
 			}
 
-			var ParametersList = new List<object> { UnparsedValue };
+			List<object> ParametersList = new List<object> { UnparsedValue };
 
 			if (bWithCulture)
 			{
@@ -885,7 +885,7 @@ namespace UnrealBuildTool
 		/// <param name="MethodName">Name of the method.</param>
 		private static void InvokeIfExists(Type ClassType, string MethodName)
 		{
-			var Method = ClassType.GetMethod(MethodName, new Type[] { }, new ParameterModifier[] { });
+			MethodInfo Method = ClassType.GetMethod(MethodName, new Type[] { }, new ParameterModifier[] { });
 
 			if (Method != null && Method.IsPublic && Method.IsStatic)
 			{
@@ -910,7 +910,7 @@ namespace UnrealBuildTool
 		/// <returns>String representation of xsd type corresponding given field type.</returns>
 		private static string GetFieldXSDType(Type FieldType)
 		{
-			var TypeMap = new Dictionary<Type, string>
+			Dictionary<Type, string> TypeMap = new Dictionary<Type, string>
 			{
 				{ typeof(string), "string" },
 				{ typeof(bool), "boolean" },
@@ -933,12 +933,12 @@ namespace UnrealBuildTool
 		/// <returns>Array field value as XML elements.</returns>
 		private static XElement[] GetArrayFieldXMLValue(FieldInfo Field)
 		{
-			var NS = XNamespace.Get("http://www.w3.org/2001/XMLSchema");
+			XNamespace NS = XNamespace.Get("http://www.w3.org/2001/XMLSchema");
 
-			var Array = Field.GetValue(null) as IEnumerable;
-			var ElementList = new List<XElement>();
+			IEnumerable Array = Field.GetValue(null) as IEnumerable;
+			List<XElement> ElementList = new List<XElement>();
 
-			foreach (var Element in Array)
+			foreach (object Element in Array)
 			{
 				ElementList.Add(new XElement(NS + "Item", GetObjectXMLValue(Field.FieldType.GetElementType(), Element)));
 			}
@@ -988,7 +988,7 @@ namespace UnrealBuildTool
 
 		private static XElement CreateXSDElementForField(FieldInfo Field)
 		{
-			var NS = XNamespace.Get("http://www.w3.org/2001/XMLSchema");
+			XNamespace NS = XNamespace.Get("http://www.w3.org/2001/XMLSchema");
 
 			if (Field.FieldType.IsArray)
 			{
@@ -1018,7 +1018,7 @@ namespace UnrealBuildTool
 
 		private static XElement CreateFieldXMLElement(FieldInfo Field)
 		{
-			var NS = XNamespace.Get("https://www.unrealengine.com/BuildConfiguration");
+			XNamespace NS = XNamespace.Get("https://www.unrealengine.com/BuildConfiguration");
 
 			if (Field.FieldType.IsArray)
 			{
@@ -1032,7 +1032,7 @@ namespace UnrealBuildTool
 
 		private static string WriteToBuffer(XmlDocument Document)
 		{
-			using (var Reader = new XmlNodeReader(Document))
+			using (XmlNodeReader Reader = new XmlNodeReader(Document))
 			{
 				Reader.MoveToContent();
 				return WriteToBuffer(XDocument.Load(Reader));
@@ -1041,16 +1041,16 @@ namespace UnrealBuildTool
 
 		private static string WriteToBuffer(XDocument Document)
 		{
-			var Settings = new XmlWriterSettings();
+			XmlWriterSettings Settings = new XmlWriterSettings();
 			Settings.Indent = true;
 			Settings.IndentChars = "\t";
 			Settings.CloseOutput = true;
 
-			using (var MS = new MemoryStream())
+			using (MemoryStream MS = new MemoryStream())
 			{
-				using (var SW = new StreamWriter(MS, Encoding.UTF8))
+				using (StreamWriter SW = new StreamWriter(MS, Encoding.UTF8))
 				{
-					using (var Writer = XmlWriter.Create(SW, Settings))
+					using (XmlWriter Writer = XmlWriter.Create(SW, Settings))
 					{
 						Document.WriteTo(Writer);
 					}
diff --git a/Engine/Source/Programs/UnrealBuildTool/WinRT/WinRTDeploy.cs b/Engine/Source/Programs/UnrealBuildTool/WinRT/WinRTDeploy.cs
--- a/Engine/Source/Programs/UnrealBuildTool/WinRT/WinRTDeploy.cs
+++ b/Engine/Source/Programs/UnrealBuildTool/WinRT/WinRTDeploy.cs
@@ -47,7 +47,7 @@ namespace UnrealBuildTool
 
 			SearchOption OptionToSearch = bInIncludeSubDirectories ? SearchOption.AllDirectories : SearchOption.TopDirectoryOnly;
 
-			var SourceDirs = new List<string>(Directory.GetDirectories(InSourceDirectory, "*.*", OptionToSearch));
+			List<string> SourceDirs = new List<string>(Directory.GetDirectories(InSourceDirectory, "*.*", OptionToSearch));
 			foreach (string SourceDir in SourceDirs)
 			{
 				string SubDir = SourceDir.Replace(InSourceDirectory, "");
@@ -55,8 +55,8 @@ namespace UnrealBuildTool
 				Directory.CreateDirectory(DestDir);
 			}
 
-			var SourceFiles = new List<string>(Directory.GetFiles(InSourceDirectory, InWildCard, OptionToSearch));
-			var DestFiles = new List<string>(Directory.GetFiles(InDestinationDirectory, InWildCard, OptionToSearch));
+			List<string> SourceFiles = new List<string>(Directory.GetFiles(InSourceDirectory, InWildCard, OptionToSearch));
+			List<string> DestFiles = new List<string>(Directory.GetFiles(InDestinationDirectory, InWildCard, OptionToSearch));
 
 			// Keep a list of the files in the source directory... without the source path
 			List<string> FilesInSource = new List<string>();
diff --git a/Engine/Source/Programs/UnrealBuildTool/WinRT/WinRTToolChain.cs b/Engine/Source/Programs/UnrealBuildTool/WinRT/WinRTToolChain.cs
--- a/Engine/Source/Programs/UnrealBuildTool/WinRT/WinRTToolChain.cs
+++ b/Engine/Source/Programs/UnrealBuildTool/WinRT/WinRTToolChain.cs
@@ -492,7 +492,7 @@ namespace UnrealBuildTool
 			// Log.TraceInformation("Compile Arguments for {0}:", ModuleName);
 			// Log.TraceInformation(Arguments);
 
-			var BuildPlatform = UEBuildPlatform.GetBuildPlatformForCPPTargetPlatform(CompileEnvironment.Config.Target.Platform);
+			UEBuildPlatform BuildPlatform = UEBuildPlatform.GetBuildPlatformForCPPTargetPlatform(CompileEnvironment.Config.Target.Platform);
 
 			// Create a compile action for each source file.
 			CPPOutput Result = new CPPOutput();
@@ -529,7 +529,7 @@ namespace UnrealBuildTool
 					string OriginalPCHHeaderDirectory = Path.GetDirectoryName(SourceFile.AbsolutePath);
 					FileArguments += string.Format(" /I \"{0}\"", OriginalPCHHeaderDirectory);
 
-					var PCHExtension = UEBuildPlatform.GetBuildPlatform(UnrealTargetPlatform.WinRT).GetBinaryExtension(UEBuildBinaryType.PrecompiledHeader);
+					string PCHExtension = UEBuildPlatform.GetBuildPlatform(UnrealTargetPlatform.WinRT).GetBinaryExtension(UEBuildBinaryType.PrecompiledHeader);
 
 					// Add the precompiled header file to the produced items list.
 					FileItem PrecompiledHeaderFile = FileItem.GetItemByFileReference(
@@ -586,7 +586,7 @@ namespace UnrealBuildTool
 					CompileAction.StatusDescription = Path.GetFileName(SourceFile.AbsolutePath);
 				}
 
-				var ObjectFileExtension = UEBuildPlatform.GetBuildPlatform(UnrealTargetPlatform.WinRT).GetBinaryExtension(UEBuildBinaryType.Object);
+				string ObjectFileExtension = UEBuildPlatform.GetBuildPlatform(UnrealTargetPlatform.WinRT).GetBinaryExtension(UEBuildBinaryType.Object);
 
 				// Add the object file to the produced item list.
 				FileItem ObjectFile = FileItem.GetItemByFileReference(
diff --git a/Engine/Source/Programs/UnrealBuildTool/Windows/UEBuildWindows.cs b/Engine/Source/Programs/UnrealBuildTool/Windows/UEBuildWindows.cs
--- a/Engine/Source/Programs/UnrealBuildTool/Windows/UEBuildWindows.cs
+++ b/Engine/Source/Programs/UnrealBuildTool/Windows/UEBuildWindows.cs
@@ -262,8 +262,8 @@ namespace UnrealBuildTool
 				InBuildTarget.GlobalLinkEnvironment.Config.ExcludedLibraries.Add("LIBCPD");
 
 				//@todo ATL: Currently, only VSAccessor requires ATL (which is only used in editor builds)
-				// When compiling games, we do not want to include ATL - and we can't when compiling Rocket games
-				// due to VS 2012 Express not including ATL.
+				// When compiling games, we do not want to include ATL - and we can't when compiling games
+				// made with Launcher build due to VS 2012 Express not including ATL.
 				// If more modules end up requiring ATL, this should be refactored into a BuildTarget flag (bNeedsATL)
 				// that is set by the modules the target includes to allow for easier tracking.
 				// Alternatively, if VSAccessor is modified to not require ATL than we should always exclude the libraries.
@@ -610,7 +610,7 @@ namespace UnrealBuildTool
 
 			string VSPath = null;
 
-			foreach (var PossibleRegPath in PossibleRegPaths)
+			foreach (string PossibleRegPath in PossibleRegPaths)
 			{
 				VSPath = (string)Registry.GetValue(string.Format(@"HKEY_LOCAL_MACHINE\SOFTWARE\{0}\{1}.0", PossibleRegPath, VSVersion), "InstallDir", null);
 
diff --git a/Engine/Source/Programs/UnrealBuildTool/Windows/VCEnvironment.cs b/Engine/Source/Programs/UnrealBuildTool/Windows/VCEnvironment.cs
--- a/Engine/Source/Programs/UnrealBuildTool/Windows/VCEnvironment.cs
+++ b/Engine/Source/Programs/UnrealBuildTool/Windows/VCEnvironment.cs
@@ -171,7 +171,7 @@ namespace UnrealBuildTool
 			string FinalResult = null;
 			foreach (string IndividualVersion in Version.Split('|'))
 			{
-				var Result = Microsoft.Win32.Registry.GetValue(@"HKEY_CURRENT_USER\SOFTWARE\Microsoft\Microsoft SDKs\Windows\" + IndividualVersion, "InstallationFolder", null)
+				object Result = Microsoft.Win32.Registry.GetValue(@"HKEY_CURRENT_USER\SOFTWARE\Microsoft\Microsoft SDKs\Windows\" + IndividualVersion, "InstallationFolder", null)
 					?? Microsoft.Win32.Registry.GetValue(@"HKEY_LOCAL_MACHINE\SOFTWARE\Wow6432Node\Microsoft\Microsoft SDKs\Windows\" + IndividualVersion, "InstallationFolder", null)
 					?? Microsoft.Win32.Registry.GetValue(@"HKEY_CURRENT_USER\SOFTWARE\Wow6432Node\Microsoft\Microsoft SDKs\Windows\" + IndividualVersion, "InstallationFolder", null);
 
@@ -231,7 +231,7 @@ namespace UnrealBuildTool
 					return string.Empty;
 			}
 			string FinalResult = string.Empty;
-			var Result = Microsoft.Win32.Registry.GetValue(@"HKEY_LOCAL_MACHINE\SOFTWARE\Microsoft\Microsoft SDKs\NETFXSDK\" + Version, "KitsInstallationFolder", null)
+			object Result = Microsoft.Win32.Registry.GetValue(@"HKEY_LOCAL_MACHINE\SOFTWARE\Microsoft\Microsoft SDKs\NETFXSDK\" + Version, "KitsInstallationFolder", null)
 					  ?? Microsoft.Win32.Registry.GetValue(@"HKEY_LOCAL_MACHINE\SOFTWARE\Wow6432Node\Microsoft\Microsoft SDKs\NETFXSDK\" + Version, "KitsInstallationFolder", null);
 
 			if (Result != null)
@@ -264,7 +264,7 @@ namespace UnrealBuildTool
 			// Based on VCVarsQueryRegistry
 			string FinalResult = null;
 			{
-				var Result = Microsoft.Win32.Registry.GetValue(@"HKEY_LOCAL_MACHINE\SOFTWARE\Wow6432Node\Microsoft\Microsoft SDKs\Windows SDKs\" + Version, "InstallationFolder", null)
+				object Result = Microsoft.Win32.Registry.GetValue(@"HKEY_LOCAL_MACHINE\SOFTWARE\Wow6432Node\Microsoft\Microsoft SDKs\Windows SDKs\" + Version, "InstallationFolder", null)
 						  ?? Microsoft.Win32.Registry.GetValue(@"HKEY_CURRENT_USER\SOFTWARE\Wow6432Node\Microsoft\Microsoft SDKs\Windows SDKs\" + Version, "InstallationFolder", null);
 				if (Result == null)
 				{
@@ -325,7 +325,7 @@ namespace UnrealBuildTool
 		static string GetVSToolPath64Bit(string BaseVSToolPath)
 		{
 			// Use the native 64-bit compiler if present, otherwise use the amd64-on-x86 compiler. VS2012 Express only includes the latter.
-			var Result = Path.Combine(BaseVSToolPath, "../../VC/bin/amd64");
+			string Result = Path.Combine(BaseVSToolPath, "../../VC/bin/amd64");
 			if (File.Exists(Path.Combine(Result, "cl.exe")))
 			{
 				return Result;
@@ -382,7 +382,7 @@ namespace UnrealBuildTool
 				throw new BuildException("Failed to find cl.exe in the default toolchain directory " + CompilerExe + ". Please verify that \"Common Tools for Visual C++ 2015\" was selected when installing Visual Studio 2015.");
 			}
 
-			var ExeVersionInfo = FileVersionInfo.GetVersionInfo(CompilerExe);
+			FileVersionInfo ExeVersionInfo = FileVersionInfo.GetVersionInfo(CompilerExe);
 			if (ExeVersionInfo == null)
 			{
 				throw new BuildException("Failed to read the version number of: " + CompilerExe);
@@ -399,7 +399,7 @@ namespace UnrealBuildTool
 			// If we were asked to use Clang, then we'll redirect the path to the compiler to the LLVM installation directory
 			if (WindowsPlatform.bCompileWithClang && WindowsPlatform.bAllowClangLinker)
 			{
-				var Result = Path.Combine(Environment.GetFolderPath(Environment.SpecialFolder.ProgramFiles), "LLVM", "bin", "lld.exe");
+				string Result = Path.Combine(Environment.GetFolderPath(Environment.SpecialFolder.ProgramFiles), "LLVM", "bin", "lld.exe");
 				if (!File.Exists(Result))
 				{
 					Result = Path.Combine(Environment.GetFolderPath(Environment.SpecialFolder.ProgramFilesX86), "LLVM", "bin", "lld.exe");
diff --git a/Engine/Source/Programs/UnrealBuildTool/Windows/VCToolChain.cs b/Engine/Source/Programs/UnrealBuildTool/Windows/VCToolChain.cs
--- a/Engine/Source/Programs/UnrealBuildTool/Windows/VCToolChain.cs
+++ b/Engine/Source/Programs/UnrealBuildTool/Windows/VCToolChain.cs
@@ -979,7 +979,7 @@ namespace UnrealBuildTool
 
 		public override CPPOutput CompileCPPFiles(UEBuildTarget Target, CPPEnvironment CompileEnvironment, List<FileItem> SourceFiles, string ModuleName)
 		{
-			var EnvVars = VCEnvironment.SetEnvironment(CompileEnvironment.Config.Target.Platform, bSupportWindowsXP);
+			VCEnvironment EnvVars = VCEnvironment.SetEnvironment(CompileEnvironment.Config.Target.Platform, bSupportWindowsXP);
 
 			StringBuilder SharedArguments = new StringBuilder();
 			AppendCLArguments_Global(CompileEnvironment, EnvVars, SharedArguments);
@@ -1058,11 +1058,11 @@ namespace UnrealBuildTool
 			foreach (string Definition in CompileEnvironment.Config.Definitions)
 			{
 				// Escape all quotation marks so that they get properly passed with the command line.
-				var DefinitionArgument = Definition.Contains("\"") ? Definition.Replace("\"", "\\\"") : Definition;
+				string DefinitionArgument = Definition.Contains("\"") ? Definition.Replace("\"", "\\\"") : Definition;
 				AddDefinition(SharedArguments, DefinitionArgument);
 			}
 
-			var BuildPlatform = UEBuildPlatform.GetBuildPlatformForCPPTargetPlatform(CompileEnvironment.Config.Target.Platform);
+			UEBuildPlatform BuildPlatform = UEBuildPlatform.GetBuildPlatformForCPPTargetPlatform(CompileEnvironment.Config.Target.Platform);
 
 			// Create a compile action for each source file.
 			CPPOutput Result = new CPPOutput();
@@ -1113,7 +1113,7 @@ namespace UnrealBuildTool
 					string OriginalPCHHeaderDirectory = Path.GetDirectoryName(SourceFile.AbsolutePath);
 					AddIncludePath(FileArguments, OriginalPCHHeaderDirectory);
 
-					var PrecompiledFileExtension = UEBuildPlatform.GetBuildPlatform(UnrealTargetPlatform.Win64).GetBinaryExtension(UEBuildBinaryType.PrecompiledHeader);
+					string PrecompiledFileExtension = UEBuildPlatform.GetBuildPlatform(UnrealTargetPlatform.Win64).GetBinaryExtension(UEBuildBinaryType.PrecompiledHeader);
 					// Add the precompiled header file to the produced items list.
 					FileItem PrecompiledHeaderFile = FileItem.GetItemByFileReference(
 						FileReference.Combine(
@@ -1207,7 +1207,7 @@ namespace UnrealBuildTool
 
 				if (bEmitsObjectFile)
 				{
-					var ObjectFileExtension = UEBuildPlatform.GetBuildPlatform(UnrealTargetPlatform.Win64).GetBinaryExtension(UEBuildBinaryType.Object);
+					string ObjectFileExtension = UEBuildPlatform.GetBuildPlatform(UnrealTargetPlatform.Win64).GetBinaryExtension(UEBuildBinaryType.Object);
 					// Add the object file to the produced item list.
 					FileItem ObjectFile = FileItem.GetItemByFileReference(
 						FileReference.Combine(
@@ -1311,9 +1311,9 @@ namespace UnrealBuildTool
 				if (BuildConfiguration.bRunUnrealCodeAnalyzer)
 				{
 					string UCAMode = BuildConfiguration.bUCACheckUObjectThreadSafety ? @"-CheckThreadSafety " : @"-CreateIncludeFiles ";
-					var ObjectFileExtension = BuildConfiguration.bUCACheckUObjectThreadSafety ? @".tsc" : @".includes";
+					string ObjectFileExtension = BuildConfiguration.bUCACheckUObjectThreadSafety ? @".tsc" : @".includes";
 					FileItem ObjectFile = FileItem.GetItemByFileReference(FileReference.Combine(CompileEnvironment.Config.OutputDirectory, Path.GetFileName(SourceFile.AbsolutePath) + ObjectFileExtension));
-					var ClangPath = System.IO.Path.Combine(CompileAction.WorkingDirectory, @"ThirdParty", @"llvm", @"3.6.2", @"bin", @"vs2015", @"x86", @"release", @"clang++.exe");
+					string ClangPath = System.IO.Path.Combine(CompileAction.WorkingDirectory, @"ThirdParty", @"llvm", @"3.6.2", @"bin", @"vs2015", @"x86", @"release", @"clang++.exe");
 					UnrealCodeAnalyzerArguments = UCAMode + SourceFile.AbsolutePath + @" -OutputFile=""" + ObjectFile.AbsolutePath + @""" -- " + ClangPath + @" --driver-mode=cl ";
 				}
 
@@ -1362,11 +1362,11 @@ namespace UnrealBuildTool
 
 		public override CPPOutput CompileRCFiles(UEBuildTarget Target, CPPEnvironment Environment, List<FileItem> RCFiles)
 		{
-			var EnvVars = VCEnvironment.SetEnvironment(Environment.Config.Target.Platform, bSupportWindowsXP);
+			VCEnvironment EnvVars = VCEnvironment.SetEnvironment(Environment.Config.Target.Platform, bSupportWindowsXP);
 
 			CPPOutput Result = new CPPOutput();
 
-			var BuildPlatform = UEBuildPlatform.GetBuildPlatformForCPPTargetPlatform(Environment.Config.Target.Platform);
+			UEBuildPlatform BuildPlatform = UEBuildPlatform.GetBuildPlatformForCPPTargetPlatform(Environment.Config.Target.Platform);
 
 			foreach (FileItem RCFile in RCFiles)
 			{
@@ -1379,7 +1379,7 @@ namespace UnrealBuildTool
 				// Resource tool can run remotely if possible
 				CompileAction.bCanExecuteRemotely = true;
 
-				var Arguments = new StringBuilder();
+				StringBuilder Arguments = new StringBuilder();
 
 				if (WindowsPlatform.bCompileWithClang)
 				{
@@ -1413,7 +1413,7 @@ namespace UnrealBuildTool
 				}
 
 				// System include paths.
-				foreach (var SystemIncludePath in Environment.Config.CPPIncludeInfo.SystemIncludePaths)
+				foreach (string SystemIncludePath in Environment.Config.CPPIncludeInfo.SystemIncludePaths)
 				{
 					AddIncludePath(Arguments, SystemIncludePath);
 				}
@@ -1449,7 +1449,7 @@ namespace UnrealBuildTool
 
 		public override FileItem LinkFiles(LinkEnvironment LinkEnvironment, bool bBuildImportLibraryOnly)
 		{
-			var EnvVars = VCEnvironment.SetEnvironment(LinkEnvironment.Config.Target.Platform, bSupportWindowsXP);
+			VCEnvironment EnvVars = VCEnvironment.SetEnvironment(LinkEnvironment.Config.Target.Platform, bSupportWindowsXP);
 
 			// @todo UWP: These paths should be added in SetUpEnvironment(), not here.  Also is this actually needed for classic desktop targets or only UWP?
 			if (WindowsPlatform.Compiler == WindowsCompiler.VisualStudio2015 && WindowsPlatform.bUseWindowsSDK10)
@@ -1713,12 +1713,12 @@ namespace UnrealBuildTool
 		public override void CompileCSharpProject(CSharpEnvironment CompileEnvironment, FileReference ProjectFileName, FileReference DestinationFile)
 		{
 			// Initialize environment variables required for spawned tools.
-			var EnvVars = VCEnvironment.SetEnvironment(CompileEnvironment.EnvironmentTargetPlatform, bSupportWindowsXP);
+			VCEnvironment EnvVars = VCEnvironment.SetEnvironment(CompileEnvironment.EnvironmentTargetPlatform, bSupportWindowsXP);
 
-			var BuildProjectAction = new Action(ActionType.BuildProject);
+			Action BuildProjectAction = new Action(ActionType.BuildProject);
 
 			// Specify the source file (prerequisite) for the action
-			var ProjectFileItem = FileItem.GetExistingItemByFileReference(ProjectFileName);
+			FileItem ProjectFileItem = FileItem.GetExistingItemByFileReference(ProjectFileName);
 			if (ProjectFileItem == null)
 			{
 				throw new BuildException("Expected C# project file {0} to exist.", ProjectFileName);
@@ -1726,13 +1726,13 @@ namespace UnrealBuildTool
 
 			// Add the project and the files contained to the prerequisites.
 			BuildProjectAction.PrerequisiteItems.Add(ProjectFileItem);
-			var ProjectFile = new VCSharpProjectFile(ProjectFileName);
-			var ProjectPreReqs = ProjectFile.GetCSharpDependencies();
-			var ProjectFolder = ProjectFileName.Directory;
+			VCSharpProjectFile ProjectFile = new VCSharpProjectFile(ProjectFileName);
+			List<string> ProjectPreReqs = ProjectFile.GetCSharpDependencies();
+			DirectoryReference ProjectFolder = ProjectFileName.Directory;
 			foreach (string ProjectPreReqRelativePath in ProjectPreReqs)
 			{
 				FileReference ProjectPreReqAbsolutePath = FileReference.Combine(ProjectFolder, ProjectPreReqRelativePath);
-				var ProjectPreReqFileItem = FileItem.GetExistingItemByFileReference(ProjectPreReqAbsolutePath);
+				FileItem ProjectPreReqFileItem = FileItem.GetExistingItemByFileReference(ProjectPreReqAbsolutePath);
 				if (ProjectPreReqFileItem == null)
 				{
 					throw new BuildException("Expected C# dependency {0} to exist.", ProjectPreReqAbsolutePath);
@@ -1783,7 +1783,7 @@ namespace UnrealBuildTool
 			// matches what UBT extracted from the vcvars*.bat using SetEnvironmentVariablesFromBatchFile().  We'll use the variables we
 			// extracted to populate the project file's list of include paths
 			// @todo projectfiles: Should we only do this for VC++ platforms?
-			var IncludePaths = Environment.GetEnvironmentVariable("INCLUDE");
+			string IncludePaths = Environment.GetEnvironmentVariable("INCLUDE");
 			if (!String.IsNullOrEmpty(IncludePaths) && !IncludePaths.EndsWith(";"))
 			{
 				IncludePaths += ";";
@@ -1806,7 +1806,7 @@ namespace UnrealBuildTool
 		/// </summary>
 		protected static void ClangCompilerOutputFormatter(object sender, DataReceivedEventArgs e)
 		{
-			var Output = e.Data;
+			string Output = e.Data;
 			if (Output == null)
 			{
 				return;
diff --git a/Engine/Source/Programs/UnrealGameSync/UnrealGameSync/Forms/MainWindow.cs b/Engine/Source/Programs/UnrealGameSync/UnrealGameSync/Forms/MainWindow.cs
--- a/Engine/Source/Programs/UnrealGameSync/UnrealGameSync/Forms/MainWindow.cs
+++ b/Engine/Source/Programs/UnrealGameSync/UnrealGameSync/Forms/MainWindow.cs
@@ -1953,7 +1953,7 @@ namespace UnrealGameSync
 					BuildListContextMenu_Sync.Font = new Font(SystemFonts.MenuFont, bIsCurrentChange? FontStyle.Regular : FontStyle.Bold);
 					BuildListContextMenu_SyncContentOnly.Visible = !bIsBusy && ShouldSyncPrecompiledEditor;
 					BuildListContextMenu_SyncOnlyThisChange.Visible = !bIsBusy && !bIsCurrentChange && ContextMenuChange.Number > Workspace.CurrentChangeNumber && Workspace.CurrentChangeNumber != -1;
-					BuildListContextMenu_Build.Visible = !bIsBusy && bIsCurrentChange && !ShouldSyncPrecompiledEditor;
+					BuildListContextMenu_Build.Visible = !bIsBusy && bIsCurrentChange && !ShouldSyncPrecompiledEditor && OptionsContextMenu_UseIncrementalBuilds.Checked;
 					BuildListContextMenu_Rebuild.Visible = !bIsBusy && bIsCurrentChange && !ShouldSyncPrecompiledEditor;
 					BuildListContextMenu_GenerateProjectFiles.Visible = !bIsBusy && bIsCurrentChange;
 					BuildListContextMenu_LaunchEditor.Visible = !bIsBusy && ContextMenuChange.Number == Workspace.CurrentChangeNumber;
@@ -2239,7 +2239,7 @@ namespace UnrealGameSync
 
 		private void BuildListContextMenu_SyncContentOnly_Click(object sender, EventArgs e)
 		{
-			StartWorkspaceUpdate(ContextMenuChange.Number, WorkspaceUpdateOptions.Sync);
+			StartWorkspaceUpdate(ContextMenuChange.Number, WorkspaceUpdateOptions.Sync | WorkspaceUpdateOptions.SkipShaders);
 		}
 
 		private void BuildListContextMenu_SyncOnlyThisChange_Click(object sender, EventArgs e)
@@ -2983,5 +2983,14 @@ namespace UnrealGameSync
 			Settings.bKeepInTray ^= true;
 			Settings.Save();
 		}
+
+		private void BuildList_KeyDown(object Sender, KeyEventArgs Args)
+		{
+			if(Args.Control && Args.KeyCode == Keys.C && BuildList.SelectedItems.Count > 0)
+			{
+				int SelectedChange = ((PerforceChangeSummary)BuildList.SelectedItems[0].Tag).Number;
+				Clipboard.SetText(String.Format("{0}", SelectedChange));
+			}
+		}
 	}
 }
diff --git a/Engine/Source/Programs/UnrealGameSync/UnrealGameSync/PerforceMonitor.cs b/Engine/Source/Programs/UnrealGameSync/UnrealGameSync/PerforceMonitor.cs
--- a/Engine/Source/Programs/UnrealGameSync/UnrealGameSync/PerforceMonitor.cs
+++ b/Engine/Source/Programs/UnrealGameSync/UnrealGameSync/PerforceMonitor.cs
@@ -255,7 +255,7 @@ namespace UnrealGameSync
 			// If there's something to check for, find all the content changes after this changelist
 			if(bRequiresUpdate)
 			{
-				string[] CodeExtensions = { ".cs", ".h", ".cpp" };
+				string[] CodeExtensions = { ".cs", ".h", ".cpp", ".usf" };
 
 				// Find all the content changes in this range. Include a few extra changes in case there have been more changes submitted since the last query (it seems -m is much faster than specifying a changelist range)
 				List<PerforceChangeSummary> CodeChanges;
diff --git a/Engine/Source/Programs/UnrealGameSync/UnrealGameSync/Properties/AssemblyInfo.cs b/Engine/Source/Programs/UnrealGameSync/UnrealGameSync/Properties/AssemblyInfo.cs
--- a/Engine/Source/Programs/UnrealGameSync/UnrealGameSync/Properties/AssemblyInfo.cs
+++ b/Engine/Source/Programs/UnrealGameSync/UnrealGameSync/Properties/AssemblyInfo.cs
@@ -34,5 +34,5 @@ using System.Runtime.InteropServices;
 // You can specify all the values or you can default the Build and Revision Numbers 
 // by using the '*' as shown below:
 // [assembly: AssemblyVersion("1.0.*")]
-[assembly: AssemblyVersion("1.66.*")]
-[assembly: AssemblyFileVersion("1.66.0.0")]
+[assembly: AssemblyVersion("1.69.*")]
+[assembly: AssemblyFileVersion("1.69.0.0")]
diff --git a/Engine/Source/Programs/UnrealGameSync/UnrealGameSync/UserSettings.cs b/Engine/Source/Programs/UnrealGameSync/UnrealGameSync/UserSettings.cs
--- a/Engine/Source/Programs/UnrealGameSync/UnrealGameSync/UserSettings.cs
+++ b/Engine/Source/Programs/UnrealGameSync/UnrealGameSync/UserSettings.cs
@@ -285,7 +285,7 @@ namespace UnrealGameSync
 			GeneralSection.SetValue("OpenSolutionAfterSync", bOpenSolutionAfterSync);
 			GeneralSection.SetValue("ShowLogWindow", bShowLogWindow);
 			GeneralSection.SetValue("AutoResolveConflicts", bAutoResolveConflicts);
-			GeneralSection.SetValue("UseIncrementalBuilds", bUseIncrementalBuilds);
+			GeneralSection.SetValue("IncrementalBuilds", bUseIncrementalBuilds);
 			GeneralSection.SetValue("ShowLocalTimes", bShowLocalTimes);
 			GeneralSection.SetValue("LastProjectFileName", LastProjectFileName);
 			GeneralSection.SetValue("KeepInTray", bKeepInTray);
diff --git a/Engine/Source/Programs/UnrealGameSync/UnrealGameSync/Workspace.cs b/Engine/Source/Programs/UnrealGameSync/UnrealGameSync/Workspace.cs
--- a/Engine/Source/Programs/UnrealGameSync/UnrealGameSync/Workspace.cs
+++ b/Engine/Source/Programs/UnrealGameSync/UnrealGameSync/Workspace.cs
@@ -26,6 +26,7 @@ namespace UnrealGameSync
 		ScheduledBuild = 0x80,
 		RunAfterSync = 0x100,
 		OpenSolutionAfterSync = 0x200,
+		SkipShaders = 0x400,
 	}
 
 	enum WorkspaceUpdateResult
@@ -266,6 +267,10 @@ namespace UnrealGameSync
 					Filter.Exclude("..." + BuildVersionFileName);
 					Filter.Exclude("..." + VersionHeaderFileName);
 					Filter.Exclude("..." + ObjectVersionFileName);
+					if(Context.Options.HasFlag(WorkspaceUpdateOptions.SkipShaders))
+					{
+						Filter.Exclude("*.usf");
+					}
 					SyncFiles.RemoveAll(x => !Filter.Matches(x));
 
 					// Sync them all
@@ -322,7 +327,7 @@ namespace UnrealGameSync
 
 						// Find the last code change before this changelist. For consistency in versioning between local builds and precompiled binaries, we need to use the last submitted code changelist as our version number.
 						List<PerforceChangeSummary> CodeChanges;
-						if(!Perforce.FindChanges(new string[]{ ".cs", ".h", ".cpp" }.SelectMany(x => SyncPaths.Select(y => String.Format("{0}{1}@<={2}", y, x, PendingChangeNumber))), 1, out CodeChanges, Log))
+						if(!Perforce.FindChanges(new string[]{ ".cs", ".h", ".cpp", ".usf" }.SelectMany(x => SyncPaths.Select(y => String.Format("{0}{1}@<={2}", y, x, PendingChangeNumber))), 1, out CodeChanges, Log))
 						{
 							StatusMessage = String.Format("Couldn't determine last code changelist before CL {0}.", PendingChangeNumber);
 							return WorkspaceUpdateResult.FailedToSync;
diff --git a/Engine/Source/Runtime/Core/Public/Misc/EngineBuildSettings.h b/Engine/Source/Runtime/Core/Public/Misc/EngineBuildSettings.h
--- a/Engine/Source/Runtime/Core/Public/Misc/EngineBuildSettings.h
+++ b/Engine/Source/Runtime/Core/Public/Misc/EngineBuildSettings.h
@@ -22,7 +22,7 @@ public:
 	static bool IsSourceDistribution();
 
 	/**
-	 * @return True if a given engine distribution contains source (as opposed to, say, Rocket builds)
+	 * @return True if a given engine distribution contains source (as opposed to, say, Launcher builds)
 	 */
 	static bool IsSourceDistribution(const FString& RootDir);
 };
diff --git a/Engine/Source/Runtime/Core/Public/Modules/ModuleManager.h b/Engine/Source/Runtime/Core/Public/Modules/ModuleManager.h
--- a/Engine/Source/Runtime/Core/Public/Modules/ModuleManager.h
+++ b/Engine/Source/Runtime/Core/Public/Modules/ModuleManager.h
@@ -722,7 +722,7 @@ class FDefaultGameModuleImpl
 
 /**
  * Macro for declaring the GIsDebugGame variable for monolithic development builds. NB: This define, and the UE_BUILD_DEVELOPMENT_WITH_DEBUGGAME defines like it, should NEVER be 
- * directly used or defined for engine code, because it prevents sharing the same build products with the development build (important for Rocket build sizes). In modular builds, 
+ * directly used or defined for engine code, because it prevents sharing the same build products with the development build (important for Launcher build sizes). In modular builds, 
  * DebugGame modules will be loaded by specifying the -debug parameter on the command-line.
  */
 #if IS_MONOLITHIC && UE_BUILD_DEVELOPMENT
diff --git a/Engine/Source/Runtime/Projects/Public/ProjectDescriptor.h b/Engine/Source/Runtime/Projects/Public/ProjectDescriptor.h
--- a/Engine/Source/Runtime/Projects/Public/ProjectDescriptor.h
+++ b/Engine/Source/Runtime/Projects/Public/ProjectDescriptor.h
@@ -41,7 +41,7 @@ struct PROJECTS_API FProjectDescriptor
 	 * open a project with a different version of the editor and need the upgrade/downgrade UI flow. The normal engine 
 	 * version doesn't work for those purposes, because you can have multiple 4.x branches in various states on one machine.
 	 *
-	 * For Rocket users, this field gets set to something stable like "4.7" or "4.8", so you can swap projects and game binaries 
+	 * For Launcher users, this field gets set to something stable like "4.7" or "4.8", so you can swap projects and game binaries 
 	 * between users, and it'll automatically work on any platform or machine regardless of where the engine is installed. You 
 	 * can only have one binary release of each major engine version installed at once.
 	 * 
