Improve per target toolchain settings.
Fix project generation issues.

diff --git a/Engine/Source/Programs/UnrealBuildTool/Configuration/UEBuildTarget.cs b/Engine/Source/Programs/UnrealBuildTool/Configuration/UEBuildTarget.cs
--- a/Engine/Source/Programs/UnrealBuildTool/Configuration/UEBuildTarget.cs
+++ b/Engine/Source/Programs/UnrealBuildTool/Configuration/UEBuildTarget.cs
@@ -3830,6 +3830,7 @@ namespace UnrealBuildTool
 			GlobalCompileEnvironment.Config.Target.Architecture = ActiveArchitecture;
 			GlobalLinkEnvironment.Config.Target.Architecture = ActiveArchitecture;
 
+			Rules.ConfigureToolchain(TargetInfo);
 
 			// Set up the platform-specific environment.
 			PlatformContext.SetUpEnvironment(this);
diff --git a/Engine/Source/Programs/UnrealBuildTool/System/RulesCompiler.cs b/Engine/Source/Programs/UnrealBuildTool/System/RulesCompiler.cs
--- a/Engine/Source/Programs/UnrealBuildTool/System/RulesCompiler.cs
+++ b/Engine/Source/Programs/UnrealBuildTool/System/RulesCompiler.cs
@@ -766,6 +766,16 @@ namespace UnrealBuildTool
 			return false;
 		}
 
+		/// <summary>
+		/// Give the target an opportunity to override toolchain settings
+		/// </summary>
+		/// <param name="Target">The target currently being setup</param>
+		/// <returns>true if successful, false if not</returns>
+		public virtual bool ConfigureToolchain(TargetInfo Target)
+		{
+			return true;
+		}
+
 		/// <summary>
 		/// Get the supported platforms for this target
 		/// </summary>
