Don't update toolchains when generating project files.

diff --git a/Engine/Source/Programs/UnrealBuildTool/Configuration/UEBuildTarget.cs b/Engine/Source/Programs/UnrealBuildTool/Configuration/UEBuildTarget.cs
--- a/Engine/Source/Programs/UnrealBuildTool/Configuration/UEBuildTarget.cs
+++ b/Engine/Source/Programs/UnrealBuildTool/Configuration/UEBuildTarget.cs
@@ -3830,7 +3830,10 @@ namespace UnrealBuildTool
 			GlobalCompileEnvironment.Config.Target.Architecture = ActiveArchitecture;
 			GlobalLinkEnvironment.Config.Target.Architecture = ActiveArchitecture;
 
-			Rules.ConfigureToolchain(TargetInfo);
+			if (!ProjectFileGenerator.bGenerateProjectFiles)
+			{
+				Rules.ConfigureToolchain(TargetInfo);
+			}
 
 			// Set up the platform-specific environment.
 			PlatformContext.SetUpEnvironment(this);
