VS2013 toolchain check for /Zo readded

diff --git a/Engine/Source/Programs/UnrealBuildTool/Windows/VCToolChain.cs b/Engine/Source/Programs/UnrealBuildTool/Windows/VCToolChain.cs
--- a/Engine/Source/Programs/UnrealBuildTool/Windows/VCToolChain.cs
+++ b/Engine/Source/Programs/UnrealBuildTool/Windows/VCToolChain.cs
@@ -528,7 +528,15 @@ namespace UnrealBuildTool
 				{
 					// Allow optimized code to be debugged more easily.  This makes PDBs a bit larger, but doesn't noticeably affect
 					// compile times.  The executable code is not affected at all by this switch, only the debugging information.
-					Arguments.Append(" /Zo");
+					if (EnvVars.CLExeVersion >= new Version("18.0.30723"))
+					{
+						// VC2013 Update 3 has a new flag for doing this
+						Arguments.Append(" /Zo");
+					}
+					else
+					{
+						Arguments.Append(" /d2Zi+");
+					}
 				}
 			}
 
