fix for compile failure during precompiile step when Vulkan SDK is not present

diff --git a/Engine/Source/Programs/UnrealBuildTool/Configuration/UEBuildTarget.cs b/Engine/Source/Programs/UnrealBuildTool/Configuration/UEBuildTarget.cs
--- a/Engine/Source/Programs/UnrealBuildTool/Configuration/UEBuildTarget.cs
+++ b/Engine/Source/Programs/UnrealBuildTool/Configuration/UEBuildTarget.cs
@@ -3212,8 +3212,15 @@ namespace UnrealBuildTool
 						if (bCanPrecompile)
 						{
 							UEBuildModule Module = FindOrCreateModuleByName(FilteredModuleName);
-							Module.RecursivelyCreateModules(this);
-							PrecompiledModules.Add(Module);
+							try
+							{
+								Module.RecursivelyCreateModules(this);
+								PrecompiledModules.Add(Module);
+							}
+							catch(BuildException)
+							{
+
+							}
 						}
 					}
 				}
diff --git a/Engine/Source/Runtime/VulkanRHI/VulkanRHI.Build.cs b/Engine/Source/Runtime/VulkanRHI/VulkanRHI.Build.cs
--- a/Engine/Source/Runtime/VulkanRHI/VulkanRHI.Build.cs
+++ b/Engine/Source/Runtime/VulkanRHI/VulkanRHI.Build.cs
@@ -61,5 +61,9 @@ public class VulkanRHI : ModuleRules
                     );
             }
 		}
+		else
+		{
+			throw new UnrealBuildTool.BuildException("Vulkan SDK not found");
+		}
     }
 }
