fix for GenerateProjects failure

diff --git a/Engine/Source/Programs/UnrealBuildTool/Configuration/UEBuildPlatform.cs b/Engine/Source/Programs/UnrealBuildTool/Configuration/UEBuildPlatform.cs
--- a/Engine/Source/Programs/UnrealBuildTool/Configuration/UEBuildPlatform.cs
+++ b/Engine/Source/Programs/UnrealBuildTool/Configuration/UEBuildPlatform.cs
@@ -90,6 +90,15 @@ namespace UnrealBuildTool
 		{
 		}
 
+		public virtual bool ValidatePrecompiledModule(string Module)
+		{
+			if (Module == "VulkanRHI" || Module == "VulkanShaderFormat")
+			{
+				return !String.IsNullOrEmpty(Environment.GetEnvironmentVariable("VK_SDK_PATH"));
+			}
+			return true;
+		}
+
 		/// <summary>
 		/// Setup the target environment for building
 		/// </summary>
@@ -496,8 +505,6 @@ namespace UnrealBuildTool
 			}
 		}
 
-
-
 		/// <summary>
 		/// If this platform can be compiled with XGE
 		/// </summary>
diff --git a/Engine/Source/Programs/UnrealBuildTool/Configuration/UEBuildTarget.cs b/Engine/Source/Programs/UnrealBuildTool/Configuration/UEBuildTarget.cs
--- a/Engine/Source/Programs/UnrealBuildTool/Configuration/UEBuildTarget.cs
+++ b/Engine/Source/Programs/UnrealBuildTool/Configuration/UEBuildTarget.cs
@@ -3144,7 +3144,11 @@ namespace UnrealBuildTool
 							string RelativeFileName = ModuleFileName.MakeRelativeTo(UnrealBuildTool.EngineSourceDirectory);
 							if (ExcludeFolders.All(x => RelativeFileName.IndexOf(x, StringComparison.InvariantCultureIgnoreCase) == -1) && !PrecompiledModules.Any(x => x.Name == ModuleName))
 							{
-								FilteredModuleNames.Add(ModuleName);
+								// allow platforms to exclude modules
+								if (PlatformContext.ValidatePrecompiledModule(ModuleName))
+								{
+									FilteredModuleNames.Add(ModuleName);
+								}
 							}
 						}
 					}
@@ -3212,15 +3216,8 @@ namespace UnrealBuildTool
 						if (bCanPrecompile)
 						{
 							UEBuildModule Module = FindOrCreateModuleByName(FilteredModuleName);
-							try
-							{
-								Module.RecursivelyCreateModules(this);
-								PrecompiledModules.Add(Module);
-							}
-							catch(BuildException)
-							{
-
-							}
+							Module.RecursivelyCreateModules(this);
+							PrecompiledModules.Add(Module);
 						}
 					}
 				}
diff --git a/Engine/Source/Runtime/VulkanRHI/VulkanRHI.Build.cs b/Engine/Source/Runtime/VulkanRHI/VulkanRHI.Build.cs
--- a/Engine/Source/Runtime/VulkanRHI/VulkanRHI.Build.cs
+++ b/Engine/Source/Runtime/VulkanRHI/VulkanRHI.Build.cs
@@ -61,9 +61,5 @@ public class VulkanRHI : ModuleRules
                     );
             }
 		}
-		else
-		{
-			throw new UnrealBuildTool.BuildException("Vulkan SDK not found");
-		}
     }
 }
