Revert unnecessary changes to handle Vulkan SDK from CL 2912965 and CL 2912914. Throwing a BuildException from a module rules constructor is already supported.

diff --git a/Engine/Source/Programs/UnrealBuildTool/Configuration/UEBuildPlatform.cs b/Engine/Source/Programs/UnrealBuildTool/Configuration/UEBuildPlatform.cs
--- a/Engine/Source/Programs/UnrealBuildTool/Configuration/UEBuildPlatform.cs
+++ b/Engine/Source/Programs/UnrealBuildTool/Configuration/UEBuildPlatform.cs
@@ -90,15 +90,6 @@ namespace UnrealBuildTool
 		{
 		}
 
-		public virtual bool ValidatePrecompiledModule(string Module)
-		{
-			if (Module == "VulkanRHI" || Module == "VulkanShaderFormat")
-			{
-				return !String.IsNullOrEmpty(Environment.GetEnvironmentVariable("VK_SDK_PATH"));
-			}
-			return true;
-		}
-
 		/// <summary>
 		/// Setup the target environment for building
 		/// </summary>
@@ -505,6 +496,8 @@ namespace UnrealBuildTool
 			}
 		}
 
+
+
 		/// <summary>
 		/// If this platform can be compiled with XGE
 		/// </summary>
diff --git a/Engine/Source/Programs/UnrealBuildTool/Configuration/UEBuildTarget.cs b/Engine/Source/Programs/UnrealBuildTool/Configuration/UEBuildTarget.cs
--- a/Engine/Source/Programs/UnrealBuildTool/Configuration/UEBuildTarget.cs
+++ b/Engine/Source/Programs/UnrealBuildTool/Configuration/UEBuildTarget.cs
@@ -3144,11 +3144,7 @@ namespace UnrealBuildTool
 							string RelativeFileName = ModuleFileName.MakeRelativeTo(UnrealBuildTool.EngineSourceDirectory);
 							if (ExcludeFolders.All(x => RelativeFileName.IndexOf(x, StringComparison.InvariantCultureIgnoreCase) == -1) && !PrecompiledModules.Any(x => x.Name == ModuleName))
 							{
-								// allow platforms to exclude modules
-								if (PlatformContext.ValidatePrecompiledModule(ModuleName))
-								{
-									FilteredModuleNames.Add(ModuleName);
-								}
+								FilteredModuleNames.Add(ModuleName);
 							}
 						}
 					}
