fix for XGE being enabled for Android

diff --git a/Engine/Source/Programs/UnrealBuildTool/Android/UEBuildAndroid.cs b/Engine/Source/Programs/UnrealBuildTool/Android/UEBuildAndroid.cs
--- a/Engine/Source/Programs/UnrealBuildTool/Android/UEBuildAndroid.cs
+++ b/Engine/Source/Programs/UnrealBuildTool/Android/UEBuildAndroid.cs
@@ -248,7 +248,7 @@ namespace UnrealBuildTool
 
 		public override bool CanUseXGE()
 		{
-			return true;
+			return false;
 		}
 
 		public override string GetBinaryExtension(UEBuildBinaryType InBinaryType)
