Switched 2012 toolchain passthrough back off.
Changed 2012 project generation flag to -2012unsupported.

diff --git a/Engine/Source/Programs/UnrealBuildTool/System/VCProject.cs b/Engine/Source/Programs/UnrealBuildTool/System/VCProject.cs
--- a/Engine/Source/Programs/UnrealBuildTool/System/VCProject.cs
+++ b/Engine/Source/Programs/UnrealBuildTool/System/VCProject.cs
@@ -1099,10 +1099,6 @@ namespace UnrealBuildTool
 
 					// @todo UWP: For the MS toolchains, if an override was set for project generation, push that into the build strings to override the build toolchain as well
 					string BuildToolOverride = "";
-                    if (UnrealBuildTool.CommandLineContains("-2012"))
-                    {
-                        BuildToolOverride = " -2012";
-                    }
 					if (UnrealBuildTool.CommandLineContains("-2013"))
 					{
 						BuildToolOverride = " -2013";
diff --git a/Engine/Source/Programs/UnrealBuildTool/System/VCProjectFileGenerator.cs b/Engine/Source/Programs/UnrealBuildTool/System/VCProjectFileGenerator.cs
--- a/Engine/Source/Programs/UnrealBuildTool/System/VCProjectFileGenerator.cs
+++ b/Engine/Source/Programs/UnrealBuildTool/System/VCProjectFileGenerator.cs
@@ -121,7 +121,7 @@ namespace UnrealBuildTool
 
 		public enum VCProjectFileFormat
         {
-            VisualStudio2012,
+            VisualStudio2012, // Unsupported
 			VisualStudio2013,
 			VisualStudio2015,
 		}
@@ -169,7 +169,7 @@ namespace UnrealBuildTool
 		static VCProjectFileGenerator()
 		{
 			switch (WindowsPlatform.Compiler)
-            {
+			{
 				case WindowsCompiler.VisualStudio2013:
 					ProjectFileFormat = VCProjectFileFormat.VisualStudio2013;
 					break;
@@ -201,7 +201,7 @@ namespace UnrealBuildTool
 			// Default the project file format to whatever UBT is using
 			{
 				switch (WindowsPlatform.Compiler)
-                {
+				{
 					case WindowsCompiler.VisualStudio2013:
 						ProjectFileFormat = VCProjectFileFormat.VisualStudio2013;
 						break;
@@ -215,7 +215,8 @@ namespace UnrealBuildTool
 			{
 				switch (CurArgument.ToUpperInvariant())
                 {
-                    case "-2012":
+					// Visual Studio 2012 project generation has been deprecated and is no longer maintained.
+                    case "-2012UNSUPPORTED":
                         ProjectFileFormat = VCProjectFileFormat.VisualStudio2012;
                         break;
 
@@ -269,7 +270,7 @@ namespace UnrealBuildTool
 			// Certain platforms override the project file format because their debugger add-ins may not yet support the latest
 			// version of Visual Studio.  This is their chance to override that.
 			// ...but only if the user didn't override this via the command-line.
-            if (!UnrealBuildTool.CommandLineContains("-2015") && !UnrealBuildTool.CommandLineContains("-2013") && !UnrealBuildTool.CommandLineContains("-2012"))
+            if (!UnrealBuildTool.CommandLineContains("-2015") && !UnrealBuildTool.CommandLineContains("-2013") && !UnrealBuildTool.CommandLineContains("-2012unsupported"))
 			{
 				foreach (UnrealTargetPlatform SupportedPlatform in SupportedPlatforms)
 				{
