Conditional SSF compilation.

diff --git a/Engine/Source/Developer/MeshUtilities/MeshUtilities.Build.cs b/Engine/Source/Developer/MeshUtilities/MeshUtilities.Build.cs
--- a/Engine/Source/Developer/MeshUtilities/MeshUtilities.Build.cs
+++ b/Engine/Source/Developer/MeshUtilities/MeshUtilities.Build.cs
@@ -31,25 +31,28 @@ public class MeshUtilities : ModuleRules
 			}
 		);
 
-        AddEngineThirdPartyPrivateStaticDependencies(Target, "nvTriStrip");
-		AddEngineThirdPartyPrivateStaticDependencies(Target, "ForsythTriOptimizer");
-		AddEngineThirdPartyPrivateStaticDependencies(Target, "MeshSimplifier");
-		AddEngineThirdPartyPrivateStaticDependencies(Target, "MikkTSpace");
+        AddThirdPartyPrivateStaticDependencies(Target, "nvTriStrip");
+		AddThirdPartyPrivateStaticDependencies(Target, "ForsythTriOptimizer");
+		AddThirdPartyPrivateStaticDependencies(Target, "MeshSimplifier");
+		AddThirdPartyPrivateStaticDependencies(Target, "MikkTSpace");
 
 		if ((Target.Platform == UnrealTargetPlatform.Win64) || (Target.Platform == UnrealTargetPlatform.Win32))
 		{
-			AddEngineThirdPartyPrivateStaticDependencies(Target, "DX9");
+			AddThirdPartyPrivateStaticDependencies(Target, "DX9");
 		}
 
 		if (UEBuildConfiguration.bCompileSimplygon == true)
 		{
-			DynamicallyLoadedModuleNames.AddRange(
-				new string[] {
-                "SimplygonSwarm"
+			AddThirdPartyPrivateDynamicDependencies(Target, "SimplygonMeshReduction");
+            
+            if (UEBuildConfiguration.bCompileSimplygonSSF == true)
+            {
+                DynamicallyLoadedModuleNames.AddRange(
+                    new string[] {
+                    "SimplygonSwarm"
+                }
+                );
             }
-			);
-
-			AddEngineThirdPartyPrivateDynamicDependencies(Target, "SimplygonMeshReduction");
 		}
 	}
 }
diff --git a/Engine/Source/Editor/GameProjectGeneration/Private/GameProjectUtils.cpp b/Engine/Source/Editor/GameProjectGeneration/Private/GameProjectUtils.cpp
--- a/Engine/Source/Editor/GameProjectGeneration/Private/GameProjectUtils.cpp
+++ b/Engine/Source/Editor/GameProjectGeneration/Private/GameProjectUtils.cpp
@@ -3129,10 +3129,10 @@ bool GameProjectUtils::HasDefaultBuildSettings(const FName InPlatformInfoName)
 	// first check default build settings for all platforms
 	TArray<FString> BoolKeys, IntKeys, StringKeys, BuildKeys;
 	BuildKeys.Add(TEXT("bCompileApex")); BuildKeys.Add(TEXT("bCompileBox2D")); BuildKeys.Add(TEXT("bCompileICU"));
-	BuildKeys.Add(TEXT("bCompileSimplygon")); BuildKeys.Add(TEXT("bCompi8leLeanAndMeanUE")); BuildKeys.Add(TEXT("bIncludeADO"));
-	BuildKeys.Add(TEXT("bCompileRecast")); BuildKeys.Add(TEXT("bCompileSpeedTree")); BuildKeys.Add(TEXT("bCompileWithPluginSupport"));
-	BuildKeys.Add(TEXT("bCompilePhysXVehicle")); BuildKeys.Add(TEXT("bCompileFreeType")); BuildKeys.Add(TEXT("bCompileForSize"));
-	BuildKeys.Add(TEXT("bCompileCEF3"));
+	BuildKeys.Add(TEXT("bCompileSimplygon")); BuildKeys.Add(TEXT("bCompileSimplygonSSF")); BuildKeys.Add(TEXT("bCompi8leLeanAndMeanUE"));
+	BuildKeys.Add(TEXT("bIncludeADO"));	BuildKeys.Add(TEXT("bCompileRecast")); BuildKeys.Add(TEXT("bCompileSpeedTree"));
+	BuildKeys.Add(TEXT("bCompileWithPluginSupport")); BuildKeys.Add(TEXT("bCompilePhysXVehicle")); BuildKeys.Add(TEXT("bCompileFreeType"));
+	BuildKeys.Add(TEXT("bCompileForSize"));	BuildKeys.Add(TEXT("bCompileCEF3"));
 
 	const PlatformInfo::FPlatformInfo* const PlatInfo = PlatformInfo::FindPlatformInfo(InPlatformInfoName);
 	check(PlatInfo);
diff --git a/Engine/Source/Programs/UnrealBuildTool/Android/UEBuildAndroid.cs b/Engine/Source/Programs/UnrealBuildTool/Android/UEBuildAndroid.cs
--- a/Engine/Source/Programs/UnrealBuildTool/Android/UEBuildAndroid.cs
+++ b/Engine/Source/Programs/UnrealBuildTool/Android/UEBuildAndroid.cs
@@ -86,6 +86,7 @@ namespace UnrealBuildTool
 			UEBuildConfiguration.bBuildEditor = false;
 			UEBuildConfiguration.bBuildDeveloperTools = false;
 			UEBuildConfiguration.bCompileSimplygon = false;
+            UEBuildConfiguration.bCompileSimplygonSSF = false;
 
 			UEBuildConfiguration.bCompileRecast = true;
 
@@ -180,6 +181,7 @@ namespace UnrealBuildTool
 			InBuildTarget.GlobalLinkEnvironment.Config.AdditionalLibraries.Add("cxa_demangle");
 
 			UEBuildConfiguration.bCompileSimplygon = false;
+            UEBuildConfiguration.bCompileSimplygonSSF = false;
 			BuildConfiguration.bDeployAfterCompile = true;
 
 			//!! RONIN need to fix this
diff --git a/Engine/Source/Programs/UnrealBuildTool/Configuration/UEBuildConfiguration.cs b/Engine/Source/Programs/UnrealBuildTool/Configuration/UEBuildConfiguration.cs
--- a/Engine/Source/Programs/UnrealBuildTool/Configuration/UEBuildConfiguration.cs
+++ b/Engine/Source/Programs/UnrealBuildTool/Configuration/UEBuildConfiguration.cs
@@ -93,6 +93,12 @@ namespace UnrealBuildTool
 		[XmlConfig]
 		public static bool bCompileSimplygon;
 
+        /// <summary>
+        /// Whether we should compile in support for Simplygon's SSF library or not.
+        /// </summary>
+        [XmlConfig]
+        public static bool bCompileSimplygonSSF;
+
 		/// <summary>
 		/// Whether we should compile in support for Steam OnlineSubsystem or not. [RCL] FIXME 2014-Apr-17: bCompileSteamOSS means "bHasSteamworksInstalled" for some code, these meanings need to be untangled
 		/// </summary>
@@ -341,6 +347,7 @@ namespace UnrealBuildTool
 			bForceBuildTargetPlatforms = false;
 			bForceBuildShaderFormats = false;
 			bCompileSimplygon = true;
+            bCompileSimplygonSSF = true;
 			bCompileLeanAndMeanUE = false;
 			bCompileAgainstEngine = true;
 			bCompileAgainstCoreUObject = true;
@@ -388,6 +395,10 @@ namespace UnrealBuildTool
 				&& Directory.Exists(UEBuildConfiguration.UEThirdPartySourceDirectory + "NotForLicensees") == true
 				&& Directory.Exists(UEBuildConfiguration.UEThirdPartySourceDirectory + "NotForLicensees/Simplygon") == true
 				&& Directory.Exists("Developer/SimplygonMeshReduction") == true;
+
+            bCompileSimplygonSSF = bCompileSimplygonSSF
+            && Directory.Exists(UEBuildConfiguration.UEThirdPartySourceDirectory + "NotForLicensees") == true
+            && Directory.Exists(UEBuildConfiguration.UEThirdPartySourceDirectory + "NotForLicensees/SSF");
 		}
 
 		/// <summary>
@@ -402,6 +413,7 @@ namespace UnrealBuildTool
 				bBuildEditor = false;
 				bBuildDeveloperTools = false;
 				bCompileSimplygon = false;
+                bCompileSimplygonSSF = false;
 				bCompileSpeedTree = false;
 			}
 		}
diff --git a/Engine/Source/Programs/UnrealBuildTool/Configuration/UEBuildPlatform.cs b/Engine/Source/Programs/UnrealBuildTool/Configuration/UEBuildPlatform.cs
--- a/Engine/Source/Programs/UnrealBuildTool/Configuration/UEBuildPlatform.cs
+++ b/Engine/Source/Programs/UnrealBuildTool/Configuration/UEBuildPlatform.cs
@@ -201,6 +201,12 @@ namespace UnrealBuildTool
 					UEBuildConfiguration.bCompileSimplygon = bValue;
 				}
 
+                bValue = UEBuildConfiguration.bCompileSimplygonSSF;
+                if (Ini.GetBool("/Script/BuildSettings.BuildSettings", "bCompileSimplygonSSF", out bValue))
+                {
+                    UEBuildConfiguration.bCompileSimplygonSSF = bValue;
+                }
+
 				bValue = UEBuildConfiguration.bCompileLeanAndMeanUE;
 				if (Ini.GetBool("/Script/BuildSettings.BuildSettings", "bCompileLeanAndMeanUE", out bValue))
 				{
@@ -781,7 +787,7 @@ namespace UnrealBuildTool
 		public virtual bool HasDefaultBuildConfig(UnrealTargetPlatform Platform, DirectoryReference ProjectDirectoryName)
 		{
 			string[] BoolKeys = new string[] {
-				"bCompileApex", "bCompileBox2D", "bCompileICU", "bCompileSimplygon", 
+				"bCompileApex", "bCompileBox2D", "bCompileICU", "bCompileSimplygon", "bCompileSimplygonSSF",
 				"bCompileLeanAndMeanUE", "bIncludeADO", "bCompileRecast", "bCompileSpeedTree", 
 				"bCompileWithPluginSupport", "bCompilePhysXVehicle", "bCompileFreeType", 
 				"bCompileForSize", "bCompileCEF3"
diff --git a/Engine/Source/Programs/UnrealBuildTool/HTML5/UEBuildHTML5.cs b/Engine/Source/Programs/UnrealBuildTool/HTML5/UEBuildHTML5.cs
--- a/Engine/Source/Programs/UnrealBuildTool/HTML5/UEBuildHTML5.cs
+++ b/Engine/Source/Programs/UnrealBuildTool/HTML5/UEBuildHTML5.cs
@@ -67,6 +67,7 @@ namespace UnrealBuildTool
 			UEBuildConfiguration.bCompilePhysX = true;
 			UEBuildConfiguration.bRuntimePhysicsCooking = false;
 			UEBuildConfiguration.bCompileSimplygon = false;
+            UEBuildConfiguration.bCompileSimplygonSSF = false;
 			UEBuildConfiguration.bCompileForSize = true;
 		}
 
diff --git a/Engine/Source/Programs/UnrealBuildTool/IOS/UEBuildIOS.cs b/Engine/Source/Programs/UnrealBuildTool/IOS/UEBuildIOS.cs
--- a/Engine/Source/Programs/UnrealBuildTool/IOS/UEBuildIOS.cs
+++ b/Engine/Source/Programs/UnrealBuildTool/IOS/UEBuildIOS.cs
@@ -262,6 +262,7 @@ namespace UnrealBuildTool
 			UEBuildConfiguration.bCompileAPEX = false;
 			UEBuildConfiguration.bRuntimePhysicsCooking = false;
 			UEBuildConfiguration.bCompileSimplygon = false;
+            UEBuildConfiguration.bCompileSimplygonSSF = false;
 			UEBuildConfiguration.bBuildDeveloperTools = false;
 
 			// we currently don't have any simulator libs for PhysX
diff --git a/Engine/Source/Programs/UnrealBuildTool/Linux/UEBuildLinux.cs b/Engine/Source/Programs/UnrealBuildTool/Linux/UEBuildLinux.cs
--- a/Engine/Source/Programs/UnrealBuildTool/Linux/UEBuildLinux.cs
+++ b/Engine/Source/Programs/UnrealBuildTool/Linux/UEBuildLinux.cs
@@ -97,6 +97,7 @@ namespace UnrealBuildTool
 		public override void ValidateBuildConfiguration(CPPTargetConfiguration Configuration, CPPTargetPlatform Platform, bool bCreateDebugInfo)
 		{
 			UEBuildConfiguration.bCompileSimplygon = false;
+            UEBuildConfiguration.bCompileSimplygonSSF = false;
 		}
 
 		/// <summary>
@@ -144,6 +145,7 @@ namespace UnrealBuildTool
 			if (InBuildTarget.GlobalCompileEnvironment.Config.Definitions.Contains("USE_NULL_RHI=1"))
 			{
 				UEBuildConfiguration.bCompileSimplygon = false;
+                UEBuildConfiguration.bCompileSimplygonSSF = false;
 			}
 
 			if (InBuildTarget.TargetType == TargetRules.TargetType.Server)
diff --git a/Engine/Source/Programs/UnrealBuildTool/Mac/UEBuildMac.cs b/Engine/Source/Programs/UnrealBuildTool/Mac/UEBuildMac.cs
--- a/Engine/Source/Programs/UnrealBuildTool/Mac/UEBuildMac.cs
+++ b/Engine/Source/Programs/UnrealBuildTool/Mac/UEBuildMac.cs
@@ -58,6 +58,7 @@ namespace UnrealBuildTool
 		public override void ResetBuildConfiguration(UnrealTargetConfiguration Configuration)
 		{
 			UEBuildConfiguration.bCompileSimplygon = false;
+            UEBuildConfiguration.bCompileSimplygonSSF = false;
 		}
 
 		public override void ValidateBuildConfiguration(CPPTargetConfiguration Configuration, CPPTargetPlatform Platform, bool bCreateDebugInfo)
diff --git a/Engine/Source/Programs/UnrealBuildTool/System/UnrealBuildTool.cs b/Engine/Source/Programs/UnrealBuildTool/System/UnrealBuildTool.cs
--- a/Engine/Source/Programs/UnrealBuildTool/System/UnrealBuildTool.cs
+++ b/Engine/Source/Programs/UnrealBuildTool/System/UnrealBuildTool.cs
@@ -806,6 +806,7 @@ namespace UnrealBuildTool
 				else if (LowercaseArg == "-nosimplygon")
 				{
 					UEBuildConfiguration.bCompileSimplygon = false;
+                    UEBuildConfiguration.bCompileSimplygonSSF = false;
 				}
 				else if (LowercaseArg == "-nospeedtree")
 				{
diff --git a/Engine/Source/Programs/UnrealBuildTool/UWP/UEBuildUWP.cs b/Engine/Source/Programs/UnrealBuildTool/UWP/UEBuildUWP.cs
--- a/Engine/Source/Programs/UnrealBuildTool/UWP/UEBuildUWP.cs
+++ b/Engine/Source/Programs/UnrealBuildTool/UWP/UEBuildUWP.cs
@@ -139,6 +139,7 @@ namespace UnrealBuildTool
 			if (InBuildTarget.GlobalCompileEnvironment.Config.Definitions.Contains("USE_NULL_RHI=1"))
 			{
 				UEBuildConfiguration.bCompileSimplygon = false;
+                UEBuildConfiguration.bCompileSimplygonSSF = false;
 			}
 		}
 
diff --git a/Engine/Source/Programs/UnrealBuildTool/WinRT/UEBuildWinRT.cs b/Engine/Source/Programs/UnrealBuildTool/WinRT/UEBuildWinRT.cs
--- a/Engine/Source/Programs/UnrealBuildTool/WinRT/UEBuildWinRT.cs
+++ b/Engine/Source/Programs/UnrealBuildTool/WinRT/UEBuildWinRT.cs
@@ -381,6 +381,7 @@ namespace UnrealBuildTool
 			//InBuildTarget.GlobalLinkEnvironment.Config.AdditionalLibraries.Add("dwmapi.lib");
 
 			UEBuildConfiguration.bCompileSimplygon = false;
+            UEBuildConfiguration.bCompileSimplygonSSF = false;
 
 			// For 64-bit builds, we'll forcibly ignore a linker warning with DirectInput.  This is
 			// Microsoft's recommended solution as they don't have a fixed .lib for us.
diff --git a/Engine/Source/Programs/UnrealBuildTool/Windows/UEBuildWindows.cs b/Engine/Source/Programs/UnrealBuildTool/Windows/UEBuildWindows.cs
--- a/Engine/Source/Programs/UnrealBuildTool/Windows/UEBuildWindows.cs
+++ b/Engine/Source/Programs/UnrealBuildTool/Windows/UEBuildWindows.cs
@@ -342,6 +342,7 @@ namespace UnrealBuildTool
 			if (InBuildTarget.GlobalCompileEnvironment.Config.Definitions.Contains("USE_NULL_RHI=1"))
 			{
 				UEBuildConfiguration.bCompileSimplygon = false;
+                UEBuildConfiguration.bCompileSimplygonSSF = false;
 			}
 
 			// For 64-bit builds, we'll forcibly ignore a linker warning with DirectInput.  This is
