UAT/UBT: Fix problem staging the bootstrap executable into the root of packaged games, where it searches for a path in the staging files. Receipt contained non-normalized path separators.

diff --git a/Engine/Source/Programs/UnrealBuildTool/System/TargetReceipt.cs b/Engine/Source/Programs/UnrealBuildTool/System/TargetReceipt.cs
--- a/Engine/Source/Programs/UnrealBuildTool/System/TargetReceipt.cs
+++ b/Engine/Source/Programs/UnrealBuildTool/System/TargetReceipt.cs
@@ -386,11 +386,11 @@ namespace UnrealBuildTool
 			// Replace all the variables in the paths
 			foreach (BuildProduct BuildProduct in BuildProducts)
 			{
-				BuildProduct.Path = Utils.ExpandVariables(BuildProduct.Path, Variables);
+				BuildProduct.Path = Path.GetFullPath(Utils.ExpandVariables(BuildProduct.Path, Variables));
 			}
 			foreach (RuntimeDependency RuntimeDependency in RuntimeDependencies)
 			{
-				RuntimeDependency.Path = Utils.ExpandVariables(RuntimeDependency.Path, Variables);
+				RuntimeDependency.Path = Path.GetFullPath(Utils.ExpandVariables(RuntimeDependency.Path, Variables));
 			}
 		}
 
