Copy of CL 2957764

Move creation of dSYMs for Mac to after dylib dependencies are fixed. This will hopefully solve a rare problem with dsymutil not being able to find the source .dylib file.

diff --git a/Engine/Source/Programs/UnrealBuildTool/Mac/MacToolChain.cs b/Engine/Source/Programs/UnrealBuildTool/Mac/MacToolChain.cs
--- a/Engine/Source/Programs/UnrealBuildTool/Mac/MacToolChain.cs
+++ b/Engine/Source/Programs/UnrealBuildTool/Mac/MacToolChain.cs
@@ -1181,7 +1181,7 @@ namespace UnrealBuildTool
 		/// Generates debug info for a given executable
 		/// </summary>
 		/// <param name="MachOBinary">FileItem describing the executable or dylib to generate debug info for</param>
-		public FileItem GenerateDebugInfo(FileItem MachOBinary)
+		public FileItem GenerateDebugInfo(FileItem MachOBinary, FileItem FixDylibOutputFile)
 		{
 			string BinaryPath = MachOBinary.AbsolutePath;
 			if (BinaryPath.Contains(".app"))
@@ -1223,7 +1223,7 @@ namespace UnrealBuildTool
 				ToolchainDir,
 				InputFile.AbsolutePath,
 				DestFile.AbsolutePath);
-			GenDebugAction.PrerequisiteItems.Add(InputFile);
+			GenDebugAction.PrerequisiteItems.Add(FixDylibOutputFile);
 			GenDebugAction.ProducedItems.Add(DestFile);
 			GenDebugAction.CommandDescription = "";
 			GenDebugAction.StatusDescription = "Generating " + Path.GetFileName(BinaryPath);
@@ -1697,7 +1697,15 @@ namespace UnrealBuildTool
 			// For Mac, generate the dSYM file if the config file is set to do so
 			if ((BuildConfiguration.bGeneratedSYMFile == true || BuildConfiguration.bUsePDBFiles == true) && (!BinaryLinkEnvironment.Config.bIsBuildingLibrary || BinaryLinkEnvironment.Config.bIsBuildingDLL))
 			{
-				OutputFiles.Add(GenerateDebugInfo(Executable));
+				// We want dsyms to be created after all dylib dependencies are fixed. If FixDylibDependencies action was not created yet, save the info for later.
+				if (FixDylibOutputFile != null)
+				{
+					OutputFiles.Add(GenerateDebugInfo(Executable, FixDylibOutputFile));
+				}
+				else
+				{
+					ExecutablesThatNeedDsyms.Add(Executable);
+				}
 			}
 
 			// If building for Mac on a Mac, use actions to finalize the builds (otherwise, we use Deploy)
@@ -1711,16 +1719,26 @@ namespace UnrealBuildTool
 				return OutputFiles;
 			}
 
-			FileItem FixDylibOutputFile = FixDylibDependencies(BinaryLinkEnvironment, Executable);
+			FixDylibOutputFile = FixDylibDependencies(BinaryLinkEnvironment, Executable);
 			OutputFiles.Add(FixDylibOutputFile);
 			if (!BinaryLinkEnvironment.Config.bIsBuildingConsoleApplication)
 			{
 				OutputFiles.Add(FinalizeAppBundle(BinaryLinkEnvironment, Executable, FixDylibOutputFile));
 			}
 
+			// Add dsyms that we couldn't add before FixDylibDependencies action was created
+			foreach (FileItem Exe in ExecutablesThatNeedDsyms)
+			{
+				OutputFiles.Add(GenerateDebugInfo(Exe, FixDylibOutputFile));
+			}
+			ExecutablesThatNeedDsyms.Clear();
+
 			return OutputFiles;
 		}
 
+		private FileItem FixDylibOutputFile = null;
+		private List<FileItem> ExecutablesThatNeedDsyms = new List<FileItem>();
+
 		public override void StripSymbols(string SourceFileName, string TargetFileName)
 		{
 			SetupXcodePaths(false);
