Fixed a problem with UBT and UAT not finding UT icon on Mac

diff --git a/Engine/Source/Programs/AutomationTool/Mac/MacPlatform.Automation.cs b/Engine/Source/Programs/AutomationTool/Mac/MacPlatform.Automation.cs
--- a/Engine/Source/Programs/AutomationTool/Mac/MacPlatform.Automation.cs
+++ b/Engine/Source/Programs/AutomationTool/Mac/MacPlatform.Automation.cs
@@ -284,6 +284,10 @@ public class MacPlatform : Platform
 				if (!Directory.Exists(SourceBundlePath))
 				{
 					SourceBundlePath = CombinePaths(SC.ArchiveDirectory, "Engine", "Binaries", "Mac", ExeName + ".app");
+					if (!Directory.Exists(SourceBundlePath))
+					{
+						SourceBundlePath = CombinePaths(SC.ArchiveDirectory, "Engine", "Binaries", "Mac", "UE4.app");
+					}
 				}
 				Directory.Move(SourceBundlePath, BundlePath);
 
diff --git a/Engine/Source/Programs/UnrealBuildTool/Mac/MacToolChain.cs b/Engine/Source/Programs/UnrealBuildTool/Mac/MacToolChain.cs
--- a/Engine/Source/Programs/UnrealBuildTool/Mac/MacToolChain.cs
+++ b/Engine/Source/Programs/UnrealBuildTool/Mac/MacToolChain.cs
@@ -1008,11 +1008,17 @@ namespace UnrealBuildTool
 					bool bIsLauncherProduct = ExeName.StartsWith("EpicGamesLauncher") || ExeName.StartsWith("EpicGamesBootstrapLauncher");
 					string[] ExeNameParts = ExeName.Split('-');
 					string GameName = ExeNameParts[0];
+					FileReference UProjectFilePath = null;
 
 					if (GameName == "EpicGamesBootstrapLauncher")
 					{
 						GameName = "EpicGamesLauncher";
 					}
+					else if (GameName == "UE4" && ProjectFile != null)
+					{
+						UProjectFilePath = ProjectFile;
+						GameName = UProjectFilePath.GetFileNameWithoutAnyExtensions();
+					}
 
 					AppendMacLine(FinalizeAppBundleScript, "mkdir -p \"{0}.app/Contents/MacOS\"", ExeName);
 					AppendMacLine(FinalizeAppBundleScript, "mkdir -p \"{0}.app/Contents/Resources\"", ExeName);
@@ -1023,10 +1029,9 @@ namespace UnrealBuildTool
 					string IconName = "UE4";
 					string BundleVersion = bIsLauncherProduct ? LoadLauncherDisplayVersion() : LoadEngineDisplayVersion();
 					string EngineSourcePath = ConvertPath(Directory.GetCurrentDirectory()).Replace("$", "\\$");
-					FileReference UProjectFilePath;
 					string CustomResourcesPath = "";
 					string CustomBuildPath = "";
-					if (!UProjectInfo.TryGetProjectFileName(GameName, out UProjectFilePath))
+					if (UProjectFilePath == null && !UProjectInfo.TryGetProjectFileName(GameName, out UProjectFilePath))
 					{
 						string[] TargetFiles = Directory.GetFiles(Directory.GetCurrentDirectory(), GameName + ".Target.cs", SearchOption.AllDirectories);
 						if (TargetFiles.Length == 1)
