Reimplementing CL #2989663 by [REDACTED_PERSON] from Dev-Core:

Timestamp files are now forced to be written out with a new time, fixing iterative build times due to failed timestamp checks.

diff --git a/Engine/Source/Programs/UnrealBuildTool/System/ExternalExecution.cs b/Engine/Source/Programs/UnrealBuildTool/System/ExternalExecution.cs
--- a/Engine/Source/Programs/UnrealBuildTool/System/ExternalExecution.cs
+++ b/Engine/Source/Programs/UnrealBuildTool/System/ExternalExecution.cs
@@ -915,7 +915,7 @@ namespace UnrealBuildTool
 							AllUObjectFiles.AddRange(Module.PublicUObjectClassesHeaders.ConvertAll(Item => Item.AbsolutePath));
 							AllUObjectFiles.AddRange(Module.PublicUObjectHeaders.ConvertAll(Item => Item.AbsolutePath));
 							AllUObjectFiles.AddRange(Module.PrivateUObjectHeaders.ConvertAll(Item => Item.AbsolutePath));
-							ResponseFile.Create(TimestampFile, AllUObjectFiles);
+							ResponseFile.Create(TimestampFile, AllUObjectFiles, ResponseFile.CreateOptions.WriteEvenIfUnchanged);
 						}
 					}
 				}
diff --git a/Engine/Source/Programs/UnrealBuildTool/System/ResponseFile.cs b/Engine/Source/Programs/UnrealBuildTool/System/ResponseFile.cs
--- a/Engine/Source/Programs/UnrealBuildTool/System/ResponseFile.cs
+++ b/Engine/Source/Programs/UnrealBuildTool/System/ResponseFile.cs
@@ -9,22 +9,32 @@ namespace UnrealBuildTool
 {
 	public class ResponseFile
 	{
+		[FlagsAttribute]
+		public enum CreateOptions
+		{
+			None                 = 0x0,
+			WriteEvenIfUnchanged = 0x1
+		}
+
 		/// <summary>
 		/// Creates a file from a list of strings; each string is placed on a line in the file.
 		/// </summary>
 		/// <param name="TempFileName">Name of response file</param>
 		/// <param name="Lines">List of lines to write to the response file</param>
-		public static FileReference Create(FileReference TempFileName, List<string> Lines)
+		public static FileReference Create(FileReference TempFileName, List<string> Lines, CreateOptions Options = CreateOptions.None)
 		{
 			FileInfo TempFileInfo = new FileInfo(TempFileName.FullName);
 			if (TempFileInfo.Exists)
 			{
-				string Body = string.Join(Environment.NewLine, Lines);
-				// Reuse the existing response file if it remains unchanged
-				string OriginalBody = File.ReadAllText(TempFileName.FullName);
-				if (string.Equals(OriginalBody, Body, StringComparison.Ordinal))
+				if ((Options & CreateOptions.WriteEvenIfUnchanged) != CreateOptions.WriteEvenIfUnchanged)
 				{
-					return TempFileName;
+					string Body = string.Join(Environment.NewLine, Lines);
+					// Reuse the existing response file if it remains unchanged
+					string OriginalBody = File.ReadAllText(TempFileName.FullName);
+					if (string.Equals(OriginalBody, Body, StringComparison.Ordinal))
+					{
+						return TempFileName;
+					}
 				}
 				// Delete the existing file if it exists and requires modification
 				TempFileInfo.IsReadOnly = false;
