Remaking CL 3000142 from Dev-Build: Stop engine Shared PCHs being used for game plugins when building in DebugGame. The first modules compiled are currently plugins, and they choose to use an engine shared PCH before optimization settings are available to determine that it's not compatible.

diff --git a/Engine/Source/Programs/UnrealBuildTool/Configuration/UEBuildModule.cs b/Engine/Source/Programs/UnrealBuildTool/Configuration/UEBuildModule.cs
--- a/Engine/Source/Programs/UnrealBuildTool/Configuration/UEBuildModule.cs
+++ b/Engine/Source/Programs/UnrealBuildTool/Configuration/UEBuildModule.cs
@@ -1143,6 +1143,14 @@ namespace UnrealBuildTool
 						SharedPCHHeaderInfo LargestSharedPCHHeader = GlobalCompileEnvironment.SharedPCHHeaderFiles[LargestSharedPCHHeaderFileIndex];
 						++LargestSharedPCHHeader.NumModulesUsingThisPCH;
 
+						// Don't allow game modules to use engine PCHs in DebugGame - the optimization settings aren't correct. 
+						// @todo: we should be creating shared PCHs ahead of time, and only using them if our settings match. as it is, the first modules compiled
+						// (which are currently plugins) get to call the shots for how the shared PCH gets built, and that might be a game plugin built in debug...
+						if(Target.Configuration == UnrealTargetConfiguration.DebugGame && SharedPCHHeaderFile.Reference.IsUnderDirectory(UnrealBuildTool.EngineDirectory) && !RulesFile.IsUnderDirectory(UnrealBuildTool.EngineDirectory))
+						{
+							SharedPCHModuleName = null;
+							SharedPCHHeaderFile = null;
+						}
 					}
 					else
 					{
