Fix single-file compile matching multiple files that end in the same filename, and prevent deployment of executables when running a single-file compile.

diff --git a/Engine/Source/Programs/UnrealBuildTool/System/ActionGraph.cs b/Engine/Source/Programs/UnrealBuildTool/System/ActionGraph.cs
--- a/Engine/Source/Programs/UnrealBuildTool/System/ActionGraph.cs
+++ b/Engine/Source/Programs/UnrealBuildTool/System/ActionGraph.cs
@@ -328,10 +328,10 @@ namespace UnrealBuildTool
 				AllActions.RemoveAll(x => x.ActionType != ActionType.Compile);
 
 				// Check all of the leftover compilation actions for the one we want... that one is always outdated.
+				FileItem SingleFileToCompile = FileItem.GetExistingItemByPath(BuildConfiguration.SingleFileToCompile);
 				foreach (Action Action in AllActions)
 				{
-					// Slightly hacky.  We know a compiling X:/Blah/Actor.cpp ends up with StatusDescription == "Actor.cpp".
-					bool bIsSingleFileAction = BuildConfiguration.SingleFileToCompile.EndsWith(Action.StatusDescription.ToLowerInvariant());
+					bool bIsSingleFileAction = Action.PrerequisiteItems.Contains(SingleFileToCompile);
 					OutdatedActionDictionary[Action] = bIsSingleFileAction;
 				}
 
diff --git a/Engine/Source/Programs/UnrealBuildTool/System/UnrealBuildTool.cs b/Engine/Source/Programs/UnrealBuildTool/System/UnrealBuildTool.cs
--- a/Engine/Source/Programs/UnrealBuildTool/System/UnrealBuildTool.cs
+++ b/Engine/Source/Programs/UnrealBuildTool/System/UnrealBuildTool.cs
@@ -1348,7 +1348,7 @@ namespace UnrealBuildTool
 
 							// If we build w/ bXGEExport true, we didn't REALLY build at this point, 
 							// so don't bother with doing the PrepTargetForDeployment call. 
-							if ((Result == ECompilationResult.Succeeded) && (BuildConfiguration.bDeployAfterCompile == true) && (BuildConfiguration.bXGEExport == false) &&
+							if ((Result == ECompilationResult.Succeeded) && (BuildConfiguration.bDeployAfterCompile == true) && (BuildConfiguration.bXGEExport == false) && String.IsNullOrEmpty(BuildConfiguration.SingleFileToCompile) && 
 								(UEBuildConfiguration.bGenerateManifest == false) && (UEBuildConfiguration.bGenerateExternalFileList == false) && (UEBuildConfiguration.bCleanProject == false) && (UEBuildConfiguration.bListBuildFolders == false))
 							{
 								List<TargetDescriptor> TargetDescs = UEBuildTarget.ParseTargetCommandLine(Arguments, ref ProjectFile);
