Copying //UE4/WEX-Staging to Dev-Main (//UE4/Dev-Main)

diff --git a/Engine/Source/Programs/UnrealBuildTool/Android/UEDeployAndroid.cs b/Engine/Source/Programs/UnrealBuildTool/Android/UEDeployAndroid.cs
--- a/Engine/Source/Programs/UnrealBuildTool/Android/UEDeployAndroid.cs
+++ b/Engine/Source/Programs/UnrealBuildTool/Android/UEDeployAndroid.cs
@@ -1565,6 +1565,25 @@ namespace UnrealBuildTool
 			return false;
 		}
 
+        private bool RequiresOBB(bool bDisallowPackageInAPK, string OBBLocation)
+        {
+            if (bDisallowPackageInAPK)
+            {
+                Log.TraceInformation("APK contains data.");
+                return false;
+            }
+            else if (!String.IsNullOrEmpty(Environment.GetEnvironmentVariable("uebp_LOCAL_ROOT")))
+            {
+                Log.TraceInformation("On build machine.");
+                return true;
+            }
+            else
+            {
+                Log.TraceInformation("Looking for OBB.");
+                return File.Exists(OBBLocation);
+            }
+        }
+
 		private void MakeApk(AndroidToolChain ToolChain, string ProjectName, string ProjectDirectory, string OutputPath, string EngineDirectory, bool bForDistribution, string CookFlavor, bool bMakeSeparateApks, bool bIncrementalPackage, bool bDisallowPackagingDataInApk)
 		{
 			Log.TraceInformation("\n===={0}====PREPARING TO MAKE APK=================================================================", DateTime.Now.ToString());
@@ -1739,7 +1758,7 @@ namespace UnrealBuildTool
 			{
 				BuildList = from Arch in Arches
 							from GPUArch in GPUArchitectures
-							let manifest = GenerateManifest(ToolChain, ProjectName, bForDistribution, bPackageDataInsideApk, GameBuildFilesPath, bDisallowPackagingDataInApk ? false : File.Exists(ObbFileLocation), bDisableVerifyOBBOnStartUp, Arch, GPUArch)
+							let manifest = GenerateManifest(ToolChain, ProjectName, bForDistribution, bPackageDataInsideApk, GameBuildFilesPath, RequiresOBB(bDisallowPackagingDataInApk, ObbFileLocation), bDisableVerifyOBBOnStartUp, Arch, GPUArch)
 							select Tuple.Create(Arch, GPUArch, manifest);
 			}
 			else
@@ -1747,7 +1766,7 @@ namespace UnrealBuildTool
 				BuildList = from Arch in Arches
 							from GPUArch in GPUArchitectures
 							let manifestFile = Path.Combine(IntermediateAndroidPath, Arch + "_" + GPUArch + "_AndroidManifest.xml")
-							let manifest = GenerateManifest(ToolChain, ProjectName, bForDistribution, bPackageDataInsideApk, GameBuildFilesPath, bDisallowPackagingDataInApk ? false : File.Exists(ObbFileLocation), bDisableVerifyOBBOnStartUp, Arch, GPUArch)
+							let manifest = GenerateManifest(ToolChain, ProjectName, bForDistribution, bPackageDataInsideApk, GameBuildFilesPath, RequiresOBB(bDisallowPackagingDataInApk, ObbFileLocation), bDisableVerifyOBBOnStartUp, Arch, GPUArch)
 							let OldManifest = File.Exists(manifestFile) ? File.ReadAllText(manifestFile) : ""
 							where manifest != OldManifest
 							select Tuple.Create(Arch, GPUArch, manifest);
diff --git a/Engine/Source/Runtime/Engine/Private/Animation/AnimSingleNodeInstanceProxy.cpp b/Engine/Source/Runtime/Engine/Private/Animation/AnimSingleNodeInstanceProxy.cpp
--- a/Engine/Source/Runtime/Engine/Private/Animation/AnimSingleNodeInstanceProxy.cpp
+++ b/Engine/Source/Runtime/Engine/Private/Animation/AnimSingleNodeInstanceProxy.cpp
@@ -31,7 +31,10 @@ bool FAnimSingleNodeInstanceProxy::Evaluate(FPoseContext& Output)
 		= false;
 #endif
 
-	if (CurrentAsset != NULL)
+	if (CurrentAsset != NULL &&
+//@HSL_BEGIN - CCL - Seeing crashes due to processing animations that have begun to be destroyed
+		!CurrentAsset->HasAnyFlags(RF_BeginDestroyed))
+//@HSL_END
 	{
 		//@TODO: animrefactor: Seems like more code duplication than we need
 		if (UBlendSpaceBase* BlendSpace = Cast<UBlendSpaceBase>(CurrentAsset))
diff --git a/Engine/Source/Runtime/Engine/Private/PlayerController.cpp b/Engine/Source/Runtime/Engine/Private/PlayerController.cpp
--- a/Engine/Source/Runtime/Engine/Private/PlayerController.cpp
+++ b/Engine/Source/Runtime/Engine/Private/PlayerController.cpp
@@ -91,6 +91,9 @@ APlayerController::APlayerController(const FObjectInitializer& ObjectInitializer
 	bIsLocalPlayerController = false;
 
 	ClickEventKeys.Add(EKeys::LeftMouseButton);
+	// @pjs - adding right mouse button for HSL change, this may just need to be part of a derived class
+	ClickEventKeys.Add(EKeys::RightMouseButton);
+	// @pjs
 
 	if (RootComponent)
 	{
@@ -3813,7 +3816,9 @@ void APlayerController::TickPlayerInput(const float DeltaSeconds, const bool bGa
 			UGameViewportClient* ViewportClient = LocalPlayer->ViewportClient;
 
 			// Only send mouse hit events if we're directly over the viewport.
-			if ( ViewportClient && ViewportClient->GetGameViewportWidget().IsValid() && ViewportClient->GetGameViewportWidget()->IsDirectlyHovered() )
+//@HSL_BEGIN_PF - CCL - This code does not work, not sure why yet.  Returns false unless you click, briefly
+			//if ( ViewportClient && ViewportClient->GetGameViewportWidget().IsValid() && ViewportClient->GetGameViewportWidget()->IsDirectlyHovered() )
+//@HSL_END
 			{
 				if ( LocalPlayer->ViewportClient->GetMousePosition(MousePosition) )
 				{
diff --git a/Engine/Source/Runtime/Engine/Private/Slate/SceneViewport.cpp b/Engine/Source/Runtime/Engine/Private/Slate/SceneViewport.cpp
--- a/Engine/Source/Runtime/Engine/Private/Slate/SceneViewport.cpp
+++ b/Engine/Source/Runtime/Engine/Private/Slate/SceneViewport.cpp
@@ -1091,7 +1091,9 @@ void FSceneViewport::ResizeFrame(uint32 NewWindowSizeX, uint32 NewWindowSizeY, E
 			}
 
 			// Resize viewport
-			FVector2D ViewportSize(SizeX, SizeY);
+// @HSL_CHANGE_BEGIN - [REDACTED_EMAIL] - 5/31/2016 - Fixing windowed mode
+			FVector2D ViewportSize = WindowToResize->GetWindowSizeFromClientSize(FVector2D(SizeX, SizeY));
+// @HSL_CHANGE_END - [REDACTED_EMAIL] - 5/31/2016
 			FVector2D NewViewportSize = WindowToResize->GetViewportSize();
 
 			if (NewViewportSize != ViewportSize || NewWindowMode != WindowMode)
diff --git a/Engine/Source/Runtime/Slate/Private/Widgets/Input/SButton.cpp b/Engine/Source/Runtime/Slate/Private/Widgets/Input/SButton.cpp
--- a/Engine/Source/Runtime/Slate/Private/Widgets/Input/SButton.cpp
+++ b/Engine/Source/Runtime/Slate/Private/Widgets/Input/SButton.cpp
@@ -430,3 +430,10 @@ void SButton::SetButtonStyle(const FButtonStyle* ButtonStyle)
 	HoveredSound = Style->HoveredSlateSound;
 	PressedSound = Style->PressedSlateSound;
 }
+
+// @HSL_BEGIN - [REDACTED_EMAIL] - 5/25/2016 - Need to release the button when we invalidate this button from being a cursor captor, otherwise we won't be able to press it again until we release it
+void SButton::OnMouseCaptureLost()
+{
+	bIsPressed = false;
+}
+// @HSL_END - [REDACTED_EMAIL] - 5/25/2016
\ No newline at end of file
diff --git a/Engine/Source/Runtime/Slate/Public/Widgets/Input/SButton.h b/Engine/Source/Runtime/Slate/Public/Widgets/Input/SButton.h
--- a/Engine/Source/Runtime/Slate/Public/Widgets/Input/SButton.h
+++ b/Engine/Source/Runtime/Slate/Public/Widgets/Input/SButton.h
@@ -160,6 +160,9 @@ public:
 	virtual void OnMouseLeave( const FPointerEvent& MouseEvent ) override;
 
 	virtual bool IsInteractable() const override;
+// @HSL_BEGIN - [REDACTED_EMAIL] - 5/25/2016 - Need to release the button when we invalidate this button from being a cursor captor, otherwise we won't be able to press it again until we release it
+	virtual void OnMouseCaptureLost() override;
+// @HSL_END - [REDACTED_EMAIL] - 5/25/2016
 
 protected:
 
diff --git a/Engine/Source/Runtime/Slate/Public/Widgets/Views/SListView.h b/Engine/Source/Runtime/Slate/Public/Widgets/Views/SListView.h
--- a/Engine/Source/Runtime/Slate/Public/Widgets/Views/SListView.h
+++ b/Engine/Source/Runtime/Slate/Public/Widgets/Views/SListView.h
@@ -1053,6 +1053,13 @@ public:
 		return SelectedItems.Num();
 	}
 
+	/** Deletes all items and rebuilds the list */
+	void RebuildList()
+	{
+		WidgetGenerator.Clear();
+		RequestListRefresh();
+	}
+
 	/**
 	 * Returns a list of selected item indices, or an empty array if nothing is selected
 	 *
diff --git a/Engine/Source/Runtime/SlateCore/Private/Fonts/FontCache.cpp b/Engine/Source/Runtime/SlateCore/Private/Fonts/FontCache.cpp
--- a/Engine/Source/Runtime/SlateCore/Private/Fonts/FontCache.cpp
+++ b/Engine/Source/Runtime/SlateCore/Private/Fonts/FontCache.cpp
@@ -888,6 +888,9 @@ FCharacterList& FSlateFontCache::GetCharacterList( const FSlateFontInfo &InFontI
 	// Create a key for looking up each character
 	const FSlateFontKey FontKey( InFontInfo, FontScale );
 
+	//@HSL_BEGIN - Chance.Lyon - A critical section to help make this class thread-safe */
+	FScopeLock ScopeLock(&CacheCriticalSection);
+	//@HSL_END
 	TSharedRef< class FCharacterList >* CachedCharacterList = FontToCharacterListCache.Find( FontKey );
 
 	if( CachedCharacterList )
@@ -1006,6 +1009,9 @@ void FSlateFontCache::FlushObject( const UObject* const InObject )
 	}
 
 	bool bHasRemovedEntries = false;
+	//@HSL_BEGIN - Chance.Lyon - A critical section to help make this class thread-safe */
+	FScopeLock ScopeLock(&CacheCriticalSection);
+	//@HSL_END
 	for( auto It = FontToCharacterListCache.CreateIterator(); It; ++It )
 	{
 		if( It.Key().GetFontInfo().FontObject == InObject)
diff --git a/Engine/Source/Runtime/SlateCore/Public/Fonts/FontCache.h b/Engine/Source/Runtime/SlateCore/Public/Fonts/FontCache.h
--- a/Engine/Source/Runtime/SlateCore/Public/Fonts/FontCache.h
+++ b/Engine/Source/Runtime/SlateCore/Public/Fonts/FontCache.h
@@ -882,4 +882,9 @@ private:
 
 	/** The frame counter the last time the font cache was asked to be flushed */
 	uint64 FrameCounterLastFlushRequest;
+
+//@HSL_BEGIN - Chance.Lyon - A critical section to help make this class thread-safe */
+	/** Critical section for thread synchronization for the font cache */
+	mutable FCriticalSection CacheCriticalSection;
+//@HSL_END
 };
diff --git a/Engine/Source/Runtime/SlateCore/Public/Widgets/SWindow.h b/Engine/Source/Runtime/SlateCore/Public/Widgets/SWindow.h
--- a/Engine/Source/Runtime/SlateCore/Public/Widgets/SWindow.h
+++ b/Engine/Source/Runtime/SlateCore/Public/Widgets/SWindow.h
@@ -695,10 +695,12 @@ private:
 	/** The window's desired size takes into account the ratio between the slate units and the pixel size */
 	virtual FVector2D ComputeDesiredSize(float) const override;
 
+public:
+// @HSL_CHANGE_BEGIN - [REDACTED_EMAIL] - 5/31/2016 - need access to this to fix windowed mode
 	// For a given client size, calculate the window size required to accomodate any potential non-OS borders and tilebars
 	FVector2D GetWindowSizeFromClientSize(FVector2D InClientSize);
+// @HSL_CHANGE_END - [REDACTED_EMAIL] - 5/31/2016 - need access to this to fix windowed mode
 
-public:
 	/** @return true if this window will be focused when it is first shown */
 	inline bool IsFocusedInitially() const
 	{
diff --git a/Engine/Source/Runtime/UMG/Private/Components/WidgetComponent.cpp b/Engine/Source/Runtime/UMG/Private/Components/WidgetComponent.cpp
--- a/Engine/Source/Runtime/UMG/Private/Components/WidgetComponent.cpp
+++ b/Engine/Source/Runtime/UMG/Private/Components/WidgetComponent.cpp
@@ -19,7 +19,9 @@
 
 DECLARE_CYCLE_STAT(TEXT("3DHitTesting"), STAT_Slate3DHitTesting, STATGROUP_Slate);
 
-static const FName SharedLayerName(TEXT("WidgetComponentScreenLayer"));
+//@HSL_BEGIN - Colin.Pyle - 6-9-2016 - Giving WiddgetComponents the ability to be on multiple layers.
+//static const FName SharedLayerName(TEXT("WidgetComponentScreenLayer"));
+//@HSL_END
 
 class FWorldWidgetScreenLayer : public IGameLayer
 {
@@ -450,6 +452,10 @@ UWidgetComponent::UWidgetComponent( const FObjectInitializer& PCIP )
 	, bIsTwoSided( false )
 	, ParabolaDistortion( 0 )
 	, TickWhenOffscreen( false )
+	//@HSL_BEGIN - Colin.Pyle - 6-9-2016 - Giving WiddgetComponents the ability to be on multiple layers.
+	, SharedLayerName(TEXT("WidgetComponentScreenLayer"))
+	, LayerZOrder(-100)
+	//@HSL_END
 {
 	PrimaryComponentTick.bCanEverTick = true;
 	bTickInEditor = true;
@@ -698,7 +704,9 @@ void UWidgetComponent::TickComponent(float DeltaTime, enum ELevelTick TickType,
 #if !UE_SERVER
 	if (!IsRunningDedicatedServer())
 	{
-		static const int32 LayerZOrder = -100;
+		//@HSL_BEGIN - Colin.Pyle - 6-9-2016 - Giving WiddgetComponents the ability to be on multiple layers.
+		//static const int32 LayerZOrder = -100;
+		//@HSL_END
 
 		UpdateWidget();
 
diff --git a/Engine/Source/Runtime/UMG/Public/Components/WidgetComponent.h b/Engine/Source/Runtime/UMG/Public/Components/WidgetComponent.h
--- a/Engine/Source/Runtime/UMG/Public/Components/WidgetComponent.h
+++ b/Engine/Source/Runtime/UMG/Public/Components/WidgetComponent.h
@@ -25,12 +25,15 @@ enum class EWidgetBlendMode : uint8
 
 
 
+//@HSL_BEGIN - Colin.Pyle - 6-9-2016 - Giving WiddgetComponents the ability to be on multiple layers.
+
 /**
  * Beware! This feature is experimental and may be substantially changed or removed in future releases.
  * A 3D instance of a Widget Blueprint that can be interacted with in the world.
  */
-UCLASS(ClassGroup=Experimental, hidecategories=(Object,Activation,"Components|Activation",Sockets,Base,Lighting,LOD,Mesh), editinlinenew, meta=(BlueprintSpawnableComponent,  DevelopmentStatus=Experimental) )
+UCLASS(Blueprintable, ClassGroup=Experimental, hidecategories=(Object,Activation,"Components|Activation",Sockets,Base,Lighting,LOD,Mesh), editinlinenew, meta=(BlueprintSpawnableComponent,  DevelopmentStatus=Experimental) )
 class UMG_API UWidgetComponent : public UPrimitiveComponent
+//@HSL_END
 {
 	GENERATED_UCLASS_BODY()
 
@@ -336,6 +339,17 @@ protected:
 
 protected:
 
+	//@HSL_BEGIN - Colin.Pyle - 6-9-2016 - Giving WiddgetComponents the ability to be on multiple layers.
+
+	/** Layer Name the widget will live on */
+	UPROPERTY(EditDefaultsOnly, Category = Layers)
+	FName SharedLayerName;
+
+	/** ZOrder the layer will be created on, note this only matters on the first time a new layer is created, subsequent additions to the same layer will use the initially defined ZOrder */
+	UPROPERTY(EditDefaultsOnly, Category = Layers)
+	int32 LayerZOrder;
+	//@HSL_END
+
 	/** The grid used to find actual hit actual widgets once input has been translated to the components local space */
 	TSharedPtr<class FHittestGrid> HitTestGrid;
 	
