Copying //UE4/Dev-Platform to //UE4/Dev-Main (Source: //UE4/Dev-Platform @ 3046626)

diff --git a/Engine/Source/Programs/UnrealBuildTool/System/UnrealPluginLanguage.cs b/Engine/Source/Programs/UnrealBuildTool/System/UnrealPluginLanguage.cs
new file mode 100644
--- /dev/null
+++ b/Engine/Source/Programs/UnrealBuildTool/System/UnrealPluginLanguage.cs
@@ -0,0 +1,2969 @@
+// Copyright 1998-2016 Epic Games, Inc. All Rights Reserved.
+
+using System;
+using System.Collections.Generic;
+using System.Linq;
+using System.Text;
+using System.Xml;
+using System.Xml.Linq;
+using System.Diagnostics;
+using System.IO;
+
+namespace UnrealBuildTool
+{
+	/* UnrealPluginLanguage (UPL) is a simple XML-based language for manipulating XML and returning
+	 * strings.  It contains an <init> section which is evaluated once per architecture before any
+	 * other sections.  The state is maintained and carried forward to the next section evaluated
+	 * so the order the sections are executed matters.
+	 * 
+	 * While UPL is a general system for modifying and quering XML it is specifically used to allow
+	 * plug-ins to effect the global configuration of the package that they are a part of. For
+	 * example, this allows a plug-in to modify an Android's APK AndroidManfiest.xml file or an
+	 * IOS IPA's plist file. UBT will also query a plug-in's UPL xml file for strings to be included
+	 * in files that must be common to the package such as some .java files on Android.
+	 * 
+	 * If you need to see the instructions executed in your plugin context add the following to
+	 * enable tracing:
+	 * 
+	 *	<trace enable="true"/>
+	 *	
+	 * After this instuction all the nodes actually executed in your context will be written to the
+	 * log until you do a <trace enable="false"/>.  You can also get a dump of all the variables in
+	 * your context with this command:
+	 * 
+	 *	<dumpvars/>
+	 * 
+	 * Bool, Int, and String variable types are supported.  Any attribute may reference a variable
+	 * and will be replaced with the string equivalent before evaluation using this syntax:
+	 * 
+	 *	$B(name) = boolean variable "name"'s value
+	 *	$I(name) = integer variable "name"'s value
+	 *	$S(name) = string variable "name"'s value
+	 *	$E(name) = element variable "name"'s value
+	 * 
+	 * The following variables are initialized automatically:
+	 * 
+	 *	$S(Output) = the output returned for evaluating the section (initialized to Input)
+	 *	$S(Architecture) = target architecture (armeabi-armv7a, arm64-v8a, x86, x86_64)
+	 *	$S(PluginDir) = directory the XML file was loaded from
+	 *	$S(EngineDir) = engine directory
+	 *	$S(BuildDir) = project's platform appropriate build directory (within the Intermediate folder)
+	 *	$B(Distribution) = true if distribution build
+	 * 
+	 * Note: with the exception of the above variables, all are in the context of the plugin to
+	 * prevent namespace collision; trying to set a new value to any of the above, with the
+	 * exception of Output, will only affect the current context.
+	 * 
+	 * The following nodes allow manipulation of variables:
+	 * 
+	 *	<setBool result="" value=""/>
+	 *	<setInt result="" value=""/>
+	 *	<setString result="" value=""/>
+	 *	<setElement result="" value=""/>
+	 *	<setElement result="" xml=""/>
+	 *	
+	 * <setElement> with value creates an empty XML element with the tag set to value.
+	 * <setElement> with xml will parse the XML provided.  Remember to escape any special characters!
+	 * 
+	 * Variables may also be set from a property in an ini file:
+	 * 
+	 *	<setBoolFromProperty result="" ini="" section="" property="" default=""/>
+	 *	<setIntFromProperty result="" ini="" section="" property="" default=""/>
+	 *	<setStringFromProperty result="" ini="" section="" property="" default=""/>
+	 * 
+	 * Boolean variables may also be set to the result of applying operators:
+	 * 
+	 *	<setBoolNot result="" source=""/>
+	 *	<setBoolAnd result="" arg1="" arg2=""/>
+	 *	<setBoolOr result="" arg1="" arg2=""/>
+	 *	<setBoolIsEqual result="" arg1="" arg2=""/>
+	 *	<setBoolIsLess result="" arg1="" arg2=""/>
+	 *	<setBoolIsLessEqual result="" arg1="" arg2=""/>
+	 *	<setBoolIsGreater result="" arg1="" arg2=""/>
+	 *	<setBoolIsGreaterEqual result="" arg1="" arg2=""/>
+	 * 
+	 * Integer variables may use these arithmetic operations: 
+	 * 
+	 *	<setIntAdd result="" arg1="" arg2=""/>
+	 *	<setIntSubtract result="" arg1="" arg2=""/>
+	 *	<setIntMultiply result="" arg1="" arg2=""/>
+	 *	<setIntDivide result="" arg1="" arg2=""/>
+	 * 
+	 * Strings are manipulated with the following:
+	 * 
+ 	 *	<setStringAdd result="" arg1="" arg2=""/>
+	 *	<setStringSubstring result="" source="" start="" length=""/>
+	 *	<setStringReplace result="" source="" find="" with=""/>
+	 * 
+	 * String length may be retrieved with:
+	 * 
+	 *	<setIntLength result="" source=""/>
+	 * 
+	 * The index of a search string may be found in source with:
+	 * 
+	 *	<setIntFindString result="" source="" find=""/>
+	 *	
+	 * The following shortcut string comparisons may also be used instead of using <setIntFindString>
+	 * and checking the result:
+	 * 
+	 *	<setBoolStartsWith result="" source="" find=""/>
+	 *	<setBoolEndsWith result="" source="" find=""/>
+	 *	<setBoolContains result="" source="" find=""/>
+	 *	
+	 * Messages are written to the log with this node:
+	 *
+	 *	<log text=""/>
+	 * 
+	 * Conditional execution uses the following form:
+	 * 
+	 *	<if condition="">
+	 *		<true>
+	 *			<!-- executes if boolean variable in condition is true -->
+	 *		</true>
+	 *		<false>
+	 *			<!-- executes if boolean variable in condition is false -->
+	 *		</false>
+	 *	</if>
+	 * 
+	 * The <true> and <false> blocks are optional.  The condition must be in a boolean variable.
+	 * The boolean operator nodes may be combined to create a final state for more complex
+	 * conditions:
+	 * 
+	 *	<setBoolNot result="notDistribution" value="$B(Distribution)/>
+	 *	<setBoolEquals result="isX86" arg1="$S(Architecture)" arg2="x86"/>
+	 *	<setBoolEquals result="isX86_64" arg2="$S(Architecture)" arg2="x86_64/">
+	 *	<setBoolOr result="isIntel" arg1="$B(isX86)" arg2="$B(isX86_64)"/>
+	 *	<setBoolAnd result="intelAndNotDistribution" arg1="$B(isIntel)" arg2="$B(notDistribution)"/>
+	 *	<if condition="intelAndNotDistribution">
+	 *		<true>
+	 *			<!-- do something for Intel if not a distribution build -->
+	 *		</true>
+	 *	</if>
+	 * 
+	 * Note the "isIntel" could also be done like this:
+	 * 
+	 *	<setStringSubstring result="subarch" source="$S(Architecture)" start="0" length="3"/>
+	 *	<setBoolEquals result="isIntel" arg1="$S(subarch)" arg2="x86"/>
+	 * 
+	 * Two shortcut nodes are available for conditional execution:
+	 * 
+	 *	<isArch arch="armeabi-armv7">
+	 *		<!-- do stuff -->
+	 *	</isArch>
+	 * 
+	 * is the equivalent of:
+	 * 
+	 *	<setBoolEquals result="temp" arg1="$S(Architecture)" arg2="armeabi-armv7">
+	 *	<if condition="temp">
+	 *		<true>
+	 *			<!-- do stuff -->
+	 *		</true>
+	 *	</if>
+	 *	
+	 * and
+	 * 
+	 *	<isDistribution>
+	 *		<!-- do stuff -->
+	 *	</isDistribution>
+	 *	
+	 * is the equivalent of:
+	 * 
+	 *	<if condition="$B(Distribution)">
+	 *		<!-- do stuff -->
+	 *	</if>
+	 * 
+	 * Execution may be stopped with:
+	 * 
+	 *	<return/>
+	 *	 
+	 * Loops may be created using these nodes:
+	 * 
+	 *	<while condition="">
+	 *		<!-- do stuff -->
+	 *	</while>
+	 *	
+	 *	<break/>
+	 *	<continue/>
+	 *	
+	 * The <while> body will execute until the condition is false or a <break/> is hit.  The
+	 * <continue/> will restart execution of the loop if the condition is still true or exit.
+	 * 
+	 * Note: <break/> outside a <while> body will act the same as <return/>
+	 * 
+	 * Here is an example loop which writes 1 to 5 to the log, skipping 3.  Note the update of the
+	 * while condition should be done before the continue otherwise it may not exit.
+	 * 
+	 *	<setInt result="index" value="0"/>
+	 *	<setBool result="loopRun" value="true"/>
+	 *	<while condition="loopRun">
+	 *		<setIntAdd result="index" arg1="$I(index)" arg2="1"/>
+	 *		<setBoolIsLess result="loopRun" arg1="$I(index)" arg2="5"/>
+	 *		<setBoolIsEqual result="indexIs3" arg1="$I(index)" arg2="3"/>
+	 *		<if condition="indexIs3">
+	 *			<true>
+	 *				<continue/>
+	 *			</true>
+	 *		</if>
+	 *		<log text="$I(index)"/>
+	 *	</while>
+	 *	
+	 * It is possible to use variable replacement in generating the result variable
+	 * name as well.  This makes the creation of arrays in loops possible:
+	 * 
+	 *	<setString result="array_$I(index)" value="element $I(index) in array"/>
+	 *	
+	 * This may be retrieved using the following (value is treated as the variable
+	 * name):
+	 * 
+	 *	<setStringFrom result="out" value="array_$I(index)"/>
+	 *	
+	 * For boolean and integer types, you may use <setBoolFrom/> and <setIntFrom/>.
+	 * 
+	 * Nodes for inserting text into the section are as follows:
+	 * 
+	 *	<insert> body </insert>
+	 *	<insertValue value=""/>
+	 *	<loadLibrary name="" failmsg=""/>
+	 *	
+	 * The first one will insert either text or nodes into the returned section
+	 * string.  Please note you must use escaped characters for:
+	 * 
+	 *	< = &lt;
+	 *	> = &gt;
+	 *	& = &amp;
+	 * 
+	 *	<insertValue value=""/> evaluates variables in value before insertion.  If value contains
+	 * double quote ("), you must escape it with &quot;.
+	 * 
+	 *	<loadLibrary name="" failmsg=""/> is a shortcut to insert a system.LoadLibrary try/catch
+	 * block with an optional logged message for failure to load case.
+	 * 
+	 * You can do a search and replace in the Output with:
+	 * 
+	 *	<replace find="" with=""/>
+	 * 
+	 * Note you can also manipulate the actual $S(Output) directly, the above are more efficient:
+	 * 
+	 *	<setStringAdd result="Input" arg1="$S(Output)" arg2="sample\n"/>
+	 *	<setStringReplace result="Input" source="$S(Output)" find=".LAUNCH" with=".INFO"/>
+	 *	
+	 * XML manipulation uses the following nodes:
+	 * 
+	 *	<addElement tag="" name=""/>
+	 *	<addElements tag=""> body </addElements>
+	 *	<removeElement tag=""/>
+	 *	<setStringFromTag result="" tag="" name=""/>
+	 *	<setStringFromAttribute result="" tag="" name=""/>
+	 *	<addAttribute tag="" name="" value=""/>
+	 *	<removeAttribute tag="" name=""/>
+	 *	<loopElements tag=""> instructions </loopElements>
+	 *	
+	 * The current element is referenced with tag="$".  Element variables are referenced with $varname
+	 * since using $E(varname) will be expanded to the string equivalent of the XML.
+	 * 
+	 * <uses-permission>, <uses-feature>, and <uses-library> are updated with:
+	 * 
+	 *	<addPermission android:name="" .. />
+	 *	<addFeature android:name="" .. />
+	 *	<addLibrary android:name="" .. />
+	 *
+	 * Any attributes in the above commands are copied to the element added to the manifest so you
+	 * can do the following, for example:
+	 * 
+	 *	<addFeature android:name="android.hardware.usb.host" android:required="true"/>
+	 *	
+	 * Finally, these nodes allow copying of files useful for staging jar and so files:
+	 * 
+	 *	<copyFile src="" dst=""/>
+	 *	<copyDir src="" dst=""/>
+	 *	
+	 * The following should be used as the base for the src and dst paths:
+	 * 
+	 *	$S(PluginDir) = directory the XML file was loaded from
+ 	 *	$S(EngineDir) = engine directory
+	 *	$S(BuildDir) = project's platform appropriate build directory
+	 *	
+	 * While it is possible to write outside the APK directory, it is not recommended.
+	 * 
+	 * If you must remove files (like development-only files from distribution builds) you can
+	 * use this node:
+	 * 
+	 *	<removeFiles mask=""/>
+	 *	
+	 * It is restricted to only removing files from the BuildDir.  Here is example usage to remove
+	 * the Oculus Signature Files (osig) from the assets directory:
+	 * 
+	 *	<removeFiles mask="assets/oculussig_*"/>
+	 *	
+	 * The following sections are evaluated during the packaging or deploy stages:
+	 * 
+	 * ** For all platforms **
+	 *	<!-- init section is always evaluated once per architecture -->
+	 * 	<init> </init>
+	 * 
+	 * ** Android Specific sections **
+	 * 	<!-- optional updates applied to AndroidManifest.xml -->
+	 * 	<androidManifestUpdates> </androidManifestUpdates>
+	 * 	
+	 * 	<!-- optional additions to proguard -->
+	 * 	<proguardAdditions>	</proguardAdditions>
+	 * 	
+	 * 	<!-- optional files or directories to copy or delete from Intermediate/Android/APK before ndk-build -->
+	 * 	<prebuildCopies> </prebuildCopies>
+	 * 	
+	 * 	<!-- optional files or directories to copy or delete from Intermediate/Android/APK after ndk-build -->
+	 * 	<resourceCopies> </resourceCopies>
+	 * 	
+	 * 	<!-- optional additions to the GameActivity imports in GameActivity.java -->
+	 * 	<gameActivityImportAdditions> </gameActivityImportAdditions>
+	 * 	
+	 * 	<!-- optional additions to the GameActivity class in GameActivity.java -->
+	 * 	<gameActivityClassAdditions> </gameActivityOnClassAdditions>
+	 * 	
+	 * 	<!-- optional additions to GameActivity onCreate metadata reading in GameActivity.java -->
+	 * 	<gameActivityReadMetadata> </gameActivityReadMetadata>
+	 * 
+	 *	<!-- optional additions to GameActivity onCreate in GameActivity.java -->
+	 *	<gameActivityOnCreateAdditions> </gameActivityOnCreateAdditions>
+	 * 	
+	 * 	<!-- optional additions to GameActivity onDestroy in GameActivity.java -->
+	 * 	<gameActivityOnDestroyAdditions> </gameActivityOnDestroyAdditions>
+	 * 	
+	 * 	<!-- optional additions to GameActivity onStart in GameActivity.java -->
+	 * 	<gameActivityOnStartAdditions> </gameActivityOnStartAdditions>
+	 * 	
+	 * 	<!-- optional additions to GameActivity onStop in GameActivity.java -->
+	 * 	<gameActivityOnStopAdditions> </gameActivityOnStopAdditions>
+	 * 	
+	 * 	<!-- optional additions to GameActivity onPause in GameActivity.java -->
+	 * 	<gameActivityOnPauseAdditions> </gameActivityOnPauseAdditions>
+	 * 	
+	 * 	<!-- optional additions to GameActivity onResume in GameActivity.java -->
+	 *	<gameActivityOnResumeAdditions>	</gameActivityOnResumeAdditions>
+	 *	
+	 * 	<!-- optional additions to GameActivity onActivityResult in GameActivity.java -->
+	 * 	<gameActivityOnActivityResultAdditions>	</gameActivityOnActivityResultAdditions>
+	 * 	
+	 * 	<!-- optional libraries to load in GameActivity.java before libUE4.so -->
+	 * 	<soLoadLibrary>	</soLoadLibrary>
+	 * 	
+	 * 
+	 * Here is the complete list of supported nodes:
+	 * 
+	 * <isArch arch="">
+	 * <isDistribution>
+	 * <if> => <true> / <false>
+	 * <while condition="">
+	 * <return/>
+	 * <break/>
+	 * <continue/>
+	 * <log text=""/>
+	 * <insert> </insert>
+	 * <insertValue value=""/>
+	 * <replace find="" with""/>
+	 * <copyFile src="" dst=""/>
+	 * <copyDir src="" dst=""/>
+	 * <loadLibrary name="" failmsg=""/>
+	 * <setBool result="" value=""/>
+	 * <setBoolFrom result="" value=""/>
+	 * <setBoolFromProperty result="" ini="" section="" property="" default=""/>
+	 * <setBoolNot result="" source=""/>
+	 * <setBoolAnd result="" arg1="" arg2=""/>
+	 * <setBoolOr result="" arg1="" arg2=""/>
+	 * <setBoolIsEqual result="" arg1="" arg2=""/>
+	 * <setBoolIsLess result="" arg1="" arg2=""/>
+	 * <setBoolIsLessEqual result="" arg1="" arg2=""/>
+	 * <setBoolIsGreater result="" arg1="" arg2=""/>
+	 * <setBoolIsGreaterEqual result="" arg1="" arg2=""/>
+	 * <setInt result="" value=""/>
+	 * <setIntFrom result="" value=""/>
+	 * <setIntFromProperty result="" ini="" section="" property="" default=""/>
+	 * <setIntAdd result="" arg1="" arg2=""/>
+	 * <setIntSubtract result="" arg1="" arg2=""/>
+	 * <setIntMultiply result="" arg1="" arg2=""/>
+	 * <setIntDivide result="" arg1="" arg2=""/>
+	 * <setIntLength result="" source=""/>
+	 * <setIntFindString result="" source="" find=""/>
+	 * <setString result="" value=""/>
+	 * <setStringFrom result="" value=""/>
+	 * <setStringFromProperty result="" ini="" section="" property="" default=""/>
+	 * <setStringAdd result="" arg1="" arg2=""/>
+	 * <setStringSubstring result="" source="" index="" length=""/>
+	 * <setStringReplace result="" source="" find="" with=""/>
+	 * 
+	 */
+
+	public class UnrealPluginLanguage
+	{
+		/** The merged XML program to run */
+		private XDocument XDoc;
+
+		/** XML namespace */
+		private XNamespace XMLNameSpace;
+		private string XMLRootDefinition;
+
+		private UnrealTargetPlatform TargetPlatform;
+
+		/** Trace flag to enable debugging */
+		static private bool bGlobalTrace = false;
+
+		/** Project file reference */
+		private FileReference ProjectFile;
+		
+		static private XDocument XMLDummy = XDocument.Parse("<manifest></manifest>");
+
+		private class UPLContext
+		{
+			/** Variable state */
+			public Dictionary<string, bool> BoolVariables;
+			public Dictionary<string, int> IntVariables;
+			public Dictionary<string, string> StringVariables;
+			public Dictionary<string, XElement> ElementVariables;
+
+			/** Local context trace */
+			public bool bTrace;
+
+			public UPLContext(string Architecture, string PluginDir)
+			{
+				BoolVariables = new Dictionary<string, bool>();
+				IntVariables = new Dictionary<string, int>();
+				StringVariables = new Dictionary<string, string>();
+				ElementVariables = new Dictionary<string, XElement>();
+
+				StringVariables["Architecture"] = Architecture;
+				StringVariables["PluginDir"] = PluginDir;
+
+				bTrace = false;
+			}
+		}
+
+		private UPLContext GlobalContext;
+		private Dictionary<string, UPLContext> Contexts;
+		private int ContextIndex;
+
+		public UnrealPluginLanguage(FileReference InProjectFile, List<string> InXMLFiles, List<string> InArchitectures, string InXMLNameSpace, string InRootDefinition, UnrealTargetPlatform InTargetPlatform)
+		{
+			ProjectFile = InProjectFile;
+
+			Contexts = new Dictionary<string, UPLContext>();
+			GlobalContext = new UPLContext("", "");
+			ContextIndex = 0;
+
+			XMLNameSpace = InXMLNameSpace;
+			XMLRootDefinition = InRootDefinition;
+			TargetPlatform = InTargetPlatform;
+
+			string PathPrefix = Path.GetFileName(Directory.GetCurrentDirectory()).Equals("Source") ? ".." : "Engine";
+
+			XDoc = XDocument.Parse("<root " + InRootDefinition + "></root>");
+			foreach (string Basename in InXMLFiles)
+			{
+				string Filename = Path.Combine(PathPrefix, Basename.Replace("\\", "/"));
+				Log.TraceInformation("\nUPL: {0}", Filename);
+				if (File.Exists(Filename))
+				{
+					string PluginDir = Path.GetDirectoryName(Filename);
+					try
+					{
+						XDocument MergeDoc = XDocument.Load(Filename);
+						MergeXML(MergeDoc, PluginDir, InArchitectures);
+					}
+					catch (Exception e)
+					{
+						Log.TraceError("\nUnreal Plugin file {0} parsing failed! {1}", Filename, e);
+					}
+				}
+				else
+				{
+					Log.TraceError("\nUnreal Plugin file {0} missing!", Filename);
+					Log.TraceInformation("\nCWD: {0}", Directory.GetCurrentDirectory());
+				}
+			}
+		}
+
+		public bool GetTrace() { return bGlobalTrace; }
+		public void SetTrace() { bGlobalTrace = true; }
+		public void ClearTrace() { bGlobalTrace = false; }
+
+		public bool MergeXML(XDocument MergeDoc, string PluginDir, List<string> Architectures)
+		{
+			if (MergeDoc == null)
+			{
+				return false;
+			}
+	
+			// create a context for each architecture
+			ContextIndex++;
+			foreach (string Architecture in Architectures)
+			{
+				UPLContext Context = new UPLContext(Architecture, PluginDir);
+				Contexts[Architecture + "_" + ContextIndex] = Context;
+			}
+
+			// merge in the nodes
+			foreach (var Element in MergeDoc.Root.Elements())
+			{
+				var Parent = XDoc.Root.Element(Element.Name);
+				if (Parent != null)
+				{
+					var Entry = new XElement("Context", new XAttribute("index", ContextIndex.ToString()));
+					Entry.Add(Element.Elements());
+					Parent.Add(Entry);
+				}
+				else
+				{
+					var Entry = new XElement("Context", new XAttribute("index", ContextIndex.ToString()));
+					Entry.Add(Element.Elements());
+					var Base = new XElement(Element.Name);
+					Base.Add(Entry);
+					XDoc.Root.Add(Base);
+				}
+			}
+
+			return true;
+		}
+
+		public void SaveXML(string Filename)
+		{
+			if (XDoc != null)
+			{
+				XDoc.Save(Filename);
+			}
+		}
+
+		private string DumpContext(UPLContext Context)
+		{
+			StringBuilder Text = new StringBuilder();
+			foreach (var Variable in Context.BoolVariables)
+			{
+				Text.AppendLine(string.Format("\tbool {0} = {1}", Variable.Key, Variable.Value.ToString().ToLower()));
+			}
+			foreach (var Variable in Context.IntVariables)
+			{
+				Text.AppendLine(string.Format("\tint {0} = {1}", Variable.Key, Variable.Value));
+			}
+			foreach (var Variable in Context.StringVariables)
+			{
+				Text.AppendLine(string.Format("\tstring {0} = {1}", Variable.Key, Variable.Value));
+			}
+			foreach (var Variable in Context.ElementVariables)
+			{
+				Text.AppendLine(string.Format("\telement {0} = {1}", Variable.Key, Variable.Value));
+			}
+			return Text.ToString();
+		}
+
+		public string DumpVariables()
+		{
+			string Result = "Global Context:\n" + DumpContext(GlobalContext);
+			foreach (var Context in Contexts)
+			{
+				Result += "Context " + Context.Key + ": " + Context.Value.StringVariables["PluginDir"] + "\n" + DumpContext(Context.Value);
+			}
+			return Result;
+		}
+
+		private bool GetCondition(UPLContext Context, XElement Node, string Condition, out bool Result)
+		{
+			Result = false;
+			if (!Context.BoolVariables.TryGetValue(Condition, out Result))
+			{
+				if (!GlobalContext.BoolVariables.TryGetValue(Condition, out Result))
+				{
+					Log.TraceWarning("\nMissing condition '{0}' in '{1}' (skipping instruction)", Condition, TraceNodeString(Node));
+					return false;
+				}
+			}
+			return true;
+		}
+
+		private string ExpandVariables(UPLContext Context, string InputString)
+		{
+			string Result = InputString;
+			for (int Idx = Result.IndexOf("$B("); Idx != -1; Idx = Result.IndexOf("$B(", Idx))
+			{
+				// Find the end of the variable name
+				int EndIdx = Result.IndexOf(')', Idx + 3);
+				if (EndIdx == -1)
+				{
+					break;
+				}
+
+				// Extract the variable name from the string
+				string Name = Result.Substring(Idx + 3, EndIdx - (Idx + 3));
+
+				// Find the value for it, either from the dictionary or the environment block
+				bool Value;
+				if (!Context.BoolVariables.TryGetValue(Name, out Value))
+				{
+					if (!GlobalContext.BoolVariables.TryGetValue(Name, out Value))
+					{
+						Idx = EndIdx + 1;
+						continue;
+					}
+				}
+
+				// Replace the variable, or skip past it
+				Result = Result.Substring(0, Idx) + Value.ToString().ToLower() + Result.Substring(EndIdx + 1);
+			}
+			for (int Idx = Result.IndexOf("$I("); Idx != -1; Idx = Result.IndexOf("$I(", Idx))
+			{
+				// Find the end of the variable name
+				int EndIdx = Result.IndexOf(')', Idx + 3);
+				if (EndIdx == -1)
+				{
+					break;
+				}
+
+				// Extract the variable name from the string
+				string Name = Result.Substring(Idx + 3, EndIdx - (Idx + 3));
+
+				// Find the value for it, either from the dictionary or the environment block
+				int Value;
+				if (!Context.IntVariables.TryGetValue(Name, out Value))
+				{
+					if (!GlobalContext.IntVariables.TryGetValue(Name, out Value))
+					{
+						Idx = EndIdx + 1;
+						continue;
+					}
+				}
+
+				// Replace the variable, or skip past it
+				Result = Result.Substring(0, Idx) + Value.ToString() + Result.Substring(EndIdx + 1);
+			}
+			for (int Idx = Result.IndexOf("$S("); Idx != -1; Idx = Result.IndexOf("$S(", Idx))
+			{
+				// Find the end of the variable name
+				int EndIdx = Result.IndexOf(')', Idx + 3);
+				if (EndIdx == -1)
+				{
+					break;
+				}
+
+				// Extract the variable name from the string
+				string Name = Result.Substring(Idx + 3, EndIdx - (Idx + 3));
+
+				// Find the value for it, either from the dictionary or the environment block
+				string Value;
+				if (!Context.StringVariables.TryGetValue(Name, out Value))
+				{
+					if (!GlobalContext.StringVariables.TryGetValue(Name, out Value))
+					{
+						Idx = EndIdx + 1;
+						continue;
+					}
+				}
+
+				// Replace the variable, or skip past it
+				Result = Result.Substring(0, Idx) + Value + Result.Substring(EndIdx + 1);
+			}
+			for (int Idx = Result.IndexOf("$E("); Idx != -1; Idx = Result.IndexOf("$E(", Idx))
+			{
+				// Find the end of the variable name
+				int EndIdx = Result.IndexOf(')', Idx + 3);
+				if (EndIdx == -1)
+				{
+					break;
+				}
+
+				// Extract the variable name from the string
+				string Name = Result.Substring(Idx + 3, EndIdx - (Idx + 3));
+
+				// Find the value for it, either from the dictionary or the environment block
+				XElement Value;
+				if (!Context.ElementVariables.TryGetValue(Name, out Value))
+				{
+					if (!GlobalContext.ElementVariables.TryGetValue(Name, out Value))
+					{
+						Idx = EndIdx + 1;
+						continue;
+					}
+				}
+
+				// Replace the variable, or skip past it
+				Result = Result.Substring(0, Idx) + Value + Result.Substring(EndIdx + 1);
+			}
+			return Result;
+		}
+
+		private string TraceNodeString(XElement Node)
+		{
+			string Result = Node.Name.ToString();
+			foreach (var Attrib in Node.Attributes())
+			{
+				Result += " " + Attrib.ToString();
+			}
+			return Result;
+		}
+
+		private bool StringToBool(string Input)
+		{
+			if (Input == null)
+			{
+				return false;
+			}
+			Input = Input.ToLower();
+			return !(Input.Equals("0") || Input.Equals("false") || Input.Equals("off") || Input.Equals("no"));
+		}
+
+		private int StringToInt(string Input, XElement Node)
+		{
+			int Result = 0;
+			if (!int.TryParse(Input, out Result))
+			{
+				Log.TraceWarning("\nInvalid integer '{0}' in '{1}' (defaulting to 0)", Input, TraceNodeString(Node));
+			}
+			return Result;
+		}
+
+		private string GetAttribute(UPLContext Context, XElement Node, string AttributeName, bool bExpand = true, bool bRequired = true, string Fallback = null)
+		{
+			XAttribute Attribute = Node.Attribute(AttributeName);
+			if (Attribute == null)
+			{
+				if (bRequired)
+				{
+					Log.TraceWarning("\nMissing attribute '{0}' in '{1}' (skipping instruction)", AttributeName, TraceNodeString(Node));
+				}
+				return Fallback;
+			}
+			string Result = Attribute.Value;
+			return bExpand ? ExpandVariables(Context, Result) : Result;
+		}
+
+		private string GetAttributeWithNamespace(UPLContext Context, XElement Node, XNamespace Namespace, string AttributeName, bool bExpand = true, bool bRequired = true, string Fallback = null)
+		{
+			XAttribute Attribute = Node.Attribute(Namespace + AttributeName);
+			if (Attribute == null)
+			{
+				if (bRequired)
+				{
+					Log.TraceWarning("\nMissing attribute '{0}' in '{1}' (skipping instruction)", AttributeName, TraceNodeString(Node));
+				}
+				return Fallback;
+			}
+			string Result = Attribute.Value;
+			return bExpand ? ExpandVariables(Context, Result) : Result;
+		}
+
+		static private Dictionary<string, ConfigCacheIni_UPL> ConfigCache = null;
+
+		private ConfigCacheIni_UPL GetConfigCacheIni_UPL(string baseIniName)
+		{
+			if (ConfigCache == null)
+			{
+				ConfigCache = new Dictionary<string, ConfigCacheIni_UPL>();
+			}
+			ConfigCacheIni_UPL config = null;
+			if (!ConfigCache.TryGetValue(baseIniName, out config))
+			{
+				// note: use our own ConfigCacheIni since EngineConfiguration.cs only parses RequiredSections!
+				config = ConfigCacheIni_UPL.CreateConfigCacheIni_UPL(TargetPlatform, baseIniName, DirectoryReference.FromFile(ProjectFile));
+				ConfigCache.Add(baseIniName, config);
+			}
+			return config;
+		}
+
+		private static void CopyFileDirectory(string SourceDir, string DestDir)
+		{
+			if (!Directory.Exists(SourceDir))
+			{
+				return;
+			}
+
+			string[] Files = Directory.GetFiles(SourceDir, "*.*", SearchOption.AllDirectories);
+			foreach (string Filename in Files)
+			{
+				// make the dst filename with the same structure as it was in SourceDir
+				string DestFilename = Path.Combine(DestDir, Utils.MakePathRelativeTo(Filename, SourceDir));
+				if (File.Exists(DestFilename))
+				{
+					File.Delete(DestFilename);
+				}
+
+				// make the subdirectory if needed
+				string DestSubdir = Path.GetDirectoryName(DestFilename);
+				if (!Directory.Exists(DestSubdir))
+				{
+					Directory.CreateDirectory(DestSubdir);
+				}
+
+				File.Copy(Filename, DestFilename);
+
+				// remove any read only flags
+				FileInfo DestFileInfo = new FileInfo(DestFilename);
+				DestFileInfo.Attributes = DestFileInfo.Attributes & ~FileAttributes.ReadOnly;
+			}
+		}
+
+		private static void DeleteFiles(string Filespec)
+		{
+			string BaseDir = Path.GetDirectoryName(Filespec);
+			string Mask = Path.GetFileName(Filespec);
+
+			if (!Directory.Exists(BaseDir))
+			{
+				return;
+			}
+
+			string[] Files = Directory.GetFiles(BaseDir, Mask, SearchOption.TopDirectoryOnly);
+			foreach (string Filename in Files)
+			{
+				File.Delete(Filename);
+				Log.TraceInformation("\nDeleted file {0}", Filename);
+			}
+		}
+
+		private void AddAttribute(XElement Element, string Name, string Value)
+		{
+			XAttribute Attribute;
+			int Index = Name.IndexOf(":");
+			if (Index >= 0)
+			{
+				Name = Name.Substring(Index + 1);
+				Attribute = Element.Attribute(XMLNameSpace + Name);
+			}
+			else
+			{
+				Attribute = Element.Attribute(Name);
+			}
+
+			if (Attribute != null)
+			{
+				Attribute.SetValue(Value);
+			}
+			else
+			{
+				if (Index >= 0)
+				{
+					Element.Add(new XAttribute(XMLNameSpace + Name, Value));
+				}
+				else
+				{
+					Element.Add(new XAttribute(Name, Value));
+				}
+			}
+		}
+
+		private void RemoveAttribute(XElement Element, string Name)
+		{
+			XAttribute Attribute;
+			int Index = Name.IndexOf(":");
+			if (Index >= 0)
+			{
+				Name = Name.Substring(Index + 1);
+				Attribute = Element.Attribute(XMLNameSpace + Name);
+			}
+			else
+			{
+				Attribute = Element.Attribute(Name);
+			}
+
+			if (Attribute != null)
+			{
+				Attribute.Remove();
+			}
+		}
+
+		private void AddElements(XElement Target, XElement Source)
+		{
+			if (Source.HasElements)
+			{
+				foreach (var Index in Source.Elements())
+				{
+//					if (Target.Element(Index.Name) == null)
+					{
+						Target.Add(Index);
+					}
+				}
+			}
+		}
+
+		public string ProcessPluginNode(string Architecture, string NodeName, string Input)
+		{
+			return ProcessPluginNode(Architecture, NodeName, Input, ref XMLDummy);
+		}
+
+		public string ProcessPluginNode(string Architecture, string NodeName, string Input, ref XDocument XMLWork)
+		{
+			// add all instructions to execution list
+			var ExecutionStack = new Stack<XElement>();
+			var StartNode = XDoc.Root.Element(NodeName);
+			if (StartNode != null)
+			{
+				foreach (var Instruction in StartNode.Elements().Reverse())
+				{
+					ExecutionStack.Push(Instruction);
+				}
+			}
+
+			if (ExecutionStack.Count == 0)
+			{
+				return Input;
+			}
+
+			var ContextStack = new Stack<UPLContext>();
+			UPLContext CurrentContext = GlobalContext;
+
+			// update Output in global context
+			GlobalContext.StringVariables["Output"] = Input;
+
+			var ElementStack = new Stack<XElement>();
+			XElement CurrentElement = XMLWork.Elements().First();
+
+			// run the instructions
+			while (ExecutionStack.Count > 0)
+			{
+				var Node = ExecutionStack.Pop();
+				if (bGlobalTrace || CurrentContext.bTrace)
+				{
+					Log.TraceInformation("Execute: '{0}'", TraceNodeString(Node));
+				}
+				switch (Node.Name.ToString())
+				{
+					case "trace":
+						{
+							string Enable = GetAttribute(CurrentContext, Node, "enable");
+							if (Enable != null)
+							{
+								CurrentContext.bTrace = StringToBool(Enable);
+								if (!bGlobalTrace && CurrentContext.bTrace)
+								{
+									Log.TraceInformation("Context: '{0}' using Architecture='{1}', NodeName='{2}', Input='{3}'", CurrentContext.StringVariables["PluginDir"], Architecture, NodeName, Input);
+								}
+							}
+						}
+						break;
+
+					case "dumpvars":
+						{
+							if (!bGlobalTrace && !CurrentContext.bTrace)
+							{
+								Log.TraceInformation("Context: '{0}' using Architecture='{1}', NodeName='{2}', Input='{3}'", CurrentContext.StringVariables["PluginDir"], Architecture, NodeName, Input);
+							}
+							Log.TraceInformation("Variables:\n{0}", DumpContext(CurrentContext));
+						}
+						break;
+
+					case "Context":
+						{
+							ContextStack.Push(CurrentContext);
+							string index = GetAttribute(CurrentContext, Node, "index");
+							CurrentContext = Contexts[Architecture + "_" + index];
+							ExecutionStack.Push(new XElement("PopContext"));
+							foreach (var instruction in Node.Elements().Reverse())
+							{
+								ExecutionStack.Push(instruction);
+							}
+
+							if (bGlobalTrace || CurrentContext.bTrace)
+							{
+								Log.TraceInformation("Context: '{0}' using Architecture='{1}', NodeName='{2}', Input='{3}'", CurrentContext.StringVariables["PluginDir"], Architecture, NodeName, Input);
+							}
+						}
+						break;
+
+					case "PopContext":
+						{
+							CurrentContext = ContextStack.Pop();
+						}
+						break;
+
+					case "PopElement":
+						{
+							CurrentElement = ElementStack.Pop();
+						}
+						break;
+
+					case "isArch":
+						{
+							string arch = GetAttribute(CurrentContext, Node, "arch");
+							if (arch != null && arch.Equals(Architecture))
+							{
+								foreach (var instruction in Node.Elements().Reverse())
+								{
+									ExecutionStack.Push(instruction);
+								}
+							}
+						}
+						break;
+
+					case "isDistribution":
+						{
+							bool Result = false;
+							if (GetCondition(CurrentContext, Node, "Distribution", out Result))
+							{
+								if (Result)
+								{
+									foreach (var Instruction in Node.Elements().Reverse())
+									{
+										ExecutionStack.Push(Instruction);
+									}
+								}
+							}
+						}
+						break;
+
+					case "if":
+						{
+							bool Result;
+							if (GetCondition(CurrentContext, Node, GetAttribute(CurrentContext, Node, "condition"), out Result))
+							{
+								var ResultNode = Node.Element(Result ? "true" : "false");
+								if (ResultNode != null)
+								{
+									foreach (var Instruction in ResultNode.Elements().Reverse())
+									{
+										ExecutionStack.Push(Instruction);
+									}
+								}
+							}
+						}
+						break;
+
+					case "while":
+						{
+							bool Result;
+							if (GetCondition(CurrentContext, Node, GetAttribute(CurrentContext, Node, "condition"), out Result))
+							{
+								if (Result)
+								{
+									var ResultNode = Node.Elements();
+									if (ResultNode != null)
+									{
+										ExecutionStack.Push(Node);
+										foreach (var Instruction in ResultNode.Reverse())
+										{
+											ExecutionStack.Push(Instruction);
+										}
+									}
+								}
+							}
+						}
+						break;
+
+					case "return":
+						{
+							while (ExecutionStack.Count > 0)
+							{
+								ExecutionStack.Pop();
+							}
+						}
+						break;
+
+					case "break":
+						{
+							// remove up to while (or acts like a return if outside by removing everything)
+							while (ExecutionStack.Count > 0)
+							{
+								Node = ExecutionStack.Pop();
+								if (Node.Name.ToString().Equals("while"))
+								{
+									break;
+								}
+							}
+						}
+						break;
+
+					case "continue":
+						{
+							// remove up to while (or acts like a return if outside by removing everything)
+							while (ExecutionStack.Count > 0)
+							{
+								Node = ExecutionStack.Pop();
+								if (Node.Name.ToString().Equals("while"))
+								{
+									ExecutionStack.Push(Node);
+									break;
+								}
+							}
+						}
+						break;
+
+					case "log":
+						{
+							string Text = GetAttribute(CurrentContext, Node, "text");
+							if (Text != null)
+							{
+								Log.TraceInformation("{0}", Text);
+							}
+						}
+						break;
+
+					case "loopElements":
+						{
+							string Tag = GetAttribute(CurrentContext, Node, "tag");
+							ElementStack.Push(CurrentElement);
+							ExecutionStack.Push(new XElement("PopElement"));
+							var WorkList = (Tag == "$") ? CurrentElement.Elements().Reverse() : CurrentElement.Descendants(Tag).Reverse();
+							foreach (var WorkNode in WorkList)
+							{
+								foreach (var Instruction in Node.Elements().Reverse())
+								{
+									ExecutionStack.Push(Instruction);
+								}
+								ElementStack.Push(WorkNode);
+								ExecutionStack.Push(new XElement("PopElement"));
+							}
+						}
+						break;
+
+					case "addAttribute":
+						{
+							string Tag = GetAttribute(CurrentContext, Node, "tag");
+							string Name = GetAttribute(CurrentContext, Node, "name");
+							string Value = GetAttribute(CurrentContext, Node, "value");
+							if (Tag != null && Name != null && Value != null)
+							{
+								if (Tag.StartsWith("$"))
+								{
+									XElement Target = CurrentElement;
+									if (Tag.Length > 1)
+									{
+										if (!CurrentContext.ElementVariables.TryGetValue(Tag.Substring(1), out Target))
+										{
+											if (!GlobalContext.ElementVariables.TryGetValue(Tag.Substring(1), out Target))
+											{
+												Log.TraceWarning("\nMissing element variable '{0}' in '{1}' (skipping instruction)", Tag, TraceNodeString(Node));
+												continue;
+											}
+										}
+									}
+									AddAttribute(Target, Name, Value);
+								}
+								else
+								{
+									if (CurrentElement.Name.ToString().Equals(Tag))
+									{
+										AddAttribute(CurrentElement, Name, Value);
+									}
+									foreach (var WorkNode in CurrentElement.Descendants(Tag))
+									{
+										AddAttribute(WorkNode, Name, Value);
+									}
+								}
+							}
+						}
+						break;
+
+					case "removeAttribute":
+						{
+							string Tag = GetAttribute(CurrentContext, Node, "tag");
+							string Name = GetAttribute(CurrentContext, Node, "name");
+							if (Tag != null && Name != null)
+							{
+								if (Tag.StartsWith("$"))
+								{
+									XElement Target = CurrentElement;
+									if (Tag.Length > 1)
+									{
+										if (!CurrentContext.ElementVariables.TryGetValue(Tag.Substring(1), out Target))
+										{
+											if (!GlobalContext.ElementVariables.TryGetValue(Tag.Substring(1), out Target))
+											{
+												Log.TraceWarning("\nMissing element variable '{0}' in '{1}' (skipping instruction)", Tag, TraceNodeString(Node));
+												continue;
+											}
+										}
+									}
+									RemoveAttribute(Target, Name);
+								}
+								else
+								{
+									if (CurrentElement.Name.ToString().Equals(Tag))
+									{
+										RemoveAttribute(CurrentElement, Name);
+									}
+									foreach (var WorkNode in CurrentElement.Descendants(Tag))
+									{
+										RemoveAttribute(WorkNode, Name);
+									}
+								}
+							}
+						}
+						break;
+
+					case "addPermission":
+						{
+							string Name = GetAttributeWithNamespace(CurrentContext, Node, XMLNameSpace, "name");
+							if (Name != null)
+							{
+								// make sure it isn't already added
+								bool bFound = false;
+								foreach (var Element in XMLWork.Descendants("uses-permission"))
+								{
+									XAttribute Attribute = Element.Attribute(XMLNameSpace + "name");
+									if (Attribute != null)
+									{
+										if (Attribute.Value == Name)
+										{
+											bFound = true;
+											break;
+										}
+									}
+								}
+
+								// add it if not found
+								if (!bFound)
+								{
+									XMLWork.Element("manifest").Add(new XElement("uses-permission", Node.Attributes()));
+								}
+							}
+						}
+						break;
+
+					case "removePermission":
+						{
+							string Name = GetAttributeWithNamespace(CurrentContext, Node, XMLNameSpace, "name");
+							if (Name != null)
+							{
+								foreach (var Element in XMLWork.Descendants("uses-permission"))
+								{
+									XAttribute Attribute = Element.Attribute(XMLNameSpace + "name");
+									if (Attribute != null)
+									{
+										if (Attribute.Value == Name)
+										{
+											Element.Remove();
+											break;
+										}
+									}
+								}
+							}
+						}
+						break;
+
+					case "addFeature":
+						{
+							string Name = GetAttributeWithNamespace(CurrentContext, Node, XMLNameSpace, "name");
+							if (Name != null)
+							{
+								// make sure it isn't already added
+								bool bFound = false;
+								foreach (var Element in XMLWork.Descendants("uses-feature"))
+								{
+									XAttribute Attribute = Element.Attribute(XMLNameSpace + "name");
+									if (Attribute != null)
+									{
+										if (Attribute.Value == Name)
+										{
+											bFound = true;
+											break;
+										}
+									}
+								}
+
+								// add it if not found
+								if (!bFound)
+								{
+									XMLWork.Element("manifest").Add(new XElement("uses-feature", Node.Attributes()));
+								}
+							}
+						}
+						break;
+
+					case "removeFeature":
+						{
+							string Name = GetAttributeWithNamespace(CurrentContext, Node, XMLNameSpace, "name");
+							if (Name != null)
+							{
+								foreach (var Element in XMLWork.Descendants("uses-feature"))
+								{
+									XAttribute Attribute = Element.Attribute(XMLNameSpace + "name");
+									if (Attribute != null)
+									{
+										if (Attribute.Value == Name)
+										{
+											Element.Remove();
+											break;
+										}
+									}
+								}
+							}
+						}
+						break;
+
+					case "addLibrary":
+						{
+							string Name = GetAttributeWithNamespace(CurrentContext, Node, XMLNameSpace, "name");
+							if (Name != null)
+							{
+								// make sure it isn't already added
+								bool bFound = false;
+								foreach (var Element in XMLWork.Descendants("uses-library"))
+								{
+									XAttribute Attribute = Element.Attribute(XMLNameSpace + "name");
+									if (Attribute != null)
+									{
+										if (Attribute.Value == Name)
+										{
+											bFound = true;
+											break;
+										}
+									}
+								}
+
+								// add it if not found
+								if (!bFound)
+								{
+									XMLWork.Element("manifest").Element("application").Add(new XElement("uses-library", Node.Attributes()));
+								}
+							}
+						}
+						break;
+
+					case "removeLibrary":
+						{
+							string Name = GetAttributeWithNamespace(CurrentContext, Node, XMLNameSpace, "name");
+							if (Name != null)
+							{
+								foreach (var Element in XMLWork.Descendants("uses-library"))
+								{
+									XAttribute Attribute = Element.Attribute(XMLNameSpace + "name");
+									if (Attribute != null)
+									{
+										if (Attribute.Value == Name)
+										{
+											Element.Remove();
+											break;
+										}
+									}
+								}
+							}
+						}
+						break;
+
+					case "removeElement":
+						{
+							string Tag = GetAttribute(CurrentContext, Node, "tag");
+							if (Tag != null)
+							{
+								if (Tag == "$")
+								{
+									XElement Parent = CurrentElement.Parent;
+									CurrentElement.Remove();
+									CurrentElement = Parent;
+								}
+								else
+								{
+									// use a list since Remove() may modify it
+									foreach (var Element in XMLWork.Descendants(Tag).ToList())
+									{
+										Element.Remove();
+									}
+								}
+							}
+						}
+						break;
+
+					case "addElement":
+						{
+							string Tag = GetAttribute(CurrentContext, Node, "tag");
+							string Name = GetAttribute(CurrentContext, Node, "name");
+							if (Tag != null && Name != null)
+							{
+								XElement Element;
+								if (!CurrentContext.ElementVariables.TryGetValue(Name, out Element))
+								{
+									if (!GlobalContext.ElementVariables.TryGetValue(Name, out Element))
+									{
+										Log.TraceWarning("\nMissing element variable '{0}' in '{1}' (skipping instruction)", Name, TraceNodeString(Node));
+										continue;
+									}
+								}
+								if (Tag.StartsWith("$"))
+								{
+									XElement Target = CurrentElement;
+									if (Tag.Length > 1)
+									{
+										if (!CurrentContext.ElementVariables.TryGetValue(Tag.Substring(1), out Target))
+										{
+											if (!GlobalContext.ElementVariables.TryGetValue(Tag.Substring(1), out Target))
+											{
+												Log.TraceWarning("\nMissing element variable '{0}' in '{1}' (skipping instruction)", Tag, TraceNodeString(Node));
+												continue;
+											}
+										}
+									}
+									Target.Add(new XElement(Element));
+								}
+								else
+								{
+									if (CurrentElement.Name.ToString().Equals(Tag))
+									{
+										CurrentElement.Add(new XElement(Element));
+									}
+									foreach (var WorkNode in CurrentElement.Descendants(Tag))
+									{
+										WorkNode.Add(new XElement(Element));
+									}
+								}
+							}
+						}
+						break;
+
+					case "addElements":
+						{
+							string Tag = GetAttribute(CurrentContext, Node, "tag");
+							if (Tag != null)
+							{
+								if (Tag.StartsWith("$"))
+								{
+									XElement Target = CurrentElement;
+									if (Tag.Length > 1)
+									{
+										if (!CurrentContext.ElementVariables.TryGetValue(Tag.Substring(1), out Target))
+										{
+											if (!GlobalContext.ElementVariables.TryGetValue(Tag.Substring(1), out Target))
+											{
+												Log.TraceWarning("\nMissing element variable '{0}' in '{1}' (skipping instruction)", Tag, TraceNodeString(Node));
+												continue;
+											}
+										}
+									}
+									AddElements(Target, Node);
+								}
+								else
+								{
+									if (CurrentElement.Name.ToString().Equals(Tag))
+									{
+										AddElements(CurrentElement, Node);
+									}
+									foreach (var WorkNode in CurrentElement.Descendants(Tag))
+									{
+										AddElements(WorkNode, Node);
+									}
+								}
+							}
+						}
+						break;
+
+					case "insert":
+						{
+							if (Node.HasElements)
+							{
+								foreach (var Element in Node.Elements())
+								{
+									string Value = Element.ToString().Replace(" " + XMLRootDefinition + " ", "");
+									GlobalContext.StringVariables["Output"] += Value + "\n";
+								}
+							}
+							else
+							{
+								string Value = Node.Value.ToString();
+
+								// trim trailing tabs
+								int Index = Value.Length;
+								while (Index > 0 && Value[Index - 1] == '\t')
+								{
+									Index--;
+								}
+								if (Index < Value.Length)
+								{
+									Value = Value.Substring(0, Index);
+								}
+
+								// trim leading newlines
+								Index = 0;
+								while (Index < Value.Length && Value[Index] == '\n')
+								{
+									Index++;
+								}
+
+								if (Index < Value.Length)
+								{
+									GlobalContext.StringVariables["Output"] += Value.Substring(Index);
+								}
+							}
+						}
+						break;
+
+					case "insertValue":
+						{
+							string Value = GetAttribute(CurrentContext, Node, "value");
+							if (Value != null)
+							{
+								GlobalContext.StringVariables["Output"] += Value;
+							}
+						}
+						break;
+
+					case "replace":
+						{
+							string Find = GetAttribute(CurrentContext, Node, "find");
+							string With = GetAttribute(CurrentContext, Node, "with");
+							if (Find != null && With != null)
+							{
+								GlobalContext.StringVariables["Output"] = GlobalContext.StringVariables["Output"].Replace(Find, With);
+							}
+						}
+						break;
+
+					case "copyFile":
+						{
+							string Src = GetAttribute(CurrentContext, Node, "src");
+							string Dst = GetAttribute(CurrentContext, Node, "dst");
+							if (Src != null && Dst != null)
+							{
+								if (File.Exists(Src))
+								{
+									// check to see if newer than last time we copied
+									bool bFileExists = File.Exists(Dst);
+									TimeSpan Diff = File.GetLastWriteTimeUtc(Dst) - File.GetLastWriteTimeUtc(Src);
+									if (!bFileExists || Diff.TotalSeconds < -1 || Diff.TotalSeconds > 1)
+									{
+										if (bFileExists)
+										{
+											File.Delete(Dst);
+										}
+										Directory.CreateDirectory(Path.GetDirectoryName(Dst));
+										File.Copy(Src, Dst, true);
+										Log.TraceInformation("\nFile {0} copied to {1}", Src, Dst);
+
+										// remove any read only flags
+										FileInfo DestFileInfo = new FileInfo(Dst);
+										DestFileInfo.Attributes = DestFileInfo.Attributes & ~FileAttributes.ReadOnly;
+									}
+								}
+							}
+						}
+						break;
+
+					case "copyDir":
+						{
+							string Src = GetAttribute(CurrentContext, Node, "src");
+							string Dst = GetAttribute(CurrentContext, Node, "dst");
+							if (Src != null && Dst != null)
+							{
+								CopyFileDirectory(Src, Dst);
+								Log.TraceInformation("\nDirectory {0} copied to {1}", Src, Dst);
+							}
+						}
+						break;
+
+					case "deleteFiles":
+						{
+							string Filespec = GetAttribute(CurrentContext, Node, "filespec");
+							if (Filespec != null)
+							{
+								if (Filespec.Contains(":") || Filespec.Contains(".."))
+								{
+									Log.TraceInformation("\nFilespec {0} not allowed; ignored.", Filespec);
+								}
+								else
+								{
+									// force relative to BuildDir (and only from global context so someone doesn't try to be clever)
+									DeleteFiles(Path.Combine(GlobalContext.StringVariables["BuildDir"], Filespec));
+								}
+							}
+						}
+						break;
+
+					case "loadLibrary":
+						{
+							string Name = GetAttribute(CurrentContext, Node, "name");
+							string FailMsg = GetAttribute(CurrentContext, Node, "failmsg", true, false);
+							if (Name != null)
+							{
+								string Work = "\t\ttry\n" +
+											"\t\t{\n" +
+											"\t\t\tSystem.loadLibrary(\"" + Name + "\");\n" +
+											"\t\t}\n" +
+											"\t\tcatch (java.lang.UnsatisfiedLinkError e)\n" +
+											"\t\t{\n";
+								if (FailMsg != null)
+								{
+									Work += "\t\t\tLog.debug(\"" + FailMsg + "\");\n";
+								}
+								GlobalContext.StringVariables["Output"] += Work + "\t\t}\n";
+							}
+						}
+						break;
+
+					case "setBool":
+						{
+							string Result = GetAttribute(CurrentContext, Node, "result");
+							string Value = GetAttribute(CurrentContext, Node, "value", true, false, "false");
+							if (Result != null)
+							{
+								CurrentContext.BoolVariables[Result] = StringToBool(Value);
+							}
+						}
+						break;
+
+					case "setBoolFrom":
+						{
+							string Result = GetAttribute(CurrentContext, Node, "result");
+							string Value = GetAttribute(CurrentContext, Node, "value", true, false, "false");
+							if (Result != null)
+							{
+								Value = ExpandVariables(CurrentContext, "$B(" + Value + ")");
+								CurrentContext.BoolVariables[Result] = StringToBool(Value);
+							}
+						}
+						break;
+
+					case "setBoolFromProperty":
+						{
+							string Result = GetAttribute(CurrentContext, Node, "result");
+							string Ini = GetAttribute(CurrentContext, Node, "ini");
+							string Section = GetAttribute(CurrentContext, Node, "section");
+							string Property = GetAttribute(CurrentContext, Node, "property");
+							string DefaultVal = GetAttribute(CurrentContext, Node, "default", true, false, "false");
+							if (Result != null && Ini != null && Section != null && Property != null)
+							{
+								bool Value = StringToBool(DefaultVal);
+
+								ConfigCacheIni_UPL ConfigIni = GetConfigCacheIni_UPL(Ini);
+								if (ConfigIni != null)
+								{
+									if (!ConfigIni.GetBool(Section, Property, out Value))
+									{
+										Value = StringToBool(DefaultVal);
+									}
+								}
+								CurrentContext.BoolVariables[Result] = Value;
+							}
+						}
+						break;
+
+					case "setBoolContains":
+						{
+							string Result = GetAttribute(CurrentContext, Node, "result");
+							string Source = GetAttribute(CurrentContext, Node, "source", true, false, "");
+							string Find = GetAttribute(CurrentContext, Node, "find", true, false, "");
+							if (Result != null)
+							{
+								CurrentContext.BoolVariables[Result] = Source.Contains(Find);
+							}
+						}
+						break;
+
+					case "setBoolStartsWith":
+						{
+							string Result = GetAttribute(CurrentContext, Node, "result");
+							string Source = GetAttribute(CurrentContext, Node, "source", true, false, "");
+							string Find = GetAttribute(CurrentContext, Node, "find", true, false, "");
+							if (Result != null)
+							{
+								CurrentContext.BoolVariables[Result] = Source.StartsWith(Find);
+							}
+						}
+						break;
+
+					case "setBoolEndsWith":
+						{
+							string Result = GetAttribute(CurrentContext, Node, "result");
+							string Source = GetAttribute(CurrentContext, Node, "source", true, false, "");
+							string Find = GetAttribute(CurrentContext, Node, "find", true, false, "");
+							if (Result != null)
+							{
+								CurrentContext.BoolVariables[Result] = Source.EndsWith(Find);
+							}
+						}
+						break;
+
+					case "setBoolNot":
+						{
+							string Result = GetAttribute(CurrentContext, Node, "result");
+							string Source = GetAttribute(CurrentContext, Node, "source", true, false, "false");
+							if (Result != null)
+							{
+								CurrentContext.BoolVariables[Result] = !StringToBool(Source);
+							}
+						}
+						break;
+
+					case "setBoolAnd":
+						{
+							string Result = GetAttribute(CurrentContext, Node, "result");
+							string Arg1 = GetAttribute(CurrentContext, Node, "arg1", true, false, "false");
+							string Arg2 = GetAttribute(CurrentContext, Node, "arg2", true, false, "false");
+							if (Result != null)
+							{
+								CurrentContext.BoolVariables[Result] = StringToBool(Arg1) && StringToBool(Arg2);
+							}
+						}
+						break;
+
+					case "setBoolOr":
+						{
+							string Result = GetAttribute(CurrentContext, Node, "result");
+							string Arg1 = GetAttribute(CurrentContext, Node, "arg1", true, false, "false");
+							string Arg2 = GetAttribute(CurrentContext, Node, "arg2", true, false, "false");
+							if (Result != null)
+							{
+								CurrentContext.BoolVariables[Result] = StringToBool(Arg1) || StringToBool(Arg2);
+							}
+						}
+						break;
+
+					case "setBoolIsEqual":
+						{
+							string Result = GetAttribute(CurrentContext, Node, "result");
+							string Arg1 = GetAttribute(CurrentContext, Node, "arg1", true, false, "");
+							string Arg2 = GetAttribute(CurrentContext, Node, "arg2", true, false, "");
+							if (Result != null)
+							{
+								CurrentContext.BoolVariables[Result] = Arg1.Equals(Arg2);
+							}
+						}
+						break;
+
+					case "setBoolIsLess":
+						{
+							string Result = GetAttribute(CurrentContext, Node, "result");
+							string Arg1 = GetAttribute(CurrentContext, Node, "arg1", true, false);
+							string Arg2 = GetAttribute(CurrentContext, Node, "arg2", true, false);
+							if (Result != null)
+							{
+								CurrentContext.BoolVariables[Result] = (StringToInt(Arg1, Node) < StringToInt(Arg2, Node));
+							}
+						}
+						break;
+
+					case "setBoolIsLessEqual":
+						{
+							string Result = GetAttribute(CurrentContext, Node, "result");
+							string Arg1 = GetAttribute(CurrentContext, Node, "arg1", true, false);
+							string Arg2 = GetAttribute(CurrentContext, Node, "arg2", true, false);
+							if (Result != null)
+							{
+								CurrentContext.BoolVariables[Result] = (StringToInt(Arg1, Node) <= StringToInt(Arg2, Node));
+							}
+						}
+						break;
+
+					case "setBoolIsGreater":
+						{
+							string Result = GetAttribute(CurrentContext, Node, "result");
+							string Arg1 = GetAttribute(CurrentContext, Node, "arg1", true, false);
+							string Arg2 = GetAttribute(CurrentContext, Node, "arg2", true, false);
+							if (Result != null)
+							{
+								CurrentContext.BoolVariables[Result] = (StringToInt(Arg1, Node) > StringToInt(Arg2, Node));
+							}
+						}
+						break;
+
+					case "setBoolIsGreaterEqual":
+						{
+							string Result = GetAttribute(CurrentContext, Node, "result");
+							string Arg1 = GetAttribute(CurrentContext, Node, "arg1", true, false);
+							string Arg2 = GetAttribute(CurrentContext, Node, "arg2", true, false);
+							if (Result != null)
+							{
+								CurrentContext.BoolVariables[Result] = (StringToInt(Arg1, Node) >= StringToInt(Arg2, Node));
+							}
+						}
+						break;
+
+					case "setInt":
+						{
+							string Result = GetAttribute(CurrentContext, Node, "result");
+							string Value = GetAttribute(CurrentContext, Node, "value", true, false, "0");
+							if (Result != null)
+							{
+								CurrentContext.IntVariables[Result] = StringToInt(Value, Node);
+							}
+						}
+						break;
+
+					case "setIntFrom":
+						{
+							string Result = GetAttribute(CurrentContext, Node, "result");
+							string Value = GetAttribute(CurrentContext, Node, "value", true, false, "0");
+							if (Result != null)
+							{
+								Value = ExpandVariables(CurrentContext, "$I(" + Value + ")");
+								CurrentContext.IntVariables[Result] = StringToInt(Value, Node);
+							}
+						}
+						break;
+
+					case "setIntFromProperty":
+						{
+							string Result = GetAttribute(CurrentContext, Node, "result");
+							string Ini = GetAttribute(CurrentContext, Node, "ini");
+							string Section = GetAttribute(CurrentContext, Node, "section");
+							string Property = GetAttribute(CurrentContext, Node, "property");
+							string DefaultVal = GetAttribute(CurrentContext, Node, "default", true, false, "0");
+							if (Result != null && Ini != null && Section != null && Property != null)
+							{
+								int Value = StringToInt(DefaultVal, Node);
+
+								ConfigCacheIni_UPL ConfigIni = GetConfigCacheIni_UPL(Ini);
+								if (ConfigIni != null)
+								{
+									ConfigIni.GetInt32(Section, Property, out Value);
+								}
+								CurrentContext.IntVariables[Result] = Value;
+							}
+						}
+						break;
+
+					case "setIntAdd":
+						{
+							string Result = GetAttribute(CurrentContext, Node, "result");
+							string Arg1 = GetAttribute(CurrentContext, Node, "arg1", true, false, "0");
+							string Arg2 = GetAttribute(CurrentContext, Node, "arg2", true, false, "0");
+							if (Result != null)
+							{
+								CurrentContext.IntVariables[Result] = StringToInt(Arg1, Node) + StringToInt(Arg2, Node);
+							}
+						}
+						break;
+
+					case "setIntSubtract":
+						{
+							string Result = GetAttribute(CurrentContext, Node, "result");
+							string Arg1 = GetAttribute(CurrentContext, Node, "arg1", true, false, "0");
+							string Arg2 = GetAttribute(CurrentContext, Node, "arg2", true, false, "0");
+							if (Result != null)
+							{
+								CurrentContext.IntVariables[Result] = StringToInt(Arg1, Node) - StringToInt(Arg2, Node);
+							}
+						}
+						break;
+
+					case "setIntMultiply":
+						{
+							string Result = GetAttribute(CurrentContext, Node, "result");
+							string Arg1 = GetAttribute(CurrentContext, Node, "arg1", true, false, "1");
+							string Arg2 = GetAttribute(CurrentContext, Node, "arg2", true, false, "1");
+							if (Result != null)
+							{
+								CurrentContext.IntVariables[Result] = StringToInt(Arg1, Node) * StringToInt(Arg2, Node);
+							}
+						}
+						break;
+
+					case "setIntDivide":
+						{
+							string Result = GetAttribute(CurrentContext, Node, "result");
+							string Arg1 = GetAttribute(CurrentContext, Node, "arg1", true, false, "1");
+							string Arg2 = GetAttribute(CurrentContext, Node, "arg2", true, false, "1");
+							if (Result != null)
+							{
+								int Denominator = StringToInt(Arg2, Node);
+								if (Denominator == 0)
+								{
+									CurrentContext.IntVariables[Result] = StringToInt(Arg1, Node);
+								}
+								else
+								{
+									CurrentContext.IntVariables[Result] = StringToInt(Arg1, Node) / Denominator;
+								}
+							}
+						}
+						break;
+
+					case "setIntLength":
+						{
+							string Result = GetAttribute(CurrentContext, Node, "result");
+							string Source = GetAttribute(CurrentContext, Node, "source", true, false, "");
+							if (Result != null)
+							{
+								CurrentContext.IntVariables[Result] = Source.Length;
+							}
+						}
+						break;
+
+					case "setIntFindString":
+						{
+							string Result = GetAttribute(CurrentContext, Node, "result");
+							string Source = GetAttribute(CurrentContext, Node, "source", true, false, "");
+							string Find = GetAttribute(CurrentContext, Node, "find", true, false, "");
+							if (Result != null)
+							{
+								CurrentContext.IntVariables[Result] = Source.IndexOf(Find);
+							}
+						}
+						break;
+
+					case "setString":
+						{
+							string Result = GetAttribute(CurrentContext, Node, "result");
+							string Value = GetAttribute(CurrentContext, Node, "value", true, false, "");
+							if (Result != null)
+							{
+								if (Result == "Output")
+								{
+									GlobalContext.StringVariables["Output"] = Value;
+								}
+								else
+								{
+									CurrentContext.StringVariables[Result] = Value;
+								}
+							}
+						}
+						break;
+
+					case "setStringFrom":
+						{
+							string Result = GetAttribute(CurrentContext, Node, "result");
+							string Value = GetAttribute(CurrentContext, Node, "value", true, false, "0");
+							if (Result != null)
+							{
+								Value = ExpandVariables(CurrentContext, "$S(" + Value + ")");
+								CurrentContext.StringVariables[Result] = Value;
+							}
+						}
+						break;
+
+					case "setStringFromTag":
+						{
+							string Result = GetAttribute(CurrentContext, Node, "result");
+							string Tag = GetAttribute(CurrentContext, Node, "tag", true, false, "$");
+							if (Result != null)
+							{
+								XElement Element = CurrentElement;
+								if (Tag.StartsWith("$"))
+								{
+									if (Tag.Length > 1)
+									{
+										if (!CurrentContext.ElementVariables.TryGetValue(Tag.Substring(1), out Element))
+										{
+											if (!GlobalContext.ElementVariables.TryGetValue(Tag.Substring(1), out Element))
+											{
+												Log.TraceWarning("\nMissing element variable '{0}' in '{1}' (skipping instruction)", Tag, TraceNodeString(Node));
+												continue;
+											}
+										}
+									}
+								}
+								CurrentContext.StringVariables[Result] = Element.Name.ToString();
+							}
+						}
+						break;
+
+					case "setStringFromAttribute":
+						{
+							string Result = GetAttribute(CurrentContext, Node, "result");
+							string Tag = GetAttribute(CurrentContext, Node, "tag");
+							string Name = GetAttribute(CurrentContext, Node, "name");
+							if (Result != null && Tag != null && Name != null)
+							{
+								XElement Element = CurrentElement;
+								if (Tag.StartsWith("$"))
+								{
+									if (Tag.Length > 1)
+									{
+										if (!CurrentContext.ElementVariables.TryGetValue(Tag.Substring(1), out Element))
+										{
+											if (!GlobalContext.ElementVariables.TryGetValue(Tag.Substring(1), out Element))
+											{
+												Log.TraceWarning("\nMissing element variable '{0}' in '{1}' (skipping instruction)", Tag, TraceNodeString(Node));
+												continue;
+											}
+										}
+									}
+								}
+
+								XAttribute Attribute;
+								int Index = Name.IndexOf(":");
+								if (Index >= 0)
+								{
+									Name = Name.Substring(Index + 1);
+									Attribute = Element.Attribute(XMLNameSpace + Name);
+								}
+								else
+								{
+									Attribute = Element.Attribute(Name);
+								}
+
+								CurrentContext.StringVariables[Result] = (Attribute != null) ? Attribute.Value : "";
+							}
+						}
+						break;
+
+					case "setStringFromProperty":
+						{
+							string Result = GetAttribute(CurrentContext, Node, "result");
+							string Ini = GetAttribute(CurrentContext, Node, "ini");
+							string Section = GetAttribute(CurrentContext, Node, "section");
+							string Property = GetAttribute(CurrentContext, Node, "property");
+							string DefaultVal = GetAttribute(CurrentContext, Node, "default", true, false, "");
+							if (Result != null && Ini != null && Section != null && Property != null)
+							{
+								string Value = DefaultVal;
+
+								ConfigCacheIni_UPL ConfigIni = GetConfigCacheIni_UPL(Ini);
+								if (ConfigIni != null)
+								{
+									ConfigIni.GetString(Section, Property, out Value);
+								}
+								if (Result == "Output")
+								{
+									GlobalContext.StringVariables["Output"] = Value;
+								}
+								else
+								{
+									CurrentContext.StringVariables[Result] = Value;
+								}
+							}
+						}
+						break;
+
+					case "setStringAdd":
+						{
+							string Result = GetAttribute(CurrentContext, Node, "result");
+							string Arg1 = GetAttribute(CurrentContext, Node, "arg1", true, false, "");
+							string Arg2 = GetAttribute(CurrentContext, Node, "arg2", true, false, "");
+							if (Result != null)
+							{
+								string Value = Arg1 + Arg2;
+								if (Result == "Output")
+								{
+									GlobalContext.StringVariables["Output"] = Value;
+								}
+								else
+								{
+									CurrentContext.StringVariables[Result] = Value;
+								}
+							}
+						}
+						break;
+
+					case "setStringSubstring":
+						{
+							string Result = GetAttribute(CurrentContext, Node, "result");
+							string Source = GetAttribute(CurrentContext, Node, "source", true, false, "");
+							string Start = GetAttribute(CurrentContext, Node, "start", true, false, "0");
+							string Length = GetAttribute(CurrentContext, Node, "length", true, false, "0");
+							if (Result != null && Source != null)
+							{
+								int Index = StringToInt(Start, Node);
+								int Count = StringToInt(Length, Node);
+								Index = (Index < 0) ? 0 : (Index > Source.Length) ? Source.Length : Index;
+								Count = (Index + Count > Source.Length) ? Source.Length - Index : Count;
+								string Value = Source.Substring(Index, Count);
+								if (Result == "Output")
+								{
+									GlobalContext.StringVariables["Output"] = Value;
+								}
+								else
+								{
+									CurrentContext.StringVariables[Result] = Value;
+								}
+							}
+						}
+						break;
+
+					case "setStringReplace":
+						{
+							string Result = GetAttribute(CurrentContext, Node, "result");
+							string Source = GetAttribute(CurrentContext, Node, "source", true, false, "");
+							string Find = GetAttribute(CurrentContext, Node, "find");
+							string With = GetAttribute(CurrentContext, Node, "with", true, false, "");
+							if (Result != null && Find != null)
+							{
+								string Value = Source.Replace(Find, With);
+								if (Result == "Output")
+								{
+									GlobalContext.StringVariables["Output"] = Value;
+								}
+								else
+								{
+									CurrentContext.StringVariables[Result] = Value;
+								}
+							}
+						}
+						break;
+
+					case "setElement":
+						{
+							string Result = GetAttribute(CurrentContext, Node, "result");
+							string Value = GetAttribute(CurrentContext, Node, "value", true, false);
+							string Parse = GetAttribute(CurrentContext, Node, "xml", true, false);
+							if (Result != null)
+							{
+								if (Value != null)
+								{
+									CurrentContext.ElementVariables[Result] = new XElement(Value);
+								}
+								else if (Parse != null)
+								{
+									try
+									{
+										CurrentContext.ElementVariables[Result] = XElement.Parse(Parse);
+									}
+									catch (Exception e)
+									{
+										Log.TraceError("\nXML parsing {0} failed! {1} (skipping instruction)", Parse, e);
+									}
+								}
+							}
+						}
+						break;
+
+					default:
+						Log.TraceWarning("\nUnknown command: {0}", Node.Name);
+						break;
+				}
+			}
+
+			return GlobalContext.StringVariables["Output"];
+		}
+
+		public void Init(List<string> Architectures, bool bDistribution, string EngineDirectory, string BuildDirectory)
+		{
+			GlobalContext.BoolVariables["Distribution"] = bDistribution;
+			GlobalContext.StringVariables["EngineDir"] = EngineDirectory;
+			GlobalContext.StringVariables["BuildDir"] = BuildDirectory;
+
+			foreach (string Arch in Architectures)
+			{
+				Log.TraceInformation("UPL Init: {0}", Arch);
+				ProcessPluginNode(Arch, "init", "");
+			}
+
+			if (bGlobalTrace)
+			{
+				Log.TraceInformation("\nVariables:\n{0}", DumpVariables());
+			}
+		}
+	}
+
+	/// <summary>
+	/// Equivalent of FConfigCacheIni_UPL. Parses ini files.  This version reads ALL sections since ConfigCacheIni_UPL does NOT
+	/// </summary>
+	public class ConfigCacheIni_UPL
+	{
+		/// <summary>
+		/// Exception when parsing ini files
+		/// </summary>
+		public class IniParsingException : Exception
+		{
+			public IniParsingException(string Message)
+				: base(Message)
+			{ }
+			public IniParsingException(string Format, params object[] Args)
+				: base(String.Format(Format, Args))
+			{ }
+		}
+
+
+		// command class for being able to create config caches over and over without needing to read the ini files
+		public class Command
+		{
+			public string TrimmedLine;
+		}
+
+		class SectionCommand : Command
+		{
+			public FileReference Filename;
+			public int LineIndex;
+		}
+
+		class KeyValueCommand : Command
+		{
+			public string Key;
+			public string Value;
+			public ParseAction LastAction;
+		}
+
+		// cached ini files
+		static Dictionary<string, List<Command>> FileCache = new Dictionary<string, List<Command>>();
+		static Dictionary<string, ConfigCacheIni_UPL> IniCache = new Dictionary<string, ConfigCacheIni_UPL>();
+		static Dictionary<string, ConfigCacheIni_UPL> BaseIniCache = new Dictionary<string, ConfigCacheIni_UPL>();
+
+		// static creation functions for ini files
+		public static ConfigCacheIni_UPL CreateConfigCacheIni_UPL(UnrealTargetPlatform Platform, string BaseIniName, DirectoryReference ProjectDirectory, DirectoryReference EngineDirectory = null)
+		{
+			if (EngineDirectory == null)
+			{
+				EngineDirectory = UnrealBuildTool.EngineDirectory;
+			}
+
+			// cache base ini for use as the seed for the rest
+			if (!BaseIniCache.ContainsKey(BaseIniName))
+			{
+				BaseIniCache.Add(BaseIniName, new ConfigCacheIni_UPL(UnrealTargetPlatform.Unknown, BaseIniName, null, EngineDirectory, EngineOnly: true));
+			}
+
+			// build the new ini and cache it for later re-use
+			ConfigCacheIni_UPL BaseCache = BaseIniCache[BaseIniName];
+			string Key = GetIniPlatformName(Platform) + BaseIniName + EngineDirectory.FullName + (ProjectDirectory != null ? ProjectDirectory.FullName : "");
+			if (!IniCache.ContainsKey(Key))
+			{
+				IniCache.Add(Key, new ConfigCacheIni_UPL(Platform, BaseIniName, ProjectDirectory, EngineDirectory, BaseCache: BaseCache));
+			}
+			return IniCache[Key];
+		}
+
+		/// <summary>
+		/// List of values (or a single value)
+		/// </summary>
+		public class IniValues : List<string>
+		{
+			public IniValues()
+			{
+			}
+			public IniValues(IniValues Other)
+				: base(Other)
+			{
+			}
+			public override string ToString()
+			{
+				return String.Join(",", ToArray());
+			}
+		}
+
+		/// <summary>
+		/// Ini section (map of keys and values)
+		/// </summary>
+		public class IniSection : Dictionary<string, IniValues>
+		{
+			public IniSection()
+				: base(StringComparer.InvariantCultureIgnoreCase)
+			{ }
+			public IniSection(IniSection Other)
+				: this()
+			{
+				foreach (var Pair in Other)
+				{
+					Add(Pair.Key, new IniValues(Pair.Value));
+				}
+			}
+			public override string ToString()
+			{
+				return "IniSection";
+			}
+		}
+
+		/// <summary>
+		/// True if we are loading a hierarchy of config files that should be merged together
+		/// </summary>
+		bool bIsMergingConfigs;
+
+		/// <summary>
+		/// All sections parsed from ini file
+		/// </summary>
+		Dictionary<string, IniSection> Sections;
+
+		/// <summary>
+		/// Constructor. Parses a single ini file. No Platform settings, no engine hierarchy. Do not use this with ini files that have hierarchy!
+		/// </summary>
+		/// <param name="Filename">The ini file to load</param>
+		public ConfigCacheIni_UPL(FileReference Filename)
+		{
+			Init(Filename);
+		}
+
+		/// <summary>
+		/// Constructor. Parses ini hierarchy for the specified project.  No Platform settings.
+		/// </summary>
+		/// <param name="ProjectDirectory">Project path</param>
+		/// <param name="Platform">Target platform</param>
+		/// <param name="BaseIniName">Ini name (Engine, Editor, etc)</param>
+		public ConfigCacheIni_UPL(string BaseIniName, string ProjectDirectory, string EngineDirectory = null)
+		{
+			Init(UnrealTargetPlatform.Unknown, BaseIniName, (ProjectDirectory == null) ? null : new DirectoryReference(ProjectDirectory), (EngineDirectory == null) ? null : new DirectoryReference(EngineDirectory));
+		}
+
+		/// <summary>
+		/// Constructor. Parses ini hierarchy for the specified project.  No Platform settings.
+		/// </summary>
+		/// <param name="ProjectDirectory">Project path</param>
+		/// <param name="Platform">Target platform</param>
+		/// <param name="BaseIniName">Ini name (Engine, Editor, etc)</param>
+		public ConfigCacheIni_UPL(string BaseIniName, DirectoryReference ProjectDirectory, DirectoryReference EngineDirectory = null)
+		{
+			Init(UnrealTargetPlatform.Unknown, BaseIniName, ProjectDirectory, EngineDirectory);
+		}
+
+		/// <summary>
+		/// Constructor. Parses ini hierarchy for the specified platform and project.
+		/// </summary>
+		/// <param name="ProjectDirectory">Project path</param>
+		/// <param name="Platform">Target platform</param>
+		/// <param name="BaseIniName">Ini name (Engine, Editor, etc)</param>
+		public ConfigCacheIni_UPL(UnrealTargetPlatform Platform, string BaseIniName, string ProjectDirectory, string EngineDirectory = null)
+		{
+			Init(Platform, BaseIniName, (ProjectDirectory == null) ? null : new DirectoryReference(ProjectDirectory), (EngineDirectory == null) ? null : new DirectoryReference(EngineDirectory));
+		}
+
+		/// <summary>
+		/// Constructor. Parses ini hierarchy for the specified platform and project.
+		/// </summary>
+		/// <param name="ProjectDirectory">Project path</param>
+		/// <param name="Platform">Target platform</param>
+		/// <param name="BaseIniName">Ini name (Engine, Editor, etc)</param>
+		public ConfigCacheIni_UPL(UnrealTargetPlatform Platform, string BaseIniName, DirectoryReference ProjectDirectory, DirectoryReference EngineDirectory = null, bool EngineOnly = false, ConfigCacheIni_UPL BaseCache = null)
+		{
+			Init(Platform, BaseIniName, ProjectDirectory, EngineDirectory, EngineOnly, BaseCache);
+		}
+
+		private void InitCommon()
+		{
+			Sections = new Dictionary<string, IniSection>(StringComparer.InvariantCultureIgnoreCase);
+		}
+
+		private void Init(FileReference IniFileName)
+		{
+			InitCommon();
+			bIsMergingConfigs = false;
+			ParseIniFile(IniFileName);
+		}
+
+		private void Init(UnrealTargetPlatform Platform, string BaseIniName, DirectoryReference ProjectDirectory, DirectoryReference EngineDirectory, bool EngineOnly = false, ConfigCacheIni_UPL BaseCache = null)
+		{
+			InitCommon();
+			bIsMergingConfigs = true;
+			if (EngineDirectory == null)
+			{
+				EngineDirectory = UnrealBuildTool.EngineDirectory;
+			}
+
+			if (BaseCache != null)
+			{
+				foreach (var Pair in BaseCache.Sections)
+				{
+					Sections.Add(Pair.Key, new IniSection(Pair.Value));
+				}
+			}
+			if (EngineOnly)
+			{
+				foreach (var IniFileName in EnumerateEngineIniFileNames(EngineDirectory, BaseIniName))
+				{
+					if (IniFileName.Exists())
+					{
+						ParseIniFile(IniFileName);
+					}
+				}
+			}
+			else
+			{
+				foreach (var IniFileName in EnumerateCrossPlatformIniFileNames(ProjectDirectory, EngineDirectory, Platform, BaseIniName, BaseCache != null))
+				{
+					if (IniFileName.Exists())
+					{
+						ParseIniFile(IniFileName);
+					}
+				}
+			}
+		}
+
+		/// <summary>
+		/// Finds a section in INI
+		/// </summary>
+		/// <param name="SectionName"></param>
+		/// <returns>Found section or null</returns>
+		public IniSection FindSection(string SectionName)
+		{
+			IniSection Section;
+			Sections.TryGetValue(SectionName, out Section);
+			return Section;
+		}
+
+		/// <summary>
+		/// Finds values associated with the specified key (does not copy the list)
+		/// </summary>
+		private bool GetList(string SectionName, string Key, out IniValues Value)
+		{
+			bool Result = false;
+			var Section = FindSection(SectionName);
+			Value = null;
+			if (Section != null)
+			{
+				if (Section.TryGetValue(Key, out Value))
+				{
+					Result = true;
+				}
+			}
+			return Result;
+		}
+
+		/// <summary>
+		/// Gets all values associated with the specified key
+		/// </summary>
+		/// <param name="SectionName">Section where the key is located</param>
+		/// <param name="Key">Key name</param>
+		/// <param name="Value">Copy of the list containing all values associated with the specified key</param>
+		/// <returns>True if the key exists</returns>
+		public bool GetArray(string SectionName, string Key, out List<string> Value)
+		{
+			Value = null;
+			IniValues ValueList;
+			bool Result = GetList(SectionName, Key, out ValueList);
+			if (Result)
+			{
+				Value = new List<string>(ValueList);
+			}
+			return Result;
+		}
+
+		/// <summary>
+		/// Gets a single string value associated with the specified key.
+		/// </summary>
+		/// <param name="SectionName">Section name</param>
+		/// <param name="Key">Key name</param>
+		/// <param name="Value">Value associated with the specified key. If the key has more than one value, only the first one is returned</param>
+		/// <returns>True if the key exists</returns>
+		public bool GetString(string SectionName, string Key, out string Value)
+		{
+			Value = String.Empty;
+			IniValues ValueList;
+			bool Result = GetList(SectionName, Key, out ValueList);
+			if (Result && ValueList != null && ValueList.Count > 0)
+			{
+				Value = ValueList[0];
+				Result = true;
+			}
+			else
+			{
+				Result = false;
+			}
+			return Result;
+		}
+
+		/// <summary>
+		/// Gets a single bool value associated with the specified key.
+		/// </summary>
+		/// <param name="SectionName">Section name</param>
+		/// <param name="Key">Key name</param>
+		/// <param name="Value">Value associated with the specified key. If the key has more than one value, only the first one is returned</param>
+		/// <returns>True if the key exists</returns>
+		public bool GetBool(string SectionName, string Key, out bool Value)
+		{
+			Value = false;
+			string TextValue;
+			bool Result = GetString(SectionName, Key, out TextValue);
+			if (Result)
+			{
+				// C# Boolean type expects "False" or "True" but since we're not case sensitive, we need to suppor that manually
+				if (String.Compare(TextValue, "true", true) == 0 || String.Compare(TextValue, "1") == 0)
+				{
+					Value = true;
+				}
+				else if (String.Compare(TextValue, "false", true) == 0 || String.Compare(TextValue, "0") == 0)
+				{
+					Value = false;
+				}
+				else
+				{
+					// Failed to parse
+					Result = false;
+				}
+			}
+			return Result;
+		}
+
+		/// <summary>
+		/// Gets a single Int32 value associated with the specified key.
+		/// </summary>
+		/// <param name="SectionName">Section name</param>
+		/// <param name="Key">Key name</param>
+		/// <param name="Value">Value associated with the specified key. If the key has more than one value, only the first one is returned</param>
+		/// <returns>True if the key exists</returns>
+		public bool GetInt32(string SectionName, string Key, out int Value)
+		{
+			Value = 0;
+			string TextValue;
+			bool Result = GetString(SectionName, Key, out TextValue);
+			if (Result)
+			{
+				Result = Int32.TryParse(TextValue, out Value);
+			}
+			return Result;
+		}
+
+		/// <summary>
+		/// Gets a single GUID value associated with the specified key.
+		/// </summary>
+		/// <param name="SectionName">Section name</param>
+		/// <param name="Key">Key name</param>
+		/// <param name="Value">Value associated with the specified key. If the key has more than one value, only the first one is returned</param>
+		/// <returns>True if the key exists</returns>
+		public bool GetGUID(string SectionName, string Key, out Guid Value)
+		{
+			Value = Guid.Empty;
+			string TextValue;
+			bool Result = GetString(SectionName, Key, out TextValue);
+			if (Result)
+			{
+				string HexString = "";
+				if (TextValue.Contains("A=") && TextValue.Contains("B=") && TextValue.Contains("C=") && TextValue.Contains("D="))
+				{
+					char[] Separators = new char[] { '(', ')', '=', ',', ' ', 'A', 'B', 'C', 'D' };
+					string[] ComponentValues = TextValue.Split(Separators, StringSplitOptions.RemoveEmptyEntries);
+					if (ComponentValues.Length == 4)
+					{
+						for (int ComponentIndex = 0; ComponentIndex < 4; ComponentIndex++)
+						{
+							int IntegerValue;
+							Result &= Int32.TryParse(ComponentValues[ComponentIndex], out IntegerValue);
+							HexString += IntegerValue.ToString("X8");
+						}
+					}
+				}
+				else
+				{
+					HexString = TextValue;
+				}
+
+				try
+				{
+					Value = Guid.ParseExact(HexString, "N");
+					Result = true;
+				}
+				catch (Exception)
+				{
+					Result = false;
+				}
+			}
+			return Result;
+		}
+
+		/// <summary>
+		/// Gets a single float value associated with the specified key.
+		/// </summary>
+		/// <param name="SectionName">Section name</param>
+		/// <param name="Key">Key name</param>
+		/// <param name="Value">Value associated with the specified key. If the key has more than one value, only the first one is returned</param>
+		/// <returns>True if the key exists</returns>
+		public bool GetSingle(string SectionName, string Key, out float Value)
+		{
+			Value = 0.0f;
+			string TextValue;
+			bool Result = GetString(SectionName, Key, out TextValue);
+			if (Result)
+			{
+				Result = Single.TryParse(TextValue, out Value);
+			}
+			return Result;
+		}
+
+		/// <summary>
+		/// Gets a single double value associated with the specified key.
+		/// </summary>
+		/// <param name="SectionName">Section name</param>
+		/// <param name="Key">Key name</param>
+		/// <param name="Value">Value associated with the specified key. If the key has more than one value, only the first one is returned</param>
+		/// <returns>True if the key exists</returns>
+		public bool GetDouble(string SectionName, string Key, out double Value)
+		{
+			Value = 0.0;
+			string TextValue;
+			bool Result = GetString(SectionName, Key, out TextValue);
+			if (Result)
+			{
+				Result = Double.TryParse(TextValue, out Value);
+			}
+			return Result;
+		}
+
+		private static bool ExtractPath(string Source, out string Path)
+		{
+			int start = Source.IndexOf('"');
+			int end = Source.LastIndexOf('"');
+			if (start != 1 && end != -1 && start < end)
+			{
+				++start;
+				Path = Source.Substring(start, end - start);
+				return true;
+			}
+			else
+			{
+				Path = "";
+			}
+
+			return false;
+		}
+
+		public bool GetPath(string SectionName, string Key, out string Value)
+		{
+			string temp;
+			if (GetString(SectionName, Key, out temp))
+			{
+				return ExtractPath(temp, out Value);
+			}
+			else
+			{
+				Value = "";
+			}
+
+			return false;
+		}
+
+		/// <summary>
+		/// List of actions that can be performed on a single line from ini file
+		/// </summary>
+		enum ParseAction
+		{
+			None,
+			New,
+			Add,
+			Remove
+		}
+
+		/// <summary>
+		/// Checks what action should be performed on a single line from ini file
+		/// </summary>
+		private ParseAction GetActionForLine(ref string Line)
+		{
+			if (String.IsNullOrEmpty(Line) || Line.StartsWith(";") || Line.StartsWith("//"))
+			{
+				return ParseAction.None;
+			}
+			else if (Line.StartsWith("-"))
+			{
+				Line = Line.Substring(1).TrimStart();
+				return ParseAction.Remove;
+			}
+			else if (Line.StartsWith("+"))
+			{
+				Line = Line.Substring(1).TrimStart();
+				return ParseAction.Add;
+			}
+			else
+			{
+				// We use Add rather than New when we're not merging config files together in order 
+				// to mimic the behavior of the C++ config cache when loading a single file
+				return (bIsMergingConfigs) ? ParseAction.New : ParseAction.Add;
+			}
+		}
+
+		/// <summary>
+		/// Loads and parses ini file.
+		/// </summary>
+		public void ParseIniFile(FileReference Filename)
+		{
+			String[] IniLines = null;
+			List<Command> Commands = null;
+			if (!FileCache.ContainsKey(Filename.FullName))
+			{
+				try
+				{
+					IniLines = File.ReadAllLines(Filename.FullName);
+					Commands = new List<Command>();
+					FileCache.Add(Filename.FullName, Commands);
+				}
+				catch (Exception ex)
+				{
+					Console.WriteLine("Error reading ini file: " + Filename + " Exception: " + ex.Message);
+				}
+			}
+			else
+			{
+				Commands = FileCache[Filename.FullName];
+			}
+			if (IniLines != null)
+			{
+				IniSection CurrentSection = null;
+
+				// Line Index for exceptions
+				var LineIndex = 1;
+				var bMultiLine = false;
+				var SingleValue = "";
+				var Key = "";
+				var LastAction = ParseAction.None;
+
+				// Parse each line
+				foreach (var Line in IniLines)
+				{
+					var TrimmedLine = Line.Trim();
+					// Multiline value support
+					bool bWasMultiLine = bMultiLine;
+					bMultiLine = TrimmedLine.EndsWith("\\");
+					if (bMultiLine)
+					{
+						TrimmedLine = TrimmedLine.Substring(0, TrimmedLine.Length - 1).TrimEnd();
+					}
+					if (!bWasMultiLine)
+					{
+						if (TrimmedLine.StartsWith("["))
+						{
+							CurrentSection = FindOrAddSection(TrimmedLine, Filename, LineIndex);
+							LastAction = ParseAction.None;
+							if (CurrentSection != null)
+							{
+								SectionCommand Command = new SectionCommand();
+								Command.Filename = Filename;
+								Command.LineIndex = LineIndex;
+								Command.TrimmedLine = TrimmedLine;
+								Commands.Add(Command);
+							}
+						}
+						else
+						{
+							if (LastAction != ParseAction.None)
+							{
+								throw new IniParsingException("Parsing new key/value pair when the previous one has not yet been processed ({0}, {1}) in {2}, line {3}: {4}", Key, SingleValue, Filename, LineIndex, TrimmedLine);
+							}
+							// Check if the line is empty or a comment, also remove any +/- markers
+							LastAction = GetActionForLine(ref TrimmedLine);
+							if (LastAction != ParseAction.None)
+							{
+								/*								if (CurrentSection == null)
+																{
+																	throw new IniParsingException("Trying to parse key/value pair that doesn't belong to any section in {0}, line {1}: {2}", Filename, LineIndex, TrimmedLine);
+																}*/
+								ParseKeyValuePair(TrimmedLine, Filename, LineIndex, out Key, out SingleValue);
+							}
+						}
+					}
+					if (bWasMultiLine)
+					{
+						SingleValue += TrimmedLine;
+					}
+					if (!bMultiLine && LastAction != ParseAction.None && CurrentSection != null)
+					{
+						ProcessKeyValuePair(CurrentSection, Key, SingleValue, LastAction);
+						KeyValueCommand Command = new KeyValueCommand();
+						Command.Key = Key;
+						Command.Value = SingleValue;
+						Command.LastAction = LastAction;
+						Commands.Add(Command);
+						LastAction = ParseAction.None;
+						SingleValue = "";
+						Key = "";
+					}
+					else if (CurrentSection == null)
+					{
+						LastAction = ParseAction.None;
+					}
+					LineIndex++;
+				}
+			}
+			else if (Commands != null)
+			{
+				IniSection CurrentSection = null;
+
+				// run each command
+				for (int Idx = 0; Idx < Commands.Count; ++Idx)
+				{
+					var Command = Commands[Idx];
+					if (Command is SectionCommand)
+					{
+						CurrentSection = FindOrAddSection((Command as SectionCommand).TrimmedLine, (Command as SectionCommand).Filename, (Command as SectionCommand).LineIndex);
+					}
+					else if (Command is KeyValueCommand)
+					{
+						ProcessKeyValuePair(CurrentSection, (Command as KeyValueCommand).Key, (Command as KeyValueCommand).Value, (Command as KeyValueCommand).LastAction);
+					}
+				}
+			}
+		}
+
+		/// <summary>
+		/// Splits a line into key and value
+		/// </summary>
+		private static void ParseKeyValuePair(string TrimmedLine, FileReference Filename, int LineIndex, out string Key, out string Value)
+		{
+			var AssignIndex = TrimmedLine.IndexOf('=');
+			if (AssignIndex < 0)
+			{
+				throw new IniParsingException("Failed to find value when parsing {0}, line {1}: {2}", Filename, LineIndex, TrimmedLine);
+			}
+			Key = TrimmedLine.Substring(0, AssignIndex).Trim();
+			if (String.IsNullOrEmpty(Key))
+			{
+				throw new IniParsingException("Empty key when parsing {0}, line {1}: {2}", Filename, LineIndex, TrimmedLine);
+			}
+			Value = TrimmedLine.Substring(AssignIndex + 1).Trim();
+			if (Value.StartsWith("\""))
+			{
+				// Remove quotes
+				var QuoteEnd = Value.LastIndexOf('\"');
+				if (QuoteEnd == 0)
+				{
+					throw new IniParsingException("Mismatched quotes when parsing {0}, line {1}: {2}", Filename, LineIndex, TrimmedLine);
+				}
+				Value = Value.Substring(1, Value.Length - 2);
+			}
+		}
+
+		/// <summary>
+		/// Processes parsed key/value pair
+		/// </summary>
+		private static void ProcessKeyValuePair(IniSection CurrentSection, string Key, string SingleValue, ParseAction Action)
+		{
+			switch (Action)
+			{
+				case ParseAction.New:
+					{
+						// New/replace
+						IniValues Value;
+						if (CurrentSection.TryGetValue(Key, out Value) == false)
+						{
+							Value = new IniValues();
+							CurrentSection.Add(Key, Value);
+						}
+						Value.Clear();
+						Value.Add(SingleValue);
+					}
+					break;
+				case ParseAction.Add:
+					{
+						IniValues Value;
+						if (CurrentSection.TryGetValue(Key, out Value) == false)
+						{
+							Value = new IniValues();
+							CurrentSection.Add(Key, Value);
+						}
+						Value.Add(SingleValue);
+					}
+					break;
+				case ParseAction.Remove:
+					{
+						IniValues Value;
+						if (CurrentSection.TryGetValue(Key, out Value))
+						{
+							var ExistingIndex = Value.FindIndex(X => (String.Compare(SingleValue, X, true) == 0));
+							if (ExistingIndex >= 0)
+							{
+								Value.RemoveAt(ExistingIndex);
+							}
+						}
+					}
+					break;
+			}
+		}
+
+		/// <summary>
+		/// Finds an existing section or adds a new one
+		/// </summary>
+		private IniSection FindOrAddSection(string TrimmedLine, FileReference Filename, int LineIndex)
+		{
+			var SectionEndIndex = TrimmedLine.IndexOf(']');
+			if (SectionEndIndex != (TrimmedLine.Length - 1))
+			{
+				throw new IniParsingException("Mismatched brackets when parsing section name in {0}, line {1}: {2}", Filename, LineIndex, TrimmedLine);
+			}
+			var SectionName = TrimmedLine.Substring(1, TrimmedLine.Length - 2);
+			if (String.IsNullOrEmpty(SectionName))
+			{
+				throw new IniParsingException("Empty section name when parsing {0}, line {1}: {2}", Filename, LineIndex, TrimmedLine);
+			}
+			{
+				IniSection CurrentSection;
+				if (Sections.TryGetValue(SectionName, out CurrentSection) == false)
+				{
+					CurrentSection = new IniSection();
+					Sections.Add(SectionName, CurrentSection);
+				}
+				return CurrentSection;
+			}
+		}
+
+		/// <summary>
+		/// Returns a list of INI filenames for the engine
+		/// </summary>
+		private static IEnumerable<FileReference> EnumerateEngineIniFileNames(DirectoryReference EngineDirectory, string BaseIniName)
+		{
+			// Engine/Config/Base.ini (included in every ini type, required)
+			yield return FileReference.Combine(EngineDirectory, "Config", "Base.ini");
+
+			// Engine/Config/Base* ini
+			yield return FileReference.Combine(EngineDirectory, "Config", "Base" + BaseIniName + ".ini");
+
+			// Engine/Config/NotForLicensees/Base* ini
+			yield return FileReference.Combine(EngineDirectory, "Config", "NotForLicensees", "Base" + BaseIniName + ".ini");
+		}
+
+
+		/// <summary>
+		/// Returns a list of INI filenames for the given project
+		/// </summary>
+		private static IEnumerable<FileReference> EnumerateCrossPlatformIniFileNames(DirectoryReference ProjectDirectory, DirectoryReference EngineDirectory, UnrealTargetPlatform Platform, string BaseIniName, bool SkipEngine)
+		{
+			if (!SkipEngine)
+			{
+				// Engine/Config/Base.ini (included in every ini type, required)
+				yield return FileReference.Combine(EngineDirectory, "Config", "Base.ini");
+
+				// Engine/Config/Base* ini
+				yield return FileReference.Combine(EngineDirectory, "Config", "Base" + BaseIniName + ".ini");
+
+				// Engine/Config/NotForLicensees/Base* ini
+				yield return FileReference.Combine(EngineDirectory, "Config", "NotForLicensees", "Base" + BaseIniName + ".ini");
+
+				// NOTE: 4.7: See comment in GetSourceIniHierarchyFilenames()
+				// Engine/Config/NoRedist/Base* ini
+				// yield return Path.Combine(EngineDirectory, "Config", "NoRedist", "Base" + BaseIniName + ".ini");
+			}
+
+			if (ProjectDirectory != null)
+			{
+				// Game/Config/Default* ini
+				yield return FileReference.Combine(ProjectDirectory, "Config", "Default" + BaseIniName + ".ini");
+
+				// Game/Config/NotForLicensees/Default* ini
+				yield return FileReference.Combine(ProjectDirectory, "Config", "NotForLicensees", "Default" + BaseIniName + ".ini");
+
+				// Game/Config/NoRedist/Default* ini
+				yield return FileReference.Combine(ProjectDirectory, "Config", "NoRedist", "Default" + BaseIniName + ".ini");
+			}
+
+			string PlatformName = GetIniPlatformName(Platform);
+			if (Platform != UnrealTargetPlatform.Unknown)
+			{
+				// Engine/Config/Platform/Platform* ini
+				yield return FileReference.Combine(EngineDirectory, "Config", PlatformName, PlatformName + BaseIniName + ".ini");
+
+				if (ProjectDirectory != null)
+				{
+					// Game/Config/Platform/Platform* ini
+					yield return FileReference.Combine(ProjectDirectory, "Config", PlatformName, PlatformName + BaseIniName + ".ini");
+				}
+			}
+
+			DirectoryReference UserSettingsFolder = Utils.GetUserSettingDirectory(); // Match FPlatformProcess::UserSettingsDir()
+			DirectoryReference PersonalFolder = null; // Match FPlatformProcess::UserDir()
+			if (BuildHostPlatform.Current.Platform == UnrealTargetPlatform.Mac)
+			{
+				PersonalFolder = new DirectoryReference(Path.Combine(Environment.GetFolderPath(Environment.SpecialFolder.Personal), "Documents"));
+			}
+			else if (Environment.OSVersion.Platform == PlatformID.Unix)
+			{
+				PersonalFolder = new DirectoryReference(Path.Combine(Environment.GetFolderPath(Environment.SpecialFolder.Personal), "Documents"));
+			}
+			else
+			{
+				// Not all user accounts have a local application data directory (eg. SYSTEM, used by Jenkins for builds).
+				string PersonalFolderSetting = Environment.GetFolderPath(Environment.SpecialFolder.Personal);
+				if(!String.IsNullOrEmpty(PersonalFolderSetting))
+				{
+					PersonalFolder = new DirectoryReference(PersonalFolderSetting);
+				}
+			}
+
+			if(UserSettingsFolder != null)
+			{
+				// <AppData>/UE4/EngineConfig/User* ini
+				yield return FileReference.Combine(UserSettingsFolder, "Unreal Engine", "Engine", "Config", "User" + BaseIniName + ".ini");
+			}
+
+			if(PersonalFolder != null)
+			{
+				// <Documents>/UE4/EngineConfig/User* ini
+				yield return FileReference.Combine(PersonalFolder, "Unreal Engine", "Engine", "Config", "User" + BaseIniName + ".ini");
+			}
+
+			// Game/Config/User* ini
+			if (ProjectDirectory != null)
+			{
+				yield return FileReference.Combine(ProjectDirectory, "Config", "User" + BaseIniName + ".ini");
+			}
+		}
+
+		/// <summary>
+		/// Returns the platform name to use as part of platform-specific config files
+		/// </summary>
+		private static string GetIniPlatformName(UnrealTargetPlatform TargetPlatform)
+		{
+			if (TargetPlatform == UnrealTargetPlatform.Win32 || TargetPlatform == UnrealTargetPlatform.Win64 || TargetPlatform == UnrealTargetPlatform.UWP)
+			{
+				return "Windows";
+			}
+			else
+			{
+				return Enum.GetName(typeof(UnrealTargetPlatform), TargetPlatform);
+			}
+		}
+	}
+
+}
