fix for invalid plist generation for iOS and tvOS

diff --git a/Engine/Source/Programs/UnrealBuildTool/IOS/UEDeployIOS.cs b/Engine/Source/Programs/UnrealBuildTool/IOS/UEDeployIOS.cs
--- a/Engine/Source/Programs/UnrealBuildTool/IOS/UEDeployIOS.cs
+++ b/Engine/Source/Programs/UnrealBuildTool/IOS/UEDeployIOS.cs
@@ -517,8 +517,8 @@ namespace UnrealBuildTool
 
 				XDoc.DocumentType.InternalSubset = "";
 				InThis.UPL.ProcessPluginNode("None", "iosPListUpdates", "", ref XDoc);
-				string result = XDoc.Declaration.ToString() + "\n" + XDoc.ToString();
-				File.WriteAllText(PListFile, result);
+                string result = XDoc.Declaration.ToString() + "\n" + XDoc.ToString().Replace("<!DOCTYPE plist PUBLIC \"-//Apple//DTD PLIST 1.0//EN\" \"http://www.apple.com/DTDs/PropertyList-1.0.dtd\"[]>", "<!DOCTYPE plist PUBLIC \"-//Apple//DTD PLIST 1.0//EN\" \"http://www.apple.com/DTDs/PropertyList-1.0.dtd\">");
+                File.WriteAllText(PListFile, result);
 			}
 			else
 			{
