Fix arm64 build machine issue (using wrong ndk because of old setup.bat run, bumped script version)
#ue4
#android

diff --git a/Engine/Source/Programs/UnrealBuildTool/Android/UEBuildAndroid.cs b/Engine/Source/Programs/UnrealBuildTool/Android/UEBuildAndroid.cs
--- a/Engine/Source/Programs/UnrealBuildTool/Android/UEBuildAndroid.cs
+++ b/Engine/Source/Programs/UnrealBuildTool/Android/UEBuildAndroid.cs
@@ -449,7 +449,7 @@ namespace UnrealBuildTool
 
 		protected override String GetRequiredScriptVersionString()
 		{
-			return "3.0";
+			return "3.1";
 		}
 
 		// prefer auto sdk on android as correct 'manual' sdk detection isn't great at the moment.
