When generating the list of runtime dependencies for precompiled modules, ignore any project-relative paths unless a project file is present. Fixes warnings in 'Make Installed Build' nodes.

diff --git a/Engine/Source/Programs/UnrealBuildTool/Configuration/UEBuildTarget.cs b/Engine/Source/Programs/UnrealBuildTool/Configuration/UEBuildTarget.cs
--- a/Engine/Source/Programs/UnrealBuildTool/Configuration/UEBuildTarget.cs
+++ b/Engine/Source/Programs/UnrealBuildTool/Configuration/UEBuildTarget.cs
@@ -2048,8 +2048,12 @@ namespace UnrealBuildTool
 							{
 								foreach (RuntimeDependency RuntimeDependency in Module.RuntimeDependencies)
 								{
-									string SourcePath = TargetReceipt.InsertPathVariables(RuntimeDependency.Path, UnrealBuildTool.EngineDirectory, ProjectDirectory);
-									Receipt.PrecompiledRuntimeDependencies.Add(SourcePath);
+									// Ignore project-relative dependencies when we're compiling targets without projects - we won't be able to resolve them.
+									if(ProjectFile != null || RuntimeDependency.Path.IndexOf("$(ProjectDir)", StringComparison.InvariantCultureIgnoreCase) == -1)
+									{
+										string SourcePath = TargetReceipt.InsertPathVariables(RuntimeDependency.Path, UnrealBuildTool.EngineDirectory, ProjectDirectory);
+										Receipt.PrecompiledRuntimeDependencies.Add(SourcePath);
+									}
 								}
 							}
 						}
