Added Framework zip paths to list of required external files when precompiling targets

diff --git a/Engine/Source/Programs/UnrealBuildTool/Configuration/UEBuildTarget.cs b/Engine/Source/Programs/UnrealBuildTool/Configuration/UEBuildTarget.cs
--- a/Engine/Source/Programs/UnrealBuildTool/Configuration/UEBuildTarget.cs
+++ b/Engine/Source/Programs/UnrealBuildTool/Configuration/UEBuildTarget.cs
@@ -1760,6 +1760,15 @@ namespace UnrealBuildTool
 					Files.Add(new FileReference(Resource.ResourcePath));
 				}
 
+				// Add any zip files from Additional Frameworks
+				foreach (UEBuildFramework Framework in Rules.PublicAdditionalFrameworks)
+				{
+					if (!String.IsNullOrEmpty(Framework.FrameworkZipPath))
+					{
+						Files.Add(FileReference.Combine(Module.ModuleDirectory, Framework.FrameworkZipPath));
+					}
+				}
+
 				// Add all the include paths etc. for external modules
 				UEBuildExternalModule ExternalModule = Module as UEBuildExternalModule;
 				if (ExternalModule != null)
