Fix concatenation of paths containing non-native separators on Mac. Mono sets Path.AltDirectorySeparatorChar to '/', not '\'.

diff --git a/Engine/Source/Programs/UnrealBuildTool/System/FileSystemReference.cs b/Engine/Source/Programs/UnrealBuildTool/System/FileSystemReference.cs
--- a/Engine/Source/Programs/UnrealBuildTool/System/FileSystemReference.cs
+++ b/Engine/Source/Programs/UnrealBuildTool/System/FileSystemReference.cs
@@ -69,7 +69,7 @@ namespace UnrealBuildTool
 			foreach (string Fragment in Fragments)
 			{
 				// Check if this fragment is an absolute path
-				if ((Fragment.Length >= 2 && Fragment[1] == ':') || (Fragment.Length >= 1 && (Fragment[0] == Path.DirectorySeparatorChar || Fragment[0] == Path.AltDirectorySeparatorChar)))
+				if ((Fragment.Length >= 2 && Fragment[1] == ':') || (Fragment.Length >= 1 && (Fragment[0] == '\\' || Fragment[0] == '/')))
 				{
 					// It is. Reset the new name to the full version of this path.
 					NewFullName.Clear();
@@ -83,7 +83,7 @@ namespace UnrealBuildTool
 					{
 						// Find the end of this fragment. We may have been passed multiple paths in the same string.
 						int EndIdx = StartIdx;
-						while (EndIdx < Fragment.Length && Fragment[EndIdx] != Path.DirectorySeparatorChar && Fragment[EndIdx] != Path.AltDirectorySeparatorChar)
+						while (EndIdx < Fragment.Length && Fragment[EndIdx] != '\\' && Fragment[EndIdx] != '/')
 						{
 							EndIdx++;
 						}
