Fix warning packaging on Mac due to RadioEffectComponent being a directory rather than a file.

diff --git a/Engine/Source/Programs/UnrealBuildTool/Configuration/UEBuildTarget.cs b/Engine/Source/Programs/UnrealBuildTool/Configuration/UEBuildTarget.cs
--- a/Engine/Source/Programs/UnrealBuildTool/Configuration/UEBuildTarget.cs
+++ b/Engine/Source/Programs/UnrealBuildTool/Configuration/UEBuildTarget.cs
@@ -1757,7 +1757,14 @@ namespace UnrealBuildTool
 				// Add Additional Bundle Resources for all modules
 				foreach (UEBuildBundleResource Resource in Rules.AdditionalBundleResources)
 				{
-					Files.Add(new FileReference(Resource.ResourcePath));
+					if (Directory.Exists(Resource.ResourcePath))
+					{
+						Files.UnionWith(new DirectoryReference(Resource.ResourcePath).EnumerateFileReferences("*", SearchOption.AllDirectories));
+					}
+					else
+					{
+						Files.Add(new FileReference(Resource.ResourcePath));
+					}
 				}
 
 				// Add any zip files from Additional Frameworks
