UBT: Fix GenerateProjectFiles failure due to ambiguous paths.

diff --git a/Engine/Source/Programs/UnrealBuildTool/System/Project.cs b/Engine/Source/Programs/UnrealBuildTool/System/Project.cs
--- a/Engine/Source/Programs/UnrealBuildTool/System/Project.cs
+++ b/Engine/Source/Programs/UnrealBuildTool/System/Project.cs
@@ -412,6 +412,10 @@ namespace UnrealBuildTool
 			{
 				return InputPath;
 			}
+			else if(InputPath.EndsWith("\\") || InputPath.EndsWith("/"))
+			{
+				return NormalizeProjectPath(new DirectoryReference(InputPath));
+			}
 			else
 			{
 				return NormalizeProjectPath(new FileReference(InputPath));
@@ -421,7 +425,7 @@ namespace UnrealBuildTool
 		/// <summary>
 		/// Takes the given path and tries to rebase it relative to the project.
 		/// </summary>
-		public string NormalizeProjectPath(FileReference InputPath)
+		public string NormalizeProjectPath(FileSystemReference InputPath)
 		{
 			// Try to make it relative to the solution directory.
 			if (InputPath.IsUnderDirectory(ProjectFileGenerator.MasterProjectPath))
