UBT: Prevent exception when cleaning out stale modules if the output directory does not exist.

diff --git a/Engine/Source/Programs/UnrealBuildTool/Configuration/UEBuildTarget.cs b/Engine/Source/Programs/UnrealBuildTool/Configuration/UEBuildTarget.cs
--- a/Engine/Source/Programs/UnrealBuildTool/Configuration/UEBuildTarget.cs
+++ b/Engine/Source/Programs/UnrealBuildTool/Configuration/UEBuildTarget.cs
@@ -1714,13 +1714,16 @@ namespace UnrealBuildTool
 			// Search all the output directories for files with a name matching one of our output files
 			foreach(DirectoryReference OutputDirectory in OutputFiles.Select(x => x.Directory).Distinct())
 			{
-				foreach(FileReference ExistingFile in OutputDirectory.EnumerateFileReferences())
+				if(OutputDirectory.Exists())
 				{
-					FileReference OutputFile;
-					if(OutputFileNames.TryGetValue(ExistingFile.GetFileName(), out OutputFile) && !OutputFiles.Contains(ExistingFile))
+					foreach(FileReference ExistingFile in OutputDirectory.EnumerateFileReferences())
 					{
-						Log.TraceInformation("Deleting '{0}' to avoid ambiguity with '{1}'", ExistingFile, OutputFile);
-						CleanFile(ExistingFile.FullName);
+						FileReference OutputFile;
+						if(OutputFileNames.TryGetValue(ExistingFile.GetFileName(), out OutputFile) && !OutputFiles.Contains(ExistingFile))
+						{
+							Log.TraceInformation("Deleting '{0}' to avoid ambiguity with '{1}'", ExistingFile, OutputFile);
+							CleanFile(ExistingFile.FullName);
+						}
 					}
 				}
 			}
