UBT: Fix support for the x64-on-x86 compiler shipped with Visual Studio Express, which is causing errors for artists generating project files with UGS. Was not looking for the compiler executable at the correct path.

diff --git a/Engine/Source/Programs/UnrealBuildTool/Windows/VCEnvironment.cs b/Engine/Source/Programs/UnrealBuildTool/Windows/VCEnvironment.cs
--- a/Engine/Source/Programs/UnrealBuildTool/Windows/VCEnvironment.cs
+++ b/Engine/Source/Programs/UnrealBuildTool/Windows/VCEnvironment.cs
@@ -330,7 +330,7 @@ namespace UnrealBuildTool
 			}
 
 			// Otherwise use the amd64-on-x86 compiler. VS2012 Express only includes the latter.
-			FileReference CrossCompilePath = FileReference.Combine(VCInstallDir, "bin", "x86_amd64");
+			FileReference CrossCompilePath = FileReference.Combine(VCInstallDir, "bin", "x86_amd64", "cl.exe");
 			if(CrossCompilePath.Exists())
 			{
 				return CrossCompilePath.Directory;
