fix for code signing on Xcode 8

diff --git a/Engine/Source/Programs/IOS/iPhonePackager/CompileTime.cs b/Engine/Source/Programs/IOS/iPhonePackager/CompileTime.cs
--- a/Engine/Source/Programs/IOS/iPhonePackager/CompileTime.cs
+++ b/Engine/Source/Programs/IOS/iPhonePackager/CompileTime.cs
@@ -266,9 +266,25 @@ namespace iPhonePackager
 			}
 			string FinalMobileProvisionFilename = Path.Combine(Config.PCXcodeStagingDir, MacMobileProvisionFilename);
 			FileOperations.CopyRequiredFile(ProvisionWithPrefix, FinalMobileProvisionFilename);
-			// make sure this .mobileprovision file is newer than any other .mobileprovision file on the Mac (this file gets multiple games named the same file, 
-			// so the time stamp checking can fail when moving between games, a la the buildmachines!)
-			File.SetLastWriteTime(FinalMobileProvisionFilename, DateTime.UtcNow);
+
+            // get the UUID
+            string AllText = File.ReadAllText(FinalMobileProvisionFilename);
+            string UUID = "";
+            int idx = AllText.IndexOf("<key>UUID</key>");
+            if (idx > 0)
+            {
+                idx = AllText.IndexOf("<string>", idx);
+                if (idx > 0)
+                {
+                    idx += "<string>".Length;
+                    UUID = AllText.Substring(idx, AllText.IndexOf("</string>", idx) - idx);
+                }
+            }
+            CurrentBaseXCodeCommandLine += String.Format(" PROVISIONING_PROFILE=" + UUID);
+
+            // make sure this .mobileprovision file is newer than any other .mobileprovision file on the Mac (this file gets multiple games named the same file, 
+            // so the time stamp checking can fail when moving between games, a la the buildmachines!)
+            File.SetLastWriteTime(FinalMobileProvisionFilename, DateTime.UtcNow);
 			string ProjectFile = Config.RootRelativePath + @"Engine\Intermediate\ProjectFiles\UE4.xcodeproj\project.pbxproj";
 			if (Program.GameName != "UE4Game")
 			{
diff --git a/Engine/Source/Programs/IOS/iPhonePackager/MobileProvisionUtilities.cs b/Engine/Source/Programs/IOS/iPhonePackager/MobileProvisionUtilities.cs
--- a/Engine/Source/Programs/IOS/iPhonePackager/MobileProvisionUtilities.cs
+++ b/Engine/Source/Programs/IOS/iPhonePackager/MobileProvisionUtilities.cs
@@ -30,8 +30,10 @@ namespace iPhonePackager
 		public Utilities.PListHelper Data;
 		public DateTime CreationDate;
 		public DateTime ExpirationDate;
+        public string FileName;
+        public string UUID;
 
-		public static string FindCompatibleProvision(string CFBundleIdentifier, out bool bNameMatch, bool bCheckCert = true, bool bCheckIdentifier = true)
+        public static string FindCompatibleProvision(string CFBundleIdentifier, out bool bNameMatch, bool bCheckCert = true, bool bCheckIdentifier = true)
 		{
 			bNameMatch = false;
 
@@ -126,7 +128,10 @@ namespace iPhonePackager
 					string DebugName = Path.GetFileName(Pair.Key);
 					MobileProvision TestProvision = Pair.Value;
 
-					Program.LogVerbose("  Phase {0} considering provision '{1}' named '{2}'", Phase, DebugName, TestProvision.ProvisionName);
+                    if (TestProvision.FileName.Contains(TestProvision.UUID))
+                        continue;
+
+                    Program.LogVerbose("  Phase {0} considering provision '{1}' named '{2}'", Phase, DebugName, TestProvision.ProvisionName);
 
 					// Validate the name
 					bool bPassesNameCheck = false;
@@ -341,12 +346,17 @@ namespace iPhonePackager
 
 			// check for get-task-allow
 			bDebug = XCentPList.GetBool("get-task-allow");
-		}
 
-		/// <summary>
-		/// Does this provision contain the specified UDID?
-		/// </summary>
-		public bool ContainsUDID(string UDID)
+            if (!Data.GetString("UUID", out UUID))
+            {
+                UUID = "(unkown)";
+            }
+        }
+
+        /// <summary>
+        /// Does this provision contain the specified UDID?
+        /// </summary>
+        public bool ContainsUDID(string UDID)
 		{
 			bool bFound = false;
 			foreach (string TestUDID in ProvisionedDeviceIDs)
@@ -423,8 +433,9 @@ namespace iPhonePackager
 			FileStream InputStream = File.OpenRead(Filename);
 			MobileProvision Result = ParseFile(InputStream);
 			InputStream.Close();
+            Result.FileName = Filename;
 
-			return Result;
+            return Result;
 		}
 
 		/// <summary>
diff --git a/Engine/Source/Programs/IOS/iPhonePackager/iPhonePackager.cs b/Engine/Source/Programs/IOS/iPhonePackager/iPhonePackager.cs
--- a/Engine/Source/Programs/IOS/iPhonePackager/iPhonePackager.cs
+++ b/Engine/Source/Programs/IOS/iPhonePackager/iPhonePackager.cs
@@ -189,13 +189,21 @@ namespace iPhonePackager
 			{
 				MainCommand = Arguments[0];
 				GamePath = Arguments[1];
-			}
-			else if (Arguments.Length >= 2)
+                if (GamePath.EndsWith(".uproject"))
+                {
+                    Config.ProjectFile = GamePath;
+                }
+            }
+            else if (Arguments.Length >= 2)
 			{
 				MainCommand = Arguments[0];
 				GamePath = Arguments[1];
+                if (GamePath.EndsWith(".uproject"))
+                {
+                    Config.ProjectFile = GamePath;
+                }
 
-				for (int ArgIndex = 2; ArgIndex < Arguments.Length; ArgIndex++)
+                for (int ArgIndex = 2; ArgIndex < Arguments.Length; ArgIndex++)
 				{
 					string Arg = Arguments[ArgIndex].ToLowerInvariant();
 
@@ -862,7 +870,25 @@ namespace iPhonePackager
 						ListDevices();
 						break;
 
-					default:
+                    case "signing_match":
+                        {
+                            MobileProvision Provision;
+                            X509Certificate2 Cert;
+                            bool bNameMatch;
+                            bool bHasOverrideFile;
+                            if (CodeSignatureBuilder.FindRequiredFiles(out Provision, out Cert, out bHasOverrideFile, out bNameMatch))
+                            {
+                                // print out the provision and cert name
+                                Program.LogVerbose("CERTIFICATE-{0},PROVISION-{1}", Cert.FriendlyName, Path.GetFileName(Provision.FileName));
+                            }
+                            else
+                            {
+                                Program.LogVerbose("No matching Signing Data found!");
+                            }
+                        }
+                        break;
+
+                    default:
 						// Commands by themself default to packaging for the device
 						if (CheckArguments())
 						{
diff --git a/Engine/Source/Programs/UnrealBuildTool/IOS/IOSToolChain.cs b/Engine/Source/Programs/UnrealBuildTool/IOS/IOSToolChain.cs
--- a/Engine/Source/Programs/UnrealBuildTool/IOS/IOSToolChain.cs
+++ b/Engine/Source/Programs/UnrealBuildTool/IOS/IOSToolChain.cs
@@ -1216,9 +1216,10 @@ namespace UnrealBuildTool
 										" -scheme '" + SchemeName + "'" +
 										" -sdk " + PlatformContext.GetCodesignPlatformName() +
 										" -destination generic/platform=" + (CppPlatform == CPPTargetPlatform.IOS ? "iOS" : "tvOS") +
-										" CODE_SIGN_IDENTITY=\"iPhone Developer\"";
+                                        " CODE_SIGN_IDENTITY=\"" + (!string.IsNullOrEmpty(PlatformContext.SigningCertificate) ? PlatformContext.SigningCertificate : "IPhoneDeveloper") + "\"" +
+                                        (!string.IsNullOrEmpty(PlatformContext.MobileProvisionUUID) ? (" PROVISIONING_PROFILE=" + PlatformContext.MobileProvisionUUID) : "");
 
-						Console.WriteLine("Code signing with command line: " + CmdLine);
+                        Console.WriteLine("Code signing with command line: " + CmdLine);
 
 						Process SignProcess = new Process();
 						SignProcess.StartInfo.WorkingDirectory = RemoteShadowDirectoryMac;
@@ -1271,9 +1272,9 @@ namespace UnrealBuildTool
 						Log.TraceInformation("Copying bundled asset... LocalSource: {0}, LocalDest: {1}", LocalSource, LocalDest);
 
 						string ResultsText;
-						RunExecutableAndWait("cp", String.Format("-R -L {0} {1}", LocalSource.Replace(" ", "\\ "), LocalDest.Replace(" ", "\\ ")), out ResultsText);
-					}
-				}
+                        RunExecutableAndWait("cp", String.Format("-R -L \"{0}\" \"{1}\"", LocalSource, LocalDest), out ResultsText);
+                    }
+                }
 			}
 			else
 			{
@@ -1492,9 +1493,9 @@ namespace UnrealBuildTool
 
 						Log.TraceInformation("Copying bundled asset... RemoteSource: {0}, RemoteDest: {1}, LocalDest: {2}", RemoteSource, RemoteDest, LocalDest);
 
-						Hashtable Results = RPCUtilHelper.Command("/", String.Format("cp -R -L {0} {1}", RemoteSource.Replace(" ", "\\ "), RemoteDest.Replace(" ", "\\ ")), "", null);
+                        Hashtable Results = RPCUtilHelper.Command("/", String.Format("cp -R -L \"{0}\" \"{1}\"", RemoteSource, RemoteDest), "", null);
 
-						foreach (DictionaryEntry Entry in Results)
+                        foreach (DictionaryEntry Entry in Results)
 						{
 							Log.TraceInformation("{0}", Entry.Value);
 						}
diff --git a/Engine/Source/Programs/UnrealBuildTool/IOS/UEBuildIOS.cs b/Engine/Source/Programs/UnrealBuildTool/IOS/UEBuildIOS.cs
--- a/Engine/Source/Programs/UnrealBuildTool/IOS/UEBuildIOS.cs
+++ b/Engine/Source/Programs/UnrealBuildTool/IOS/UEBuildIOS.cs
@@ -47,10 +47,20 @@ namespace UnrealBuildTool
 		/// </summary>
 		public string MobileProvision = "";
 
-		/// <summary>
-		/// signing certificate to use for code signing
-		/// </summary>
-		public string SigningCertificate = "";
+        /// <summary>
+        /// mobile provision UUID to use for code signing
+        /// </summary>
+        public string MobileProvisionUUID = "";
+
+        /// <summary>
+        /// mobile provision UUID to use for code signing
+        /// </summary>
+        public string TeamUUID = "";
+
+        /// <summary>
+        /// signing certificate to use for code signing
+        /// </summary>
+        public string SigningCertificate = "";
 
 		/// <summary>
 		/// The list of architectures
@@ -167,29 +177,67 @@ namespace UnrealBuildTool
 
 		}
 
-		public void SetUpProjectEnvironment(CPPTargetConfiguration Configuration)
-		{
-			UnrealTargetConfiguration	unrealConfiguration;
-			
-			switch(Configuration)
-			{
-				case CPPTargetConfiguration.Shipping:
-					unrealConfiguration = UnrealTargetConfiguration.Shipping;
-					break;
-				case CPPTargetConfiguration.Development:
-					unrealConfiguration = UnrealTargetConfiguration.Development;
-					break;
-				default:
-					unrealConfiguration = UnrealTargetConfiguration.DebugGame;
-					break;
-			}
-			
-			SetUpProjectEnvironment(unrealConfiguration);
-		}
-		
-		public override void SetUpProjectEnvironment(UnrealTargetConfiguration Configuration)
-		{
-			if (!bInitializedProject)
+        public void SetUpProjectEnvironment(CPPTargetConfiguration Configuration)
+        {
+            UnrealTargetConfiguration unrealConfiguration;
+
+            switch (Configuration)
+            {
+                case CPPTargetConfiguration.Shipping:
+                    unrealConfiguration = UnrealTargetConfiguration.Shipping;
+                    break;
+                case CPPTargetConfiguration.Development:
+                    unrealConfiguration = UnrealTargetConfiguration.Development;
+                    break;
+                default:
+                    unrealConfiguration = UnrealTargetConfiguration.DebugGame;
+                    break;
+            }
+
+            SetUpProjectEnvironment(unrealConfiguration);
+        }
+
+        static string Provision = "";
+        static string Certificate = "";
+        static bool bHaveCertificate = false;
+        public static void IPPDataReceivedHandler(Object Sender, DataReceivedEventArgs Line)
+        {
+            if ((Line != null) && (Line.Data != null))
+            {
+                if (!string.IsNullOrEmpty(Certificate))
+                {
+                    if (Line.Data.Contains("CERTIFICATE-") && Line.Data.Contains(Certificate))
+                    {
+                        bHaveCertificate = true;
+                    }
+                }
+                else
+                {
+                    int cindex = Line.Data.IndexOf("CERTIFICATE-");
+                    int pindex = Line.Data.IndexOf("PROVISION-");
+                    if (cindex > -1 && pindex > -1)
+                    {
+                        cindex += "CERTIFICATE-".Length;
+                        Certificate = Line.Data.Substring(cindex, pindex - cindex - 1);
+                        pindex += "PROVISION-".Length;
+                        Provision = Line.Data.Substring(pindex);
+                    }
+                }
+            }
+        }
+
+        struct ProvisionData
+        {
+            public string MobileProvision;
+            public string Certificate;
+            public string UUID;
+            public string TeamUUID;
+        }
+        static Dictionary<string, ProvisionData> ProvisionCache = new Dictionary<string, ProvisionData>();
+
+        public override void SetUpProjectEnvironment(UnrealTargetConfiguration Configuration)
+        {
+            if (!bInitializedProject)
 			{
 				base.SetUpProjectEnvironment(Configuration);
 
@@ -300,7 +348,116 @@ namespace UnrealBuildTool
 				Ini.GetString("/Script/IOSRuntimeSettings.IOSRuntimeSettings", "MobileProvision", out MobileProvision);
 				Ini.GetString("/Script/IOSRuntimeSettings.IOSRuntimeSettings", "SigningCertificate", out SigningCertificate);
 
-				bInitializedProject = true;
+                // bundle identifier
+                string BundleIdentifier;
+                Ini.GetString("/Script/IOSRuntimeSettings.IOSRuntimeSettings", "BundleIdentifier", out BundleIdentifier);
+
+                ProvisionData Data = new ProvisionData();
+                string BundleId = BundleIdentifier.Replace("[PROJECT_NAME]", ((ProjectFile != null) ? ProjectFile.GetFileNameWithoutAnyExtensions() : "UE4Game")).Replace("_", "");
+                if (!ProvisionCache.ContainsKey(BundleId))
+                {
+                    Certificate = SigningCertificate;
+                    Provision = MobileProvision;
+                    if (!string.IsNullOrEmpty(SigningCertificate))
+                    {
+                        // verify the certificate
+                        Process IPPProcess = new Process();
+                        if (BuildHostPlatform.Current.Platform == UnrealTargetPlatform.Mac)
+                        {
+                            string IPPCmd = "\"" + UnrealBuildTool.EngineDirectory + "/Binaries/DotNET/IOS/IPhonePackager.exe\" certificates " + ((ProjectFile != null) ? ("\"" + ProjectFile.ToString() + "\"") : "Engine") + " -bundlename " + BundleIdentifier.Replace("[PROJECT_NAME]", ((ProjectFile != null) ? ProjectFile.GetFileNameWithoutAnyExtensions() : "UE4Game")).Replace("_", "");
+                            IPPProcess.StartInfo.WorkingDirectory = UnrealBuildTool.EngineDirectory.ToString();
+                            IPPProcess.StartInfo.FileName = UnrealBuildTool.EngineDirectory + "/Build/BatchFiles/Mac/RunMono.sh";
+                            IPPProcess.StartInfo.Arguments = IPPCmd;
+                            IPPProcess.OutputDataReceived += new DataReceivedEventHandler(IPPDataReceivedHandler);
+                            IPPProcess.ErrorDataReceived += new DataReceivedEventHandler(IPPDataReceivedHandler);
+                        }
+                        else
+                        {
+                            string IPPCmd = "certificates " + ((ProjectFile != null) ? ("\"" + ProjectFile.ToString() + "\"") : "Engine") + " -bundlename " + BundleId;
+                            IPPProcess.StartInfo.WorkingDirectory = UnrealBuildTool.EngineDirectory.ToString();
+                            IPPProcess.StartInfo.FileName = UnrealBuildTool.EngineDirectory + "\\Binaries\\DotNET\\IOS\\IPhonePackager.exe";
+                            IPPProcess.StartInfo.Arguments = IPPCmd;
+                            IPPProcess.OutputDataReceived += new DataReceivedEventHandler(IPPDataReceivedHandler);
+                            IPPProcess.ErrorDataReceived += new DataReceivedEventHandler(IPPDataReceivedHandler);
+                        }
+                        Utils.RunLocalProcess(IPPProcess);
+                    }
+
+                    if (string.IsNullOrEmpty(MobileProvision) // no provision specified
+                        || !File.Exists((BuildHostPlatform.Current.Platform == UnrealTargetPlatform.Mac ? (Environment.GetEnvironmentVariable("HOME") + "/Library/MobileDevice/Provisioning Profiles/") : (Environment.GetFolderPath(Environment.SpecialFolder.LocalApplicationData) + "/Apple Computer/MobileDevice/Provisioning Profiles/")) + MobileProvision) // file doesn't exist
+                        || !bHaveCertificate) // certificate doesn't exist
+                    {
+                        Certificate = "";
+                        Provision = "";
+                        Log.TraceLog("Provision not specified or not found for " + ((ProjectFile != null) ? ProjectFile.GetFileNameWithoutAnyExtensions() : "UE4Game") + ", searching for compatible match...");
+                        Process IPPProcess = new Process();
+                        if (BuildHostPlatform.Current.Platform == UnrealTargetPlatform.Mac)
+                        {
+                            string IPPCmd = "\"" + UnrealBuildTool.EngineDirectory + "/Binaries/DotNET/IOS/IPhonePackager.exe\" signing_match " + ((ProjectFile != null) ? ("\"" + ProjectFile.ToString() + "\"") : "Engine") + " -bundlename " + BundleIdentifier.Replace("[PROJECT_NAME]", ((ProjectFile != null) ? ProjectFile.GetFileNameWithoutAnyExtensions() : "UE4Game")).Replace("_", "");
+                            IPPProcess.StartInfo.WorkingDirectory = UnrealBuildTool.EngineDirectory.ToString();
+                            IPPProcess.StartInfo.FileName = UnrealBuildTool.EngineDirectory + "/Build/BatchFiles/Mac/RunMono.sh";
+                            IPPProcess.StartInfo.Arguments = IPPCmd;
+                            IPPProcess.OutputDataReceived += new DataReceivedEventHandler(IPPDataReceivedHandler);
+                            IPPProcess.ErrorDataReceived += new DataReceivedEventHandler(IPPDataReceivedHandler);
+                        }
+                        else
+                        {
+                            string IPPCmd = "signing_match " + ((ProjectFile != null) ? ("\"" + ProjectFile.ToString() + "\"") : "Engine") + " -bundlename " + BundleId;
+                            IPPProcess.StartInfo.WorkingDirectory = UnrealBuildTool.EngineDirectory.ToString();
+                            IPPProcess.StartInfo.FileName = UnrealBuildTool.EngineDirectory + "\\Binaries\\DotNET\\IOS\\IPhonePackager.exe";
+                            IPPProcess.StartInfo.Arguments = IPPCmd;
+                            IPPProcess.OutputDataReceived += new DataReceivedEventHandler(IPPDataReceivedHandler);
+                            IPPProcess.ErrorDataReceived += new DataReceivedEventHandler(IPPDataReceivedHandler);
+                        }
+                        Utils.RunLocalProcess(IPPProcess);
+                        Log.TraceLog("Provision found for " + ((ProjectFile != null) ? ProjectFile.GetFileNameWithoutAnyExtensions() : "UE4Game") + ", Provision: " + Provision + " Certificate: " + Certificate);
+                    }
+                    // add to the dictionary
+                    Data.MobileProvision = Provision;
+                    Data.Certificate = Certificate;
+
+                    // read the provision to get the UUID
+                    string filename = (BuildHostPlatform.Current.Platform == UnrealTargetPlatform.Mac ? (Environment.GetEnvironmentVariable("HOME") + "/Library/MobileDevice/Provisioning Profiles/") : (Environment.GetFolderPath(Environment.SpecialFolder.LocalApplicationData) + "/Apple Computer/MobileDevice/Provisioning Profiles/")) + Data.MobileProvision;
+                    if (File.Exists(filename))
+                    {
+                        string AllText = File.ReadAllText(filename);
+                        int idx = AllText.IndexOf("<key>UUID</key>");
+                        if (idx > 0)
+                        {
+                            idx = AllText.IndexOf("<string>", idx);
+                            if (idx > 0)
+                            {
+                                idx += "<string>".Length;
+                                Data.UUID = AllText.Substring(idx, AllText.IndexOf("</string>", idx) - idx);
+                            }
+                        }
+                        idx = AllText.IndexOf("<key>com.apple.developer.team-identifier</key>");
+                        if (idx > 0)
+                        {
+                            idx = AllText.IndexOf("<string>", idx);
+                            if (idx > 0)
+                            {
+                                idx += "<string>".Length;
+                                Data.TeamUUID = AllText.Substring(idx, AllText.IndexOf("</string>", idx) - idx);
+                            }
+                        }
+                    }
+                    else
+                    {
+                        Log.TraceWarning("No matching provision file was discovered. Please ensure you have a compatible provision installed.");
+                    }
+                    ProvisionCache.Add(BundleId, Data);
+                }
+                else
+                {
+                    Data = ProvisionCache[BundleId];
+                }
+                MobileProvision = Data.MobileProvision;
+                SigningCertificate = Data.Certificate;
+                MobileProvisionUUID = Data.UUID;
+                TeamUUID = Data.TeamUUID;
+
+                bInitializedProject = true;
 			}
 		}
 
diff --git a/Engine/Source/Programs/UnrealBuildTool/System/XcodeProject.cs b/Engine/Source/Programs/UnrealBuildTool/System/XcodeProject.cs
--- a/Engine/Source/Programs/UnrealBuildTool/System/XcodeProject.cs
+++ b/Engine/Source/Programs/UnrealBuildTool/System/XcodeProject.cs
@@ -503,7 +503,12 @@ namespace UnrealBuildTool
 			Content.Append("\t\t\tattributes = {" + ProjectFileGenerator.NewLine);
 			Content.Append("\t\t\t\tLastUpgradeCheck = 0700;" + ProjectFileGenerator.NewLine);
 			Content.Append("\t\t\t\tORGANIZATIONNAME = \"Epic Games, Inc.\";" + ProjectFileGenerator.NewLine);
-			Content.Append("\t\t\t};" + ProjectFileGenerator.NewLine);
+            Content.Append("\t\t\t\tTargetAttributes = {" + ProjectFileGenerator.NewLine);
+            Content.Append("\t\t\t\t\t" + TargetGuid + " = {" + ProjectFileGenerator.NewLine);
+            Content.Append("\t\t\t\t\t\tProvisioningStyle = Manual;" + ProjectFileGenerator.NewLine);
+            Content.Append("\t\t\t\t\t};" + ProjectFileGenerator.NewLine);
+            Content.Append("\t\t\t\t};" + ProjectFileGenerator.NewLine);
+            Content.Append("\t\t\t};" + ProjectFileGenerator.NewLine);
 			Content.Append("\t\t\tbuildConfigurationList = " + ProjectBuildConfigGuid + " /* Build configuration list for PBXProject \"" + TargetName + "\" */;" + ProjectFileGenerator.NewLine);
 			Content.Append("\t\t\tcompatibilityVersion = \"Xcode 3.2\";" + ProjectFileGenerator.NewLine);
 			Content.Append("\t\t\tdevelopmentRegion = English;" + ProjectFileGenerator.NewLine);
@@ -645,18 +650,27 @@ namespace UnrealBuildTool
 				string ValidArchs = "x86_64";
 				string SupportedPlatforms = "macosx";
 
-				if (UnrealBuildTool.IsValidPlatform(UnrealTargetPlatform.IOS))
-				{
-					IOSPlatform IOSBuildPlat = UEBuildPlatform.GetBuildPlatform(UnrealTargetPlatform.IOS) as IOSPlatform;
+                string UUID_IOS = "";
+                string UUID_TVOS = "";
+                string TEAM_IOS = "";
+                string TEAM_TVOS = "";
+                string IOS_CERT = "iPhone Developer";
+                string TVOS_CERT = "iPhone Developer";
+                if (UnrealBuildTool.IsValidPlatform(UnrealTargetPlatform.IOS))
+                {
+                    IOSPlatform IOSBuildPlat = UEBuildPlatform.GetBuildPlatform(UnrealTargetPlatform.IOS) as IOSPlatform;
 					IOSPlatformContext PlatformContextIOS = (IOSPlatformContext)IOSBuildPlat.CreateContext(ProjectFile);
 					PlatformContextIOS.SetUpProjectEnvironment(Config.BuildConfig);
 					IOSRunTimeVersion = PlatformContextIOS.GetRunTimeVersion();
 					IOSRunTimeDevices = PlatformContextIOS.GetRunTimeDevices();
 					ValidArchs += " arm64 armv7 armv7s";
 					SupportedPlatforms += " iphoneos";
-				}
+                    UUID_IOS = PlatformContextIOS.MobileProvisionUUID;
+                    TEAM_IOS = PlatformContextIOS.TeamUUID;
+                    IOS_CERT = PlatformContextIOS.SigningCertificate;
+                }
 
-				if (UnrealBuildTool.IsValidPlatform(UnrealTargetPlatform.TVOS))
+                if (UnrealBuildTool.IsValidPlatform(UnrealTargetPlatform.TVOS))
 				{
 					TVOSPlatform TVOSBuildPlat = UEBuildPlatform.GetBuildPlatform(UnrealTargetPlatform.TVOS) as TVOSPlatform;
 					TVOSPlatformContext PlatformContextTVOS = (TVOSPlatformContext)TVOSBuildPlat.CreateContext(ProjectFile);
@@ -668,9 +682,12 @@ namespace UnrealBuildTool
 						ValidArchs += " arm64 armv7 armv7s";
 					}
 					SupportedPlatforms += " appletvos";
-				}
+                    UUID_TVOS = PlatformContextTVOS.MobileProvisionUUID;
+                    TEAM_TVOS = PlatformContextTVOS.TeamUUID;
+                    TVOS_CERT = PlatformContextTVOS.SigningCertificate;
+                }
 
-				Content.Append("\t\t\t\tVALID_ARCHS = \"" + ValidArchs + "\";" + ProjectFileGenerator.NewLine);
+                Content.Append("\t\t\t\tVALID_ARCHS = \"" + ValidArchs + "\";" + ProjectFileGenerator.NewLine);
 				Content.Append("\t\t\t\tSUPPORTED_PLATFORMS = \"" + SupportedPlatforms + "\";" + ProjectFileGenerator.NewLine);
 				Content.Append("\t\t\t\t\"PRODUCT_NAME[sdk=macosx*]\" = \"" + MacExecutableFileName + "\";" + ProjectFileGenerator.NewLine);
 				if (IOSRunTimeVersion != null)
@@ -678,20 +695,22 @@ namespace UnrealBuildTool
 					Content.Append("\t\t\t\tIPHONEOS_DEPLOYMENT_TARGET = " + IOSRunTimeVersion + ";" + ProjectFileGenerator.NewLine);
 					Content.Append("\t\t\t\t\"PRODUCT_NAME[sdk=iphoneos*]\" = \"" + Config.BuildTarget + "\";" + ProjectFileGenerator.NewLine); // @todo: change to Path.GetFileName(Config.IOSExecutablePath) when we stop using payload
 					Content.Append("\t\t\t\t\"TARGETED_DEVICE_FAMILY[sdk=iphoneos*]\" = \"" + IOSRunTimeDevices + "\";" + ProjectFileGenerator.NewLine);
-					Content.Append("\t\t\t\t\"CODE_SIGN_IDENTITY[sdk=iphoneos*]\" = \"iPhone Developer\";" + ProjectFileGenerator.NewLine);
-					Content.Append("\t\t\t\t\"SDKROOT[sdk=iphoneos]\" = iphoneos;" + ProjectFileGenerator.NewLine);
-					Content.Append("\t\t\t\t\"PROVISIONING_PROFILE[sdk=iphoneos*]\" = \"\";" + ProjectFileGenerator.NewLine);
-				}
-				if (TVOSRunTimeVersion != null)
+                    Content.Append("\t\t\t\t\"CODE_SIGN_IDENTITY[sdk=iphoneos*]\" = \"" + IOS_CERT + "\";" + ProjectFileGenerator.NewLine);
+                    Content.Append("\t\t\t\t\"SDKROOT[sdk=iphoneos]\" = iphoneos;" + ProjectFileGenerator.NewLine);
+                    Content.Append("\t\t\t\t\"PROVISIONING_PROFILE_SPECIFIER[sdk=iphoneos*]\" = \"" + UUID_IOS + "\";" + ProjectFileGenerator.NewLine);
+                    Content.Append("\t\t\t\t\"DEVELOPMENT_TEAM[sdk=iphoneos*]\" = \"" + TEAM_IOS + "\";" + ProjectFileGenerator.NewLine);
+                }
+                if (TVOSRunTimeVersion != null)
 				{
 					Content.Append("\t\t\t\tTVOS_DEPLOYMENT_TARGET = " + TVOSRunTimeVersion + ";" + ProjectFileGenerator.NewLine);
 					Content.Append("\t\t\t\t\"PRODUCT_NAME[sdk=appletvos*]\" = \"" + Config.BuildTarget + "\";" + ProjectFileGenerator.NewLine); // @todo: change to Path.GetFileName(Config.TVOSExecutablePath) when we stop using payload
 					Content.Append("\t\t\t\t\"TARGETED_DEVICE_FAMILY[sdk=appletvos*]\" = \"" + TVOSRunTimeDevices + "\";" + ProjectFileGenerator.NewLine);
-					Content.Append("\t\t\t\t\"CODE_SIGN_IDENTITY[sdk=appletvos*]\" = \"iPhone Developer\";" + ProjectFileGenerator.NewLine);
-					Content.Append("\t\t\t\t\"SDKROOT[sdk=appletvos]\" = appletvos;" + ProjectFileGenerator.NewLine);
-					Content.Append("\t\t\t\t\"PROVISIONING_PROFILE[sdk=appletvos*]\" = \"\";" + ProjectFileGenerator.NewLine);
-				}
-				Content.Append("\t\t\t\t\"CONFIGURATION_BUILD_DIR[sdk=macosx*]\" = \"" + MacExecutableDir + "\";" + ProjectFileGenerator.NewLine);
+                    Content.Append("\t\t\t\t\"CODE_SIGN_IDENTITY[sdk=appletvos*]\" = \"" + TVOS_CERT + "\";" + ProjectFileGenerator.NewLine);
+                    Content.Append("\t\t\t\t\"SDKROOT[sdk=appletvos]\" = appletvos;" + ProjectFileGenerator.NewLine);
+                    Content.Append("\t\t\t\t\"PROVISIONING_PROFILE_SPECIFIER[sdk=appletvos*]\" = \"" + UUID_TVOS + "\";" + ProjectFileGenerator.NewLine);
+                    Content.Append("\t\t\t\t\"DEVELOPMENT_TEAM[sdk=appletvos*]\" = \"" + TEAM_TVOS + "\";" + ProjectFileGenerator.NewLine);
+                }
+                Content.Append("\t\t\t\t\"CONFIGURATION_BUILD_DIR[sdk=macosx*]\" = \"" + MacExecutableDir + "\";" + ProjectFileGenerator.NewLine);
 				Content.Append("\t\t\t\t\"SDKROOT[sdk=macosx]\" = macosx;" + ProjectFileGenerator.NewLine);
 				Content.Append("\t\t\t\tINFOPLIST_OUTPUT_FORMAT = xml;" + ProjectFileGenerator.NewLine);
 
