fixes for Xcode 8 signing issues re-made from Release 4.13

diff --git a/Engine/Source/Programs/AutomationTool/IOS/IOSPlatform.Automation.cs b/Engine/Source/Programs/AutomationTool/IOS/IOSPlatform.Automation.cs
--- a/Engine/Source/Programs/AutomationTool/IOS/IOSPlatform.Automation.cs
+++ b/Engine/Source/Programs/AutomationTool/IOS/IOSPlatform.Automation.cs
@@ -224,10 +224,15 @@ public class IOSPlatform : Platform
 		return 4;
 	}
 
-    public virtual UnrealBuildTool.UEDeployIOS GetDeployHandler(FileReference InProject, UnrealBuildTool.IOSPlatformContext inIOSPlatformContext)
+    public virtual UnrealBuildTool.UEDeployIOS GetDeployHandler(FileReference InProject, UnrealBuildTool.IOSPlatformContext inPlatformContext)
     {
         Console.WriteLine("Getting IOS Deploy()");
-        return new UnrealBuildTool.UEDeployIOS(InProject, inIOSPlatformContext);
+        return new UnrealBuildTool.UEDeployIOS(InProject, inPlatformContext);
+    }
+
+    public virtual UnrealBuildTool.IOSPlatformContext CreatePlatformContext(FileReference InProject, bool Distribution)
+    {
+        return new UnrealBuildTool.IOSPlatformContext(InProject, Distribution);
     }
 
     protected string MakeIPAFileName( UnrealTargetConfiguration TargetConfiguration, ProjectParams Params )
@@ -301,7 +306,7 @@ public class IOSPlatform : Platform
 
         var TargetConfiguration = SC.StageTargetConfigurations[0];
 
-        UnrealBuildTool.IOSPlatformContext BuildPlatContext = new IOSPlatformContext(Params.RawProjectPath);
+        UnrealBuildTool.IOSPlatformContext BuildPlatContext = CreatePlatformContext(Params.RawProjectPath, Params.Distribution);
         BuildPlatContext.SetUpProjectEnvironment(TargetConfiguration);
 
         //@TODO: We should be able to use this code on both platforms, when the following issues are sorted:
@@ -653,9 +658,9 @@ public class IOSPlatform : Platform
 					{
 						idx += "<string>".Length;
 						UUID = AllText.Substring(idx, AllText.IndexOf("</string>", idx) - idx);
-						Arguments += " PROVISIONING_PROFILE=" + UUID;
-					}
-				}
+                        Arguments += " PROVISIONING_PROFILE_SPECIFIER=" + UUID;
+                    }
+                }
 			}
 		}
 		IProcessResult Result = Run ("/usr/bin/env", Arguments, null, ERunOptions.Default);
@@ -839,7 +844,7 @@ public class IOSPlatform : Platform
 
                     var TargetConfiguration = SC.StageTargetConfigurations[0];
 
-                    UnrealBuildTool.IOSPlatformContext BuildPlatContext = new IOSPlatformContext(Params.RawProjectPath);
+                    UnrealBuildTool.IOSPlatformContext BuildPlatContext = CreatePlatformContext(Params.RawProjectPath, Params.Distribution);
                     BuildPlatContext.SetUpProjectEnvironment(TargetConfiguration);
 
                     GetDeployHandler(
diff --git a/Engine/Source/Programs/AutomationTool/TVOS/TVOSPlatform.Automation.cs b/Engine/Source/Programs/AutomationTool/TVOS/TVOSPlatform.Automation.cs
--- a/Engine/Source/Programs/AutomationTool/TVOS/TVOSPlatform.Automation.cs
+++ b/Engine/Source/Programs/AutomationTool/TVOS/TVOSPlatform.Automation.cs
@@ -22,6 +22,11 @@ public class TVOSPlatform : IOSPlatform
         return new UnrealBuildTool.UEDeployTVOS(InProject, inPlatformContext);
     }
 
+    public override UnrealBuildTool.IOSPlatformContext CreatePlatformContext(FileReference InProject, bool Distribution)
+    {
+        return new UnrealBuildTool.TVOSPlatformContext(InProject, Distribution);
+    }
+
     public override string GetCookPlatform(bool bDedicatedServer, bool bIsClientOnly)
 	{
 		return "TVOS";
diff --git a/Engine/Source/Programs/IOS/iPhonePackager/MobileProvisionUtilities.cs b/Engine/Source/Programs/IOS/iPhonePackager/MobileProvisionUtilities.cs
--- a/Engine/Source/Programs/IOS/iPhonePackager/MobileProvisionUtilities.cs
+++ b/Engine/Source/Programs/IOS/iPhonePackager/MobileProvisionUtilities.cs
@@ -32,6 +32,7 @@ namespace iPhonePackager
 		public DateTime ExpirationDate;
         public string FileName;
         public string UUID;
+        public string Platform;
 
         public static string FindCompatibleProvision(string CFBundleIdentifier, out bool bNameMatch, bool bCheckCert = true, bool bCheckIdentifier = true)
 		{
@@ -65,8 +66,8 @@ namespace iPhonePackager
 			// copy all of the provisions from the game directory to the library
 			if (!String.IsNullOrEmpty(Config.ProjectFile))
 			{
-				var ProjectFileBuildIOSPath = Path.GetDirectoryName(Config.ProjectFile) + "/Build/IOS/";
-				if (Directory.Exists(ProjectFileBuildIOSPath))
+                var ProjectFileBuildIOSPath = Path.GetDirectoryName(Config.ProjectFile) + "/Build/" + Config.OSString + "/";
+                if (Directory.Exists(ProjectFileBuildIOSPath))
 				{
 					foreach (string Provision in Directory.EnumerateFiles(ProjectFileBuildIOSPath, "*.mobileprovision", SearchOption.AllDirectories))
 					{
@@ -128,9 +129,14 @@ namespace iPhonePackager
 					string DebugName = Path.GetFileName(Pair.Key);
 					MobileProvision TestProvision = Pair.Value;
 
+                    // make sure the file is not managed by Xcode
                     if (TestProvision.FileName.Contains(TestProvision.UUID))
                         continue;
 
+                    // check to see if the platform is the same as what we are looking for
+                    if (!string.IsNullOrEmpty(TestProvision.Platform) && TestProvision.Platform != Config.OSString && !string.IsNullOrEmpty(Config.OSString))
+                        continue;
+
                     Program.LogVerbose("  Phase {0} considering provision '{1}' named '{2}'", Phase, DebugName, TestProvision.ProvisionName);
 
 					// Validate the name
@@ -351,6 +357,20 @@ namespace iPhonePackager
             {
                 UUID = "(unkown)";
             }
+
+            List<string> Platforms = Data.GetArray("Platform", "string");
+            if (Platforms.Contains("iOS"))
+            {
+                Platform = "IOS";
+            }
+            else if (Platforms.Contains("tvOS"))
+            {
+                Platform = "TVOS";
+            }
+            else
+            {
+                Platform = "";
+            }
         }
 
         /// <summary>
diff --git a/Engine/Source/Programs/UnrealBuildTool/IOS/IOSToolChain.cs b/Engine/Source/Programs/UnrealBuildTool/IOS/IOSToolChain.cs
--- a/Engine/Source/Programs/UnrealBuildTool/IOS/IOSToolChain.cs
+++ b/Engine/Source/Programs/UnrealBuildTool/IOS/IOSToolChain.cs
@@ -1217,7 +1217,7 @@ namespace UnrealBuildTool
 										" -sdk " + PlatformContext.GetCodesignPlatformName() +
 										" -destination generic/platform=" + (CppPlatform == CPPTargetPlatform.IOS ? "iOS" : "tvOS") +
                                         " CODE_SIGN_IDENTITY=\"" + (!string.IsNullOrEmpty(PlatformContext.SigningCertificate) ? PlatformContext.SigningCertificate : "IPhoneDeveloper") + "\"" +
-                                        (!string.IsNullOrEmpty(PlatformContext.MobileProvisionUUID) ? (" PROVISIONING_PROFILE=" + PlatformContext.MobileProvisionUUID) : "");
+                                        (!string.IsNullOrEmpty(PlatformContext.MobileProvisionUUID) ? (" PROVISIONING_PROFILE_SPECIFIER=" + PlatformContext.MobileProvisionUUID) : "");
 
                         Console.WriteLine("Code signing with command line: " + CmdLine);
 
diff --git a/Engine/Source/Programs/UnrealBuildTool/IOS/UEBuildIOS.cs b/Engine/Source/Programs/UnrealBuildTool/IOS/UEBuildIOS.cs
--- a/Engine/Source/Programs/UnrealBuildTool/IOS/UEBuildIOS.cs
+++ b/Engine/Source/Programs/UnrealBuildTool/IOS/UEBuildIOS.cs
@@ -72,18 +72,23 @@ namespace UnrealBuildTool
 		/// </summary>
 		private bool bShipForBitcode = false;
 
-		public IOSPlatformContext(FileReference InProjectFile) 
-			: this(UnrealTargetPlatform.IOS, InProjectFile)
-		{
-		}
+        private bool bForDistribtion = false;
+
+        private string BundleIdentifier = "";
+
+        public IOSPlatformContext(FileReference InProjectFile, bool ForDistribution = false)
+            : this(UnrealTargetPlatform.IOS, InProjectFile, ForDistribution)
+        {
+        }
 
-		protected IOSPlatformContext(UnrealTargetPlatform TargetPlatform, FileReference InProjectFile) 
-			: base(TargetPlatform, InProjectFile)
+        protected IOSPlatformContext(UnrealTargetPlatform TargetPlatform, FileReference InProjectFile, bool ForDistribution)
+            : base(TargetPlatform, InProjectFile)
 		{
-		}
+            bForDistribtion = ForDistribution;
+        }
 
-		// The current architecture - affects everything about how UBT operates on IOS
-		public override string GetActiveArchitecture()
+        // The current architecture - affects everything about how UBT operates on IOS
+        public override string GetActiveArchitecture()
 		{
 			return IOSArchitecture;
 		}
@@ -238,235 +243,241 @@ namespace UnrealBuildTool
         public override void SetUpProjectEnvironment(UnrealTargetConfiguration Configuration)
         {
             if (!bInitializedProject)
-			{
-				base.SetUpProjectEnvironment(Configuration);
+            {
+                base.SetUpProjectEnvironment(Configuration);
 
-				// update the configuration based on the project file
-				// look in ini settings for what platforms to compile for
-				ConfigCacheIni Ini = ConfigCacheIni.CreateConfigCacheIni(Platform, "Engine", DirectoryReference.FromFile(ProjectFile));
-				string MinVersion = "IOS_7";
-				if (Ini.GetString("/Script/IOSRuntimeSettings.IOSRuntimeSettings", "MinimumiOSVersion", out MinVersion))
-				{
-					switch (MinVersion)
-					{
-						case "IOS_61":
-							Log.TraceWarning("IOS 6 is no longer supported in UE4 as 4.11");
-							RunTimeIOSVersion = "7.0";
-							break;
-						case "IOS_7":
-							RunTimeIOSVersion = "7.0";
-							break;
-						case "IOS_8":
-							RunTimeIOSVersion = "8.0";
-							break;
-						case "IOS_9":
-							RunTimeIOSVersion = "9.0";
-							break;
-					}
-				}
+                // update the configuration based on the project file
+                // look in ini settings for what platforms to compile for
+                ConfigCacheIni Ini = ConfigCacheIni.CreateConfigCacheIni(Platform, "Engine", DirectoryReference.FromFile(ProjectFile));
+                string MinVersion = "IOS_7";
+                if (Ini.GetString("/Script/IOSRuntimeSettings.IOSRuntimeSettings", "MinimumiOSVersion", out MinVersion))
+                {
+                    switch (MinVersion)
+                    {
+                        case "IOS_61":
+                            Log.TraceWarning("IOS 6 is no longer supported in UE4 as 4.11");
+                            RunTimeIOSVersion = "7.0";
+                            break;
+                        case "IOS_7":
+                            RunTimeIOSVersion = "7.0";
+                            break;
+                        case "IOS_8":
+                            RunTimeIOSVersion = "8.0";
+                            break;
+                        case "IOS_9":
+                            RunTimeIOSVersion = "9.0";
+                            break;
+                    }
+                }
 
-				bool biPhoneAllowed = true;
-				bool biPadAllowed = true;
-				Ini.GetBool("/Script/IOSRuntimeSettings.IOSRuntimeSettings", "bSupportsIPhone", out biPhoneAllowed);
-				Ini.GetBool("/Script/IOSRuntimeSettings.IOSRuntimeSettings", "bSupportsIPad", out biPadAllowed);
-				if (biPhoneAllowed && biPadAllowed)
-				{
-					RunTimeIOSDevices = "1,2";
-				}
-				else if (biPadAllowed)
-				{
-					RunTimeIOSDevices = "2";
-				}
-				else if (biPhoneAllowed)
-				{
-					RunTimeIOSDevices = "1";
-				}
+                bool biPhoneAllowed = true;
+                bool biPadAllowed = true;
+                Ini.GetBool("/Script/IOSRuntimeSettings.IOSRuntimeSettings", "bSupportsIPhone", out biPhoneAllowed);
+                Ini.GetBool("/Script/IOSRuntimeSettings.IOSRuntimeSettings", "bSupportsIPad", out biPadAllowed);
+                if (biPhoneAllowed && biPadAllowed)
+                {
+                    RunTimeIOSDevices = "1,2";
+                }
+                else if (biPadAllowed)
+                {
+                    RunTimeIOSDevices = "2";
+                }
+                else if (biPhoneAllowed)
+                {
+                    RunTimeIOSDevices = "1";
+                }
 
-				ProjectArches = new List<string>();
-				bool bBuild = true;
-				if (Ini.GetBool("/Script/IOSRuntimeSettings.IOSRuntimeSettings", "bDevForArmV7", out bBuild) && bBuild)
-				{
-					ProjectArches.Add("armv7");
-				}
-				if (Ini.GetBool("/Script/IOSRuntimeSettings.IOSRuntimeSettings", "bDevForArm64", out bBuild) && bBuild)
-				{
-					ProjectArches.Add("arm64");
-				}
-				if (Ini.GetBool("/Script/IOSRuntimeSettings.IOSRuntimeSettings", "bDevForArmV7S", out bBuild) && bBuild)
-				{
-					ProjectArches.Add("armv7s");
-				}
+                ProjectArches = new List<string>();
+                bool bBuild = true;
+                if (Ini.GetBool("/Script/IOSRuntimeSettings.IOSRuntimeSettings", "bDevForArmV7", out bBuild) && bBuild)
+                {
+                    ProjectArches.Add("armv7");
+                }
+                if (Ini.GetBool("/Script/IOSRuntimeSettings.IOSRuntimeSettings", "bDevForArm64", out bBuild) && bBuild)
+                {
+                    ProjectArches.Add("arm64");
+                }
+                if (Ini.GetBool("/Script/IOSRuntimeSettings.IOSRuntimeSettings", "bDevForArmV7S", out bBuild) && bBuild)
+                {
+                    ProjectArches.Add("armv7s");
+                }
 
-				// force armv7 if something went wrong
-				if (ProjectArches.Count == 0)
-				{
-					ProjectArches.Add("armv7");
-				}
-				NonShippingArchitectures = ProjectArches[0];
-				for (int Index = 1; Index < ProjectArches.Count; ++Index)
-				{
-					NonShippingArchitectures += "," + ProjectArches[Index];
-				}
+                // force armv7 if something went wrong
+                if (ProjectArches.Count == 0)
+                {
+                    ProjectArches.Add("armv7");
+                }
+                NonShippingArchitectures = ProjectArches[0];
+                for (int Index = 1; Index < ProjectArches.Count; ++Index)
+                {
+                    NonShippingArchitectures += "," + ProjectArches[Index];
+                }
 
-				ProjectArches.Clear();
-				if (Ini.GetBool("/Script/IOSRuntimeSettings.IOSRuntimeSettings", "bShipForArmV7", out bBuild) && bBuild)
-				{
-					ProjectArches.Add("armv7");
-				}
-				if (Ini.GetBool("/Script/IOSRuntimeSettings.IOSRuntimeSettings", "bShipForArm64", out bBuild) && bBuild)
-				{
-					ProjectArches.Add("arm64");
-				}
-				if (Ini.GetBool("/Script/IOSRuntimeSettings.IOSRuntimeSettings", "bShipForArmV7S", out bBuild) && bBuild)
-				{
-					ProjectArches.Add("armv7s");
-				}
+                ProjectArches.Clear();
+                if (Ini.GetBool("/Script/IOSRuntimeSettings.IOSRuntimeSettings", "bShipForArmV7", out bBuild) && bBuild)
+                {
+                    ProjectArches.Add("armv7");
+                }
+                if (Ini.GetBool("/Script/IOSRuntimeSettings.IOSRuntimeSettings", "bShipForArm64", out bBuild) && bBuild)
+                {
+                    ProjectArches.Add("arm64");
+                }
+                if (Ini.GetBool("/Script/IOSRuntimeSettings.IOSRuntimeSettings", "bShipForArmV7S", out bBuild) && bBuild)
+                {
+                    ProjectArches.Add("armv7s");
+                }
 
-				// force armv7 if something went wrong
-				if (ProjectArches.Count == 0)
-				{
-					ProjectArches.Add("armv7");
-					ProjectArches.Add("arm64");
-				}
-				ShippingArchitectures = ProjectArches[0];
-				for (int Index = 1; Index < ProjectArches.Count; ++Index)
-				{
-					ShippingArchitectures += "," + ProjectArches[Index];
-				}
+                // force armv7 if something went wrong
+                if (ProjectArches.Count == 0)
+                {
+                    ProjectArches.Add("armv7");
+                    ProjectArches.Add("arm64");
+                }
+                ShippingArchitectures = ProjectArches[0];
+                for (int Index = 1; Index < ProjectArches.Count; ++Index)
+                {
+                    ShippingArchitectures += "," + ProjectArches[Index];
+                }
 
-				// determine if we need to generate the dsym
-				Ini.GetBool("/Script/IOSRuntimeSettings.IOSRuntimeSettings", "bGeneratedSYMFile", out BuildConfiguration.bGeneratedSYMFile);
-				Ini.GetBool("/Script/IOSRuntimeSettings.IOSRuntimeSettings", "bGeneratedSYMBundle", out BuildConfiguration.bGeneratedSYMBundle);
+                // determine if we need to generate the dsym
+                Ini.GetBool("/Script/IOSRuntimeSettings.IOSRuntimeSettings", "bGeneratedSYMFile", out BuildConfiguration.bGeneratedSYMFile);
+                Ini.GetBool("/Script/IOSRuntimeSettings.IOSRuntimeSettings", "bGeneratedSYMBundle", out BuildConfiguration.bGeneratedSYMBundle);
 
-				// determie if bitcode should be generated for the shipping code
-				Ini.GetBool("/Script/IOSRuntimeSettings.IOSRuntimeSettings", "bShipForBitcode", out bShipForBitcode);
+                // determie if bitcode should be generated for the shipping code
+                Ini.GetBool("/Script/IOSRuntimeSettings.IOSRuntimeSettings", "bShipForBitcode", out bShipForBitcode);
 
-				// @todo tvos: We probably want to handle TVOS versions here
-				Ini.GetString("/Script/IOSRuntimeSettings.IOSRuntimeSettings", "AdditionalLinkerFlags", out AdditionalLinkerFlags);
-				Ini.GetString("/Script/IOSRuntimeSettings.IOSRuntimeSettings", "AdditionalShippingLinkerFlags", out AdditionalShippingLinkerFlags);
+                // @todo tvos: We probably want to handle TVOS versions here
+                Ini.GetString("/Script/IOSRuntimeSettings.IOSRuntimeSettings", "AdditionalLinkerFlags", out AdditionalLinkerFlags);
+                Ini.GetString("/Script/IOSRuntimeSettings.IOSRuntimeSettings", "AdditionalShippingLinkerFlags", out AdditionalShippingLinkerFlags);
 
-				Ini.GetString("/Script/IOSRuntimeSettings.IOSRuntimeSettings", "MobileProvision", out MobileProvision);
-				Ini.GetString("/Script/IOSRuntimeSettings.IOSRuntimeSettings", "SigningCertificate", out SigningCertificate);
+                Ini.GetString("/Script/IOSRuntimeSettings.IOSRuntimeSettings", "MobileProvision", out MobileProvision);
+                Ini.GetString("/Script/IOSRuntimeSettings.IOSRuntimeSettings", "SigningCertificate", out SigningCertificate);
 
                 // bundle identifier
-                string BundleIdentifier;
                 Ini.GetString("/Script/IOSRuntimeSettings.IOSRuntimeSettings", "BundleIdentifier", out BundleIdentifier);
 
-                ProvisionData Data = new ProvisionData();
-                string BundleId = BundleIdentifier.Replace("[PROJECT_NAME]", ((ProjectFile != null) ? ProjectFile.GetFileNameWithoutAnyExtensions() : "UE4Game")).Replace("_", "");
-                if (!ProvisionCache.ContainsKey(BundleId))
+                bInitializedProject = true;
+            }
+
+            ProvisionData Data = new ProvisionData();
+            string BundleId = BundleIdentifier.Replace("[PROJECT_NAME]", ((ProjectFile != null) ? ProjectFile.GetFileNameWithoutAnyExtensions() : "UE4Game")).Replace("_", "");
+            bool bIsTVOS = GetCodesignPlatformName() == "appletvos";
+            if (!ProvisionCache.ContainsKey(BundleId + " " + bIsTVOS.ToString() + " " + bForDistribtion.ToString()))
+            {
+                Certificate = SigningCertificate;
+                Provision = MobileProvision;
+                if (!string.IsNullOrEmpty(SigningCertificate))
                 {
-                    Certificate = SigningCertificate;
-                    Provision = MobileProvision;
-                    if (!string.IsNullOrEmpty(SigningCertificate))
+                    // verify the certificate
+                    Process IPPProcess = new Process();
+                    if (BuildHostPlatform.Current.Platform == UnrealTargetPlatform.Mac)
                     {
-                        // verify the certificate
-                        Process IPPProcess = new Process();
-                        if (BuildHostPlatform.Current.Platform == UnrealTargetPlatform.Mac)
-                        {
-                            string IPPCmd = "\"" + UnrealBuildTool.EngineDirectory + "/Binaries/DotNET/IOS/IPhonePackager.exe\" certificates " + ((ProjectFile != null) ? ("\"" + ProjectFile.ToString() + "\"") : "Engine") + " -bundlename " + BundleIdentifier.Replace("[PROJECT_NAME]", ((ProjectFile != null) ? ProjectFile.GetFileNameWithoutAnyExtensions() : "UE4Game")).Replace("_", "");
-                            IPPProcess.StartInfo.WorkingDirectory = UnrealBuildTool.EngineDirectory.ToString();
-                            IPPProcess.StartInfo.FileName = UnrealBuildTool.EngineDirectory + "/Build/BatchFiles/Mac/RunMono.sh";
-                            IPPProcess.StartInfo.Arguments = IPPCmd;
-                            IPPProcess.OutputDataReceived += new DataReceivedEventHandler(IPPDataReceivedHandler);
-                            IPPProcess.ErrorDataReceived += new DataReceivedEventHandler(IPPDataReceivedHandler);
-                        }
-                        else
-                        {
-                            string IPPCmd = "certificates " + ((ProjectFile != null) ? ("\"" + ProjectFile.ToString() + "\"") : "Engine") + " -bundlename " + BundleId;
-                            IPPProcess.StartInfo.WorkingDirectory = UnrealBuildTool.EngineDirectory.ToString();
-                            IPPProcess.StartInfo.FileName = UnrealBuildTool.EngineDirectory + "\\Binaries\\DotNET\\IOS\\IPhonePackager.exe";
-                            IPPProcess.StartInfo.Arguments = IPPCmd;
-                            IPPProcess.OutputDataReceived += new DataReceivedEventHandler(IPPDataReceivedHandler);
-                            IPPProcess.ErrorDataReceived += new DataReceivedEventHandler(IPPDataReceivedHandler);
-                        }
-                        Utils.RunLocalProcess(IPPProcess);
+                        string IPPCmd = "\"" + UnrealBuildTool.EngineDirectory + "/Binaries/DotNET/IOS/IPhonePackager.exe\" certificates " + ((ProjectFile != null) ? ("\"" + ProjectFile.ToString() + "\"") : "Engine") + " -bundlename " + BundleId + (bForDistribtion ? " -distribution" : "");
+                        IPPProcess.StartInfo.WorkingDirectory = UnrealBuildTool.EngineDirectory.ToString();
+                        IPPProcess.StartInfo.FileName = UnrealBuildTool.EngineDirectory + "/Build/BatchFiles/Mac/RunMono.sh";
+                        IPPProcess.StartInfo.Arguments = IPPCmd;
+                        IPPProcess.OutputDataReceived += new DataReceivedEventHandler(IPPDataReceivedHandler);
+                        IPPProcess.ErrorDataReceived += new DataReceivedEventHandler(IPPDataReceivedHandler);
+                    }
+                    else
+                    {
+                        string IPPCmd = "certificates " + ((ProjectFile != null) ? ("\"" + ProjectFile.ToString() + "\"") : "Engine") + " -bundlename " + BundleId + (bForDistribtion ? " -distribution" : "");
+                        IPPProcess.StartInfo.WorkingDirectory = UnrealBuildTool.EngineDirectory.ToString();
+                        IPPProcess.StartInfo.FileName = UnrealBuildTool.EngineDirectory + "\\Binaries\\DotNET\\IOS\\IPhonePackager.exe";
+                        IPPProcess.StartInfo.Arguments = IPPCmd;
+                        IPPProcess.OutputDataReceived += new DataReceivedEventHandler(IPPDataReceivedHandler);
+                        IPPProcess.ErrorDataReceived += new DataReceivedEventHandler(IPPDataReceivedHandler);
                     }
+                    Utils.RunLocalProcess(IPPProcess);
+                }
+                else
+                {
+                    Certificate = bForDistribtion ? "iPhone Distribution" : "iPhone Developer";
+                    bHaveCertificate = true;
+                }
 
-                    if (string.IsNullOrEmpty(MobileProvision) // no provision specified
-                        || !File.Exists((BuildHostPlatform.Current.Platform == UnrealTargetPlatform.Mac ? (Environment.GetEnvironmentVariable("HOME") + "/Library/MobileDevice/Provisioning Profiles/") : (Environment.GetFolderPath(Environment.SpecialFolder.LocalApplicationData) + "/Apple Computer/MobileDevice/Provisioning Profiles/")) + MobileProvision) // file doesn't exist
-                        || !bHaveCertificate) // certificate doesn't exist
+                if (string.IsNullOrEmpty(MobileProvision) // no provision specified
+                    || !File.Exists((BuildHostPlatform.Current.Platform == UnrealTargetPlatform.Mac ? (Environment.GetEnvironmentVariable("HOME") + "/Library/MobileDevice/Provisioning Profiles/") : (Environment.GetFolderPath(Environment.SpecialFolder.LocalApplicationData) + "/Apple Computer/MobileDevice/Provisioning Profiles/")) + MobileProvision) // file doesn't exist
+                    || !bHaveCertificate) // certificate doesn't exist
+                {
+
+                    Certificate = "";
+                    Provision = "";
+                    Log.TraceLog("Provision not specified or not found for " + ((ProjectFile != null) ? ProjectFile.GetFileNameWithoutAnyExtensions() : "UE4Game") + ", searching for compatible match...");
+                    Process IPPProcess = new Process();
+                    if (BuildHostPlatform.Current.Platform == UnrealTargetPlatform.Mac)
                     {
-                        Certificate = "";
-                        Provision = "";
-                        Log.TraceLog("Provision not specified or not found for " + ((ProjectFile != null) ? ProjectFile.GetFileNameWithoutAnyExtensions() : "UE4Game") + ", searching for compatible match...");
-                        Process IPPProcess = new Process();
-                        if (BuildHostPlatform.Current.Platform == UnrealTargetPlatform.Mac)
-                        {
-                            string IPPCmd = "\"" + UnrealBuildTool.EngineDirectory + "/Binaries/DotNET/IOS/IPhonePackager.exe\" signing_match " + ((ProjectFile != null) ? ("\"" + ProjectFile.ToString() + "\"") : "Engine") + " -bundlename " + BundleIdentifier.Replace("[PROJECT_NAME]", ((ProjectFile != null) ? ProjectFile.GetFileNameWithoutAnyExtensions() : "UE4Game")).Replace("_", "");
-                            IPPProcess.StartInfo.WorkingDirectory = UnrealBuildTool.EngineDirectory.ToString();
-                            IPPProcess.StartInfo.FileName = UnrealBuildTool.EngineDirectory + "/Build/BatchFiles/Mac/RunMono.sh";
-                            IPPProcess.StartInfo.Arguments = IPPCmd;
-                            IPPProcess.OutputDataReceived += new DataReceivedEventHandler(IPPDataReceivedHandler);
-                            IPPProcess.ErrorDataReceived += new DataReceivedEventHandler(IPPDataReceivedHandler);
-                        }
-                        else
-                        {
-                            string IPPCmd = "signing_match " + ((ProjectFile != null) ? ("\"" + ProjectFile.ToString() + "\"") : "Engine") + " -bundlename " + BundleId;
-                            IPPProcess.StartInfo.WorkingDirectory = UnrealBuildTool.EngineDirectory.ToString();
-                            IPPProcess.StartInfo.FileName = UnrealBuildTool.EngineDirectory + "\\Binaries\\DotNET\\IOS\\IPhonePackager.exe";
-                            IPPProcess.StartInfo.Arguments = IPPCmd;
-                            IPPProcess.OutputDataReceived += new DataReceivedEventHandler(IPPDataReceivedHandler);
-                            IPPProcess.ErrorDataReceived += new DataReceivedEventHandler(IPPDataReceivedHandler);
-                        }
-                        Utils.RunLocalProcess(IPPProcess);
-                        Log.TraceLog("Provision found for " + ((ProjectFile != null) ? ProjectFile.GetFileNameWithoutAnyExtensions() : "UE4Game") + ", Provision: " + Provision + " Certificate: " + Certificate);
+                        string IPPCmd = "\"" + UnrealBuildTool.EngineDirectory + "/Binaries/DotNET/IOS/IPhonePackager.exe\" signing_match " + ((ProjectFile != null) ? ("\"" + ProjectFile.ToString() + "\"") : "Engine") + " -bundlename " + BundleId + (bIsTVOS ? " -tvos" : "") + (bForDistribtion ? " -distribution" : "");
+                        IPPProcess.StartInfo.WorkingDirectory = UnrealBuildTool.EngineDirectory.ToString();
+                        IPPProcess.StartInfo.FileName = UnrealBuildTool.EngineDirectory + "/Build/BatchFiles/Mac/RunMono.sh";
+                        IPPProcess.StartInfo.Arguments = IPPCmd;
+                        IPPProcess.OutputDataReceived += new DataReceivedEventHandler(IPPDataReceivedHandler);
+                        IPPProcess.ErrorDataReceived += new DataReceivedEventHandler(IPPDataReceivedHandler);
+                    }
+                    else
+                    {
+                        string IPPCmd = "signing_match " + ((ProjectFile != null) ? ("\"" + ProjectFile.ToString() + "\"") : "Engine") + " -bundlename " + BundleId + (bIsTVOS ? " -tvos" : "") + (bForDistribtion ? " -distribution" : "");
+                        IPPProcess.StartInfo.WorkingDirectory = UnrealBuildTool.EngineDirectory.ToString();
+                        IPPProcess.StartInfo.FileName = UnrealBuildTool.EngineDirectory + "\\Binaries\\DotNET\\IOS\\IPhonePackager.exe";
+                        IPPProcess.StartInfo.Arguments = IPPCmd;
+                        IPPProcess.OutputDataReceived += new DataReceivedEventHandler(IPPDataReceivedHandler);
+                        IPPProcess.ErrorDataReceived += new DataReceivedEventHandler(IPPDataReceivedHandler);
                     }
-                    // add to the dictionary
-                    Data.MobileProvision = Provision;
-                    Data.Certificate = Certificate;
+                    Utils.RunLocalProcess(IPPProcess);
+                    Log.TraceLog("Provision found for " + ((ProjectFile != null) ? ProjectFile.GetFileNameWithoutAnyExtensions() : "UE4Game") + ", Provision: " + Provision + " Certificate: " + Certificate);
+                }
+                // add to the dictionary
+                Data.MobileProvision = Provision;
+                Data.Certificate = Certificate.Replace("\"", "");
 
-                    // read the provision to get the UUID
-                    string filename = (BuildHostPlatform.Current.Platform == UnrealTargetPlatform.Mac ? (Environment.GetEnvironmentVariable("HOME") + "/Library/MobileDevice/Provisioning Profiles/") : (Environment.GetFolderPath(Environment.SpecialFolder.LocalApplicationData) + "/Apple Computer/MobileDevice/Provisioning Profiles/")) + Data.MobileProvision;
-                    if (File.Exists(filename))
+                // read the provision to get the UUID
+                string filename = (BuildHostPlatform.Current.Platform == UnrealTargetPlatform.Mac ? (Environment.GetEnvironmentVariable("HOME") + "/Library/MobileDevice/Provisioning Profiles/") : (Environment.GetFolderPath(Environment.SpecialFolder.LocalApplicationData) + "/Apple Computer/MobileDevice/Provisioning Profiles/")) + Data.MobileProvision;
+                if (File.Exists(filename))
+                {
+                    string AllText = File.ReadAllText(filename);
+                    int idx = AllText.IndexOf("<key>UUID</key>");
+                    if (idx > 0)
                     {
-                        string AllText = File.ReadAllText(filename);
-                        int idx = AllText.IndexOf("<key>UUID</key>");
+                        idx = AllText.IndexOf("<string>", idx);
                         if (idx > 0)
                         {
-                            idx = AllText.IndexOf("<string>", idx);
-                            if (idx > 0)
-                            {
-                                idx += "<string>".Length;
-                                Data.UUID = AllText.Substring(idx, AllText.IndexOf("</string>", idx) - idx);
-                            }
+                            idx += "<string>".Length;
+                            Data.UUID = AllText.Substring(idx, AllText.IndexOf("</string>", idx) - idx);
                         }
-                        idx = AllText.IndexOf("<key>com.apple.developer.team-identifier</key>");
+                    }
+                    idx = AllText.IndexOf("<key>com.apple.developer.team-identifier</key>");
+                    if (idx > 0)
+                    {
+                        idx = AllText.IndexOf("<string>", idx);
                         if (idx > 0)
                         {
-                            idx = AllText.IndexOf("<string>", idx);
-                            if (idx > 0)
-                            {
-                                idx += "<string>".Length;
-                                Data.TeamUUID = AllText.Substring(idx, AllText.IndexOf("</string>", idx) - idx);
-                            }
+                            idx += "<string>".Length;
+                            Data.TeamUUID = AllText.Substring(idx, AllText.IndexOf("</string>", idx) - idx);
                         }
                     }
-                    else
-                    {
-                        Log.TraceWarning("No matching provision file was discovered. Please ensure you have a compatible provision installed.");
-                    }
-                    ProvisionCache.Add(BundleId, Data);
                 }
                 else
                 {
-                    Data = ProvisionCache[BundleId];
+                    Log.TraceLog("No matching provision file was discovered. Please ensure you have a compatible provision installed.");
                 }
-                MobileProvision = Data.MobileProvision;
-                SigningCertificate = Data.Certificate;
-                MobileProvisionUUID = Data.UUID;
-                TeamUUID = Data.TeamUUID;
-
-                bInitializedProject = true;
-			}
-		}
+                ProvisionCache.Add(BundleId + " " + bIsTVOS.ToString() + " " + bForDistribtion.ToString(), Data);
+            }
+            else
+            {
+                Data = ProvisionCache[BundleId + " " + bIsTVOS.ToString() + " " + bForDistribtion.ToString()];
+            }
+            MobileProvision = Data.MobileProvision;
+            SigningCertificate = Data.Certificate;
+            MobileProvisionUUID = Data.UUID;
+            TeamUUID = Data.TeamUUID;
+        }
 
-		/// <summary>
-		/// Whether this platform should create debug information or not
-		/// </summary>
-		/// <param name="Configuration"> The UnrealTargetConfiguration being built</param>
-		/// <returns>bool    true if debug info should be generated, false if not</returns>
-		public override bool ShouldCreateDebugInfo(UnrealTargetConfiguration Configuration)
+        /// <summary>
+        /// Whether this platform should create debug information or not
+        /// </summary>
+        /// <param name="Configuration"> The UnrealTargetConfiguration being built</param>
+        /// <returns>bool    true if debug info should be generated, false if not</returns>
+        public override bool ShouldCreateDebugInfo(UnrealTargetConfiguration Configuration)
 		{
 			return true;
 		}
diff --git a/Engine/Source/Programs/UnrealBuildTool/TVOS/UEBuildTVOS.cs b/Engine/Source/Programs/UnrealBuildTool/TVOS/UEBuildTVOS.cs
--- a/Engine/Source/Programs/UnrealBuildTool/TVOS/UEBuildTVOS.cs
+++ b/Engine/Source/Programs/UnrealBuildTool/TVOS/UEBuildTVOS.cs
@@ -9,7 +9,7 @@ using System.Xml;
 
 namespace UnrealBuildTool
 {
-	class TVOSPlatformContext : IOSPlatformContext
+	public class TVOSPlatformContext : IOSPlatformContext
 	{
 		/** Which version of the iOS to allow at run time */
 		private const string RunTimeTVOSVersion = "9.0";
@@ -20,10 +20,10 @@ namespace UnrealBuildTool
 		// by default, use an empty architecture (which is really just a modifer to the platform for some paths/names)
 		public static string TVOSArchitecture = "";
 
-		public TVOSPlatformContext(FileReference InProjectFile)
-			: base(UnrealTargetPlatform.IOS, InProjectFile)
-		{
-		}
+        public TVOSPlatformContext(FileReference InProjectFile, bool ForDistribution = false)
+            : base(UnrealTargetPlatform.IOS, InProjectFile, ForDistribution)
+        {
+        }
 
 		// The current architecture - affects everything about how UBT operates on IOS
 		public override string GetActiveArchitecture()
