upgrade emsdk toolchain to 1.36.11

this is needed to fix up physx compiler errors

diff --git a/Engine/Source/Programs/UnrealBuildTool/HTML5/HTML5SDKInfo.cs b/Engine/Source/Programs/UnrealBuildTool/HTML5/HTML5SDKInfo.cs
--- a/Engine/Source/Programs/UnrealBuildTool/HTML5/HTML5SDKInfo.cs
+++ b/Engine/Source/Programs/UnrealBuildTool/HTML5/HTML5SDKInfo.cs
@@ -11,8 +11,10 @@ namespace UnrealBuildTool
 {
 	public class HTML5SDKInfo
 	{
+//		static string SDKVersion = "tag-1.36.11";
 		static string SDKVersion = "1.35.0";
 		static string NODE_VER = "4.1.1_64bit";
+//		static string LLVM_VER = "tag-1.36.11";
 		static string LLVM_VER = "e1.35.0_64bit";
 		static string PYTHON_VER = "2.7.5.3_64bit";
 
diff --git a/Engine/Source/Programs/UnrealBuildTool/HTML5/HTML5ToolChain.cs b/Engine/Source/Programs/UnrealBuildTool/HTML5/HTML5ToolChain.cs
--- a/Engine/Source/Programs/UnrealBuildTool/HTML5/HTML5ToolChain.cs
+++ b/Engine/Source/Programs/UnrealBuildTool/HTML5/HTML5ToolChain.cs
@@ -89,8 +89,13 @@ namespace UnrealBuildTool
 			// enable hardware accelerated and advanced instructions/functions
 //			Result += " -s USE_WEBGL2=1 -s FULL_ES3=1";
 			Result += " -s FULL_ES2=1";
+			Result += " -msse";
+			Result += " -msse2";
 			Result += " -s SIMD=1";
-//			Result += " -s USE_PTHREADS=2"; // 2:polyfill -- SharedInt\d+Array available by ES7
+			Result += " -s USE_PTHREADS=1"; // 2:polyfill -- SharedInt\d+Array available by ES7
+
+			// for DEBUGGING: http://stackoverflow.com/questions/42308/tool-to-track-include-dependencies
+//			Result += " -H";
 
 			// export console command handler. Export main func too because default exports ( e.g Main ) are overridden if we use custom exported functions.
 			Result += " -s EXPORTED_FUNCTIONS=\"['_main', '_on_fatal']\"";
