upgrade emsdk toolchain to 1.36.11

this is needed to fix up physx compiler errors

diff --git a/Engine/Source/Programs/UnrealBuildTool/HTML5/HTML5SDKInfo.cs b/Engine/Source/Programs/UnrealBuildTool/HTML5/HTML5SDKInfo.cs
--- a/Engine/Source/Programs/UnrealBuildTool/HTML5/HTML5SDKInfo.cs
+++ b/Engine/Source/Programs/UnrealBuildTool/HTML5/HTML5SDKInfo.cs
@@ -11,11 +11,9 @@ namespace UnrealBuildTool
 {
 	public class HTML5SDKInfo
 	{
-//		static string SDKVersion = "tag-1.36.11";
-		static string SDKVersion = "1.35.0";
+		static string SDKVersion = "tag-1.36.11";
 		static string NODE_VER = "4.1.1_64bit";
-//		static string LLVM_VER = "tag-1.36.11";
-		static string LLVM_VER = "e1.35.0_64bit";
+		static string LLVM_VER = "tag-1.36.11";
 		static string PYTHON_VER = "2.7.5.3_64bit";
 
 		// --------------------------------------------------
