Improved AddLibraryPathToRPaths in MacToolChain.cs to better handle RPATHs in packaged games

diff --git a/Engine/Source/Programs/UnrealBuildTool/Mac/MacToolChain.cs b/Engine/Source/Programs/UnrealBuildTool/Mac/MacToolChain.cs
--- a/Engine/Source/Programs/UnrealBuildTool/Mac/MacToolChain.cs
+++ b/Engine/Source/Programs/UnrealBuildTool/Mac/MacToolChain.cs
@@ -527,32 +527,35 @@ namespace UnrealBuildTool
 
 		private void AddLibraryPathToRPaths(string Library, string ExeAbsolutePath, ref List<string> RPaths, ref string LinkCommand, bool bIsBuildingAppBundle)
 		{
-			string LibraryDir = Path.GetDirectoryName(Library);
+			string LibraryFullPath = Path.GetFullPath(Library);
+ 			string LibraryDir = Path.GetDirectoryName(LibraryFullPath);
 			string ExeDir = Path.GetDirectoryName(ExeAbsolutePath);
-			if (!Library.Contains("/Engine/Binaries/Mac/") && (Library.EndsWith("dylib") || Library.EndsWith(".framework")) && LibraryDir != ExeDir)
+
+			// Only dylibs and frameworks, and only those that are outside of Engine/Binaries/Mac and Engine/Source/ThirdParty, and outside of the folder where the executable is need an additional RPATH entry
+			if ((Library.EndsWith("dylib") || Library.EndsWith(".framework")) && !LibraryFullPath.Contains("/Engine/Binaries/Mac/")
+			    && !LibraryFullPath.Contains("/Engine/Source/ThirdParty/") && LibraryDir != ExeDir && !RPaths.Contains(LibraryFullPath))
 			{
+				// First, add a path relative to the executable.
+				// We skip this for CEF3 for the Shipping Launcher. In this case we only want the RPATH to the framework inside the app bundle (added in the if (bIsBuildingAppBundle) block),
+				// otherwise OS X gatekeeper erroneously complains about not seeing the framework
 				string RelativePath = Utils.MakePathRelativeTo(LibraryDir, ExeDir).Replace("\\", "/");
-				if ((!RelativePath.Contains(LibraryDir) || RelativePath.Contains("/ThirdParty/")) && !RPaths.Contains(RelativePath))
+				if (!ExeAbsolutePath.Contains("EpicGamesLauncher-Mac-Shipping") || !Library.Contains("CEF3"))
 				{
-					// For CEF3 for the Shipping Launcher we only want the RPATH to the framework inside the app bundle, otherwise OS X gatekeeper erroneously complains about not seeing framework. 
-					if (!ExeAbsolutePath.Contains("EpicGamesLauncher-Mac-Shipping") || !Library.Contains("CEF3"))
-					{
-						RPaths.Add(RelativePath);
-						LinkCommand += " -rpath \"@loader_path/" + RelativePath + "\"";
-					}
+					LinkCommand += " -rpath \"@loader_path/" + RelativePath + "\"";
+				}
 
-					if (bIsBuildingAppBundle)
-					{
-						// If building an app bunlde, we also prepare an rpath for use in packaged game. In that case dylibs are stored in Contents/UE4/Engine/ subfolders, for example in
-						// GameName.app/Contents/UE4/Engine/Binaries/ThirdParty/PhysX/Mac/
-						// RelativePath is a path to a copy of the lib folder in engine directory tree, so we need to convert it to a path relative to what it will be inside the app bundle
-						int NumCharactersToSkip = ExeAbsolutePath.Contains("/Engine/Binaries/Mac/") ? 9 : 15; // Skip 3 or 5 sets of ../ from the beginning of relative path, depending on where the app bundle is
-						string PathInBundle = Path.Combine(Path.GetDirectoryName(ExeDir), "UE4/Engine/Binaries/Mac", RelativePath.Substring(NumCharactersToSkip));
-						Utils.CollapseRelativeDirectories(ref PathInBundle);
-						string RelativePathInBundle = Utils.MakePathRelativeTo(PathInBundle, ExeDir).Replace("\\", "/");
-						LinkCommand += " -rpath \"@loader_path/" + RelativePathInBundle + "\"";
-					}
+				// If building an app bundle, we also need an RPATH for use in packaged game. In that case dylibs are stored in Contents/UE4/Engine/ subfolders, for example in
+				// GameName.app/Contents/UE4/Engine/Binaries/ThirdParty/PhysX/Mac/
+				if (bIsBuildingAppBundle)
+				{
+					string EngineDir = UnrealBuildTool.RootDirectory.ToString();
+					string BundleUE4Dir = Path.GetFullPath(ExeDir + "/../../Contents/UE4");
+					string BundleLibraryDir = LibraryDir.Replace(EngineDir, BundleUE4Dir);
+					string BundleRelativeDir = Utils.MakePathRelativeTo(BundleLibraryDir, ExeDir).Replace("\\", "/");
+					LinkCommand += " -rpath \"@loader_path/" + BundleRelativeDir + "\"";
 				}
+
+				RPaths.Add(LibraryFullPath);
 			}
 		}
 
