Do not allow duplicated RPATHs on Mac

diff --git a/Engine/Source/Programs/UnrealBuildTool/Mac/MacToolChain.cs b/Engine/Source/Programs/UnrealBuildTool/Mac/MacToolChain.cs
--- a/Engine/Source/Programs/UnrealBuildTool/Mac/MacToolChain.cs
+++ b/Engine/Source/Programs/UnrealBuildTool/Mac/MacToolChain.cs
@@ -533,7 +533,7 @@ namespace UnrealBuildTool
 
 			// Only dylibs and frameworks, and only those that are outside of Engine/Binaries/Mac and Engine/Source/ThirdParty, and outside of the folder where the executable is need an additional RPATH entry
 			if ((Library.EndsWith("dylib") || Library.EndsWith(".framework")) && !LibraryFullPath.Contains("/Engine/Binaries/Mac/")
-			    && !LibraryFullPath.Contains("/Engine/Source/ThirdParty/") && LibraryDir != ExeDir && !RPaths.Contains(LibraryFullPath))
+			    && !LibraryFullPath.Contains("/Engine/Source/ThirdParty/") && LibraryDir != ExeDir && !RPaths.Contains(LibraryDir))
 			{
 				// First, add a path relative to the executable.
 				// We skip this for CEF3 for the Shipping Launcher. In this case we only want the RPATH to the framework inside the app bundle (added in the if (bIsBuildingAppBundle) block),
@@ -555,7 +555,7 @@ namespace UnrealBuildTool
 					LinkCommand += " -rpath \"@loader_path/" + BundleRelativeDir + "\"";
 				}
 
-				RPaths.Add(LibraryFullPath);
+				RPaths.Add(LibraryDir);
 			}
 		}
 
