UBT: Reduce /Zm value back down to 850. Failing to compile several targets on build machines. Fix for XGE configuration may counter the need for this (CL 3155909).

diff --git a/Engine/Source/Programs/UnrealBuildTool/Windows/VCToolChain.cs b/Engine/Source/Programs/UnrealBuildTool/Windows/VCToolChain.cs
--- a/Engine/Source/Programs/UnrealBuildTool/Windows/VCToolChain.cs
+++ b/Engine/Source/Programs/UnrealBuildTool/Windows/VCToolChain.cs
@@ -220,7 +220,7 @@ namespace UnrealBuildTool
 				Arguments.Append(" /Gy");
 
 				// Allow 1000% of the default memory allocation limit.
-				Arguments.Append(" /Zm1000");
+				Arguments.Append(" /Zm850");
 
 				// Disable "The file contains a character that cannot be represented in the current code page" warning for non-US windows.
 				Arguments.Append(" /wd4819");
