Linux: disable XGE for cross-builds (UE-37061).

- speculative fix for compile problems with the new toolchain.

diff --git a/Engine/Source/Programs/UnrealBuildTool/Linux/UEBuildLinux.cs b/Engine/Source/Programs/UnrealBuildTool/Linux/UEBuildLinux.cs
--- a/Engine/Source/Programs/UnrealBuildTool/Linux/UEBuildLinux.cs
+++ b/Engine/Source/Programs/UnrealBuildTool/Linux/UEBuildLinux.cs
@@ -235,7 +235,7 @@ namespace UnrealBuildTool
 
 		public override bool CanUseXGE()
 		{
-			return true;
+			return BuildHostPlatform.Current.Platform == UnrealTargetPlatform.Linux;
 		}
 
 		/// <summary>
