Back out revision 18 from //UE4/Main/Engine/Source/Programs/UnrealBuildTool/Linux/UEBuildLinux.cs

(Disabling XGE was not warranted, the problem was with the cross-toolchain).

diff --git a/Engine/Source/Programs/UnrealBuildTool/Linux/UEBuildLinux.cs b/Engine/Source/Programs/UnrealBuildTool/Linux/UEBuildLinux.cs
--- a/Engine/Source/Programs/UnrealBuildTool/Linux/UEBuildLinux.cs
+++ b/Engine/Source/Programs/UnrealBuildTool/Linux/UEBuildLinux.cs
@@ -235,7 +235,7 @@ namespace UnrealBuildTool
 
 		public override bool CanUseXGE()
 		{
-			return BuildHostPlatform.Current.Platform == UnrealTargetPlatform.Linux;
+			return true;
 		}
 
 		/// <summary>
