- Fixed some log spam to only show once

diff --git a/Engine/Source/Programs/UnrealBuildTool/IOS/UEBuildIOS.cs b/Engine/Source/Programs/UnrealBuildTool/IOS/UEBuildIOS.cs
--- a/Engine/Source/Programs/UnrealBuildTool/IOS/UEBuildIOS.cs
+++ b/Engine/Source/Programs/UnrealBuildTool/IOS/UEBuildIOS.cs
@@ -255,11 +255,11 @@ namespace UnrealBuildTool
                     switch (MinVersion)
                     {
                         case "IOS_61":
-                            Log.TraceWarning("IOS 6 is no longer supported in UE4 as 4.11");
+                            Log.TraceWarningOnce("IOS 6 is no longer supported in UE4 as 4.11");
                             RunTimeIOSVersion = "8.0";
                             break;
                         case "IOS_7":
-                            Log.TraceWarning("IOS 7 is no longer supported in UE4 as 4.14");
+                            Log.TraceWarningOnce("IOS 7 is no longer supported in UE4 as 4.14");
                             RunTimeIOSVersion = "8.0";
                             break;
                         case "IOS_8":
