Unsuppress warnings that remained mistakenly suppressed and fix one more (UE-38788).

(Merging 3204244 from Dev-Platform to Main).

diff --git a/Engine/Source/Editor/SkeletonEditor/Public/ISkeletonTreeBuilder.h b/Engine/Source/Editor/SkeletonEditor/Public/ISkeletonTreeBuilder.h
--- a/Engine/Source/Editor/SkeletonEditor/Public/ISkeletonTreeBuilder.h
+++ b/Engine/Source/Editor/SkeletonEditor/Public/ISkeletonTreeBuilder.h
@@ -69,6 +69,9 @@ struct FSkeletonTreeFilterArgs
 class ISkeletonTreeBuilder
 {
 public:
+
+	virtual ~ISkeletonTreeBuilder() {};
+
 	/**
 	 * Build an array of skeleton tree items to display in the tree.
 	 * @param	InArgs			Parameters to use when building the tree
diff --git a/Engine/Source/Programs/UnrealBuildTool/Linux/LinuxToolChain.cs b/Engine/Source/Programs/UnrealBuildTool/Linux/LinuxToolChain.cs
--- a/Engine/Source/Programs/UnrealBuildTool/Linux/LinuxToolChain.cs
+++ b/Engine/Source/Programs/UnrealBuildTool/Linux/LinuxToolChain.cs
@@ -383,10 +383,6 @@ namespace UnrealBuildTool
 				if (CompilerVersionGreaterOrEqual(3, 9, 0))
 				{
 					Result += " -Wno-undefined-var-template"; // not really a good warning to disable
-					Result += " -Wno-delete-non-virtual-dtor";  // at least for 4.14
-					Result += " -Wno-expansion-to-defined";		// at least for 4.14
-					Result += " -Wno-null-dereference";         // at least for 4.14
-					Result += " -Wno-literal-conversion";       // at least for 4.14
 				}
 			}
 
