UBT will not longer try to define encryption and signing key macros in binary builds.

diff --git a/Engine/Source/Programs/UnrealBuildTool/Configuration/UEBuildTarget.cs b/Engine/Source/Programs/UnrealBuildTool/Configuration/UEBuildTarget.cs
--- a/Engine/Source/Programs/UnrealBuildTool/Configuration/UEBuildTarget.cs
+++ b/Engine/Source/Programs/UnrealBuildTool/Configuration/UEBuildTarget.cs
@@ -3960,67 +3960,9 @@
 				GlobalCompileEnvironment.Config.Definitions.Add(String.Format("IS_PROGRAM={0}", TargetType == TargetRules.TargetType.Program ? "1" : "0"));
 			}
 
-			String[] RSAKeys;
-			String AESKey;
-			ParseEncryptionIni(ProjectDirectory, Platform, out RSAKeys, out AESKey);
-			bool bSigningEnabledByIniFile = false;
-
-			if (RSAKeys != null)
+			if (!bUseSharedBuildEnvironment)
 			{
-				GlobalCompileEnvironment.Config.Definitions.Add("DECRYPTION_KEY_MODULUS=\"" + RSAKeys[1] + "\"");
-				GlobalCompileEnvironment.Config.Definitions.Add("DECRYPTION_KEY_EXPONENT=\"" + RSAKeys[2] + "\"");
-
-				bSigningEnabledByIniFile = true;
+				SetupEncryptionAndSingningKeys();
 			}
 
-			if (!String.IsNullOrEmpty(AESKey))
-			{
-				if (AESKey.Length < 32)
-				{
-					Log.TraceError("AES key specified in configs must be at least 32 characters long!");
-				}
-				else
-				{
-					GlobalCompileEnvironment.Config.Definitions.Add("AES_KEY=\"" + AESKey + "\"");
-				}
-			}
-			
-			// If we didn't extract any keys from the new ini file setup, try looking for the old keys text file
-			if (!bSigningEnabledByIniFile && !string.IsNullOrEmpty(Rules.PakSigningKeysFile))
-			{
-				string FullFilename = Path.Combine(ProjectDirectory.FullName, Rules.PakSigningKeysFile);
-
-				Log.TraceVerbose("Adding signing keys to executable from '{0}'", FullFilename);
-
-				if (File.Exists(FullFilename))
-				{
-					string[] Lines = File.ReadAllLines(FullFilename);
-					List<string> Keys = new List<string>();
-					foreach (string Line in Lines)
-					{
-						if (!string.IsNullOrEmpty(Line))
-						{
-							if (Line.StartsWith("0x"))
-							{
-								Keys.Add(Line.Trim());
-							}
-						}
-					}
-
-					if (Keys.Count == 3)
-					{
-						GlobalCompileEnvironment.Config.Definitions.Add("DECRYPTION_KEY_MODULUS=\"" + Keys[1] + "\"");
-						GlobalCompileEnvironment.Config.Definitions.Add("DECRYPTION_KEY_EXPONENT=\"" + Keys[2] + "\"");
-					}
-					else
-					{
-						Log.TraceWarning("Contents of signing key file are invalid so will be ignored");
-					}
-				}
-				else
-				{
-					Log.TraceVerbose("Signing key file is missing! Executable will not include signing keys");
-				}
-			}
-
 			// Validate UE configuration - needs to happen before setting any environment mojo and after argument parsing.
@@ -4268,6 +4210,72 @@
 				PlatformContext);
 		}
 
+		private void SetupEncryptionAndSingningKeys()
+		{
+			String[] RSAKeys;
+			String AESKey;
+			ParseEncryptionIni(ProjectDirectory, Platform, out RSAKeys, out AESKey);
+			bool bSigningEnabledByIniFile = false;
+
+			if (RSAKeys != null)
+			{
+				GlobalCompileEnvironment.Config.Definitions.Add("DECRYPTION_KEY_MODULUS=\"" + RSAKeys[1] + "\"");
+				GlobalCompileEnvironment.Config.Definitions.Add("DECRYPTION_KEY_EXPONENT=\"" + RSAKeys[2] + "\"");
+
+				bSigningEnabledByIniFile = true;
+			}
+
+			if (!String.IsNullOrEmpty(AESKey))
+			{
+				if (AESKey.Length < 32)
+				{
+					Log.TraceError("AES key specified in configs must be at least 32 characters long!");
+				}
+				else
+				{
+					GlobalCompileEnvironment.Config.Definitions.Add("AES_KEY=\"" + AESKey + "\"");
+				}
+			}
+
+			// If we didn't extract any keys from the new ini file setup, try looking for the old keys text file
+			if (!bSigningEnabledByIniFile && !string.IsNullOrEmpty(Rules.PakSigningKeysFile))
+			{
+				string FullFilename = Path.Combine(ProjectDirectory.FullName, Rules.PakSigningKeysFile);
+
+				Log.TraceVerbose("Adding signing keys to executable from '{0}'", FullFilename);
+
+				if (File.Exists(FullFilename))
+				{
+					string[] Lines = File.ReadAllLines(FullFilename);
+					List<string> Keys = new List<string>();
+					foreach (string Line in Lines)
+					{
+						if (!string.IsNullOrEmpty(Line))
+						{
+							if (Line.StartsWith("0x"))
+							{
+								Keys.Add(Line.Trim());
+							}
+						}
+					}
+
+					if (Keys.Count == 3)
+					{
+						GlobalCompileEnvironment.Config.Definitions.Add("DECRYPTION_KEY_MODULUS=\"" + Keys[1] + "\"");
+						GlobalCompileEnvironment.Config.Definitions.Add("DECRYPTION_KEY_EXPONENT=\"" + Keys[2] + "\"");
+					}
+					else
+					{
+						Log.TraceWarning("Contents of signing key file are invalid so will be ignored");
+					}
+				}
+				else
+				{
+					Log.TraceVerbose("Signing key file is missing! Executable will not include signing keys");
+				}
+			}
+		}
+
 		void SetUpPlatformEnvironment()
 		{
 			UEBuildPlatform BuildPlatform = UEBuildPlatform.GetBuildPlatform(Platform);
diff --git a/Engine/Source/Programs/UnrealPak/Private/UnrealPak.cpp b/Engine/Source/Programs/UnrealPak/Private/UnrealPak.cpp
--- a/Engine/Source/Programs/UnrealPak/Private/UnrealPak.cpp
+++ b/Engine/Source/Programs/UnrealPak/Private/UnrealPak.cpp
@@ -251,7 +251,7 @@ bool PrepareCopyFileToPak(const FString& InMountPoint, const FPakInputPair& InFi
 		OutNewEntry.Info.Size = FileSize;
 		OutNewEntry.Info.UncompressedSize = FileSize;
 		OutNewEntry.Info.CompressionMethod = COMPRESS_None;
-		OutNewEntry.Info.bEncrypted = InFile.bNeedEncryption;
+		OutNewEntry.Info.bEncrypted = InFile.bNeedEncryption && InEncryptionKey != nullptr;
 
 		if (InOutBufferSize < PaddedEncryptedFileSize)
 		{
@@ -332,7 +332,7 @@ bool PrepareCopyCompressedFileToPak(const FString& InMountPoint, const FPakInput
 	//	Write the header, then the data
 	OutNewEntry.Filename = InFile.Dest.Mid(InMountPoint.Len());
 	OutNewEntry.Info.Offset = 0; // Don't serialize offsets here.
-	OutNewEntry.Info.bEncrypted = InFile.bNeedEncryption;
+	OutNewEntry.Info.bEncrypted = InFile.bNeedEncryption && InEncryptionKey != nullptr;
 	OutSizeToWrite = CompressedFile.TotalCompressedSize;
 	OutDataToWrite = CompressedFile.CompressedBuffer.Get();
 	//OutNewEntry.Info.Serialize(InPak,FPakInfo::PakFile_Version_Latest);	
@@ -1139,7 +1139,7 @@ bool CreatePakFile(const TCHAR* Filename, TArray<FPakInputPair>& FilesToAdd, con
 
 	if (TotalEncryptedFiles < TotalRequestedEncryptedFiles)
 	{
-		UE_LOG(LogPakFile, Warning, TEXT("%d files requested encryption, but no AES key was supplied! Encryption was skipped for these files"), TotalRequestedEncryptedFiles);
+		UE_LOG(LogPakFile, Display, TEXT("%d files requested encryption, but no AES key was supplied! Encryption was skipped for these files"), TotalRequestedEncryptedFiles);
 	}
 
 	PakFileHandle->Close();
