UBT: Prevent default resource files being linked in to binaries that already have their own resources. Fixes local compilation of ShooterGameEditor et al.

diff --git a/Engine/Source/Programs/UnrealBuildTool/Configuration/UEBuildBinary.cs b/Engine/Source/Programs/UnrealBuildTool/Configuration/UEBuildBinary.cs
--- a/Engine/Source/Programs/UnrealBuildTool/Configuration/UEBuildBinary.cs
+++ b/Engine/Source/Programs/UnrealBuildTool/Configuration/UEBuildBinary.cs
@@ -728,19 +728,27 @@ namespace UnrealBuildTool
 			// Add the default resources for dlls
 			if (Config.Type == UEBuildBinaryType.DynamicLinkLibrary)
 			{
-				// For formal builds, if this binary does not already have a resource file, compile the default resource file. This will ensure it
-				// gets the correct ORIGINAL_FILE_NAME macro.
-				if (UEBuildConfiguration.bFormalBuild && !BinaryLinkEnvironment.InputFiles.Any(x => x.Reference.HasExtension(".res")))
+				// Check if there's already a custom resource file
+				if (!BinaryLinkEnvironment.InputFiles.Any(x => x.Reference.HasExtension(".res")))
 				{
-					CPPEnvironment BinaryResourceCompileEnvironment = BinaryCompileEnvironment.DeepCopy();
-					BinaryResourceCompileEnvironment.Config.OutputDirectory = DirectoryReference.Combine(BinaryResourceCompileEnvironment.Config.OutputDirectory, Modules.First().Name);
-					FileItem DefaultResourceFile = FileItem.GetItemByFileReference(FileReference.Combine(UnrealBuildTool.EngineSourceDirectory, "Runtime", "Launch", "Resources", "Windows", "PCLaunch.rc"));
-					CPPOutput DefaultResourceOutput = ToolChain.CompileRCFiles(BinaryResourceCompileEnvironment, new List<FileItem> { DefaultResourceFile }, ActionGraph);
-					BinaryLinkEnvironment.InputFiles.AddRange(DefaultResourceOutput.ObjectFiles);
+					if (UEBuildConfiguration.bFormalBuild)
+					{
+						// For formal builds, compile the default resource file per-binary, so that it gets the correct ORIGINAL_FILE_NAME macro.
+						CPPEnvironment BinaryResourceCompileEnvironment = BinaryCompileEnvironment.DeepCopy();
+						BinaryResourceCompileEnvironment.Config.OutputDirectory = DirectoryReference.Combine(BinaryResourceCompileEnvironment.Config.OutputDirectory, Modules.First().Name);
+						FileItem DefaultResourceFile = FileItem.GetItemByFileReference(FileReference.Combine(UnrealBuildTool.EngineSourceDirectory, "Runtime", "Launch", "Resources", "Windows", "PCLaunch.rc"));
+						CPPOutput DefaultResourceOutput = ToolChain.CompileRCFiles(BinaryResourceCompileEnvironment, new List<FileItem> { DefaultResourceFile }, ActionGraph);
+						BinaryLinkEnvironment.InputFiles.AddRange(DefaultResourceOutput.ObjectFiles);
+					}
+					else
+					{
+						// For non-formal builds, we just want to share the default resource file between modules
+						BinaryLinkEnvironment.InputFiles.AddRange(BinaryLinkEnvironment.DefaultResourceFiles);
+					}
 				}
 
 				// Add all the common resource files
-				BinaryLinkEnvironment.InputFiles.AddRange(BinaryLinkEnvironment.InputDllResourceFiles);
+				BinaryLinkEnvironment.InputFiles.AddRange(BinaryLinkEnvironment.CommonResourceFiles);
 			}
 
 			return BinaryLinkEnvironment;
diff --git a/Engine/Source/Programs/UnrealBuildTool/Configuration/UEBuildTarget.cs b/Engine/Source/Programs/UnrealBuildTool/Configuration/UEBuildTarget.cs
--- a/Engine/Source/Programs/UnrealBuildTool/Configuration/UEBuildTarget.cs
+++ b/Engine/Source/Programs/UnrealBuildTool/Configuration/UEBuildTarget.cs
@@ -2544,7 +2544,7 @@ namespace UnrealBuildTool
 				{
 					FileItem VersionFile = FileItem.GetExistingItemByFileReference(FileReference.Combine(UnrealBuildTool.EngineSourceDirectory, "Runtime", "Core", "Resources", "Windows", "ModuleVersionResource.rc.inl"));
 					CPPOutput VersionOutput = TargetToolChain.CompileRCFiles(GlobalCompileEnvironment, new List<FileItem> { VersionFile }, ActionGraph);
-					GlobalLinkEnvironment.InputDllResourceFiles.AddRange(VersionOutput.ObjectFiles);
+					GlobalLinkEnvironment.CommonResourceFiles.AddRange(VersionOutput.ObjectFiles);
 
 					if(!UEBuildConfiguration.bFormalBuild)
 					{
@@ -2554,7 +2554,7 @@ namespace UnrealBuildTool
 						FileItem DefaultResourceFile = FileItem.GetExistingItemByFileReference(FileReference.Combine(UnrealBuildTool.EngineSourceDirectory, "Runtime", "Launch", "Resources", "Windows", "PCLaunch.rc"));
 						CPPOutput DefaultResourceOutput = TargetToolChain.CompileRCFiles(DefaultResourceCompileEnvironment, new List<FileItem> { DefaultResourceFile }, ActionGraph);
 
-						GlobalLinkEnvironment.InputDllResourceFiles.AddRange(DefaultResourceOutput.ObjectFiles);
+						GlobalLinkEnvironment.DefaultResourceFiles.AddRange(DefaultResourceOutput.ObjectFiles);
 					}
 				}
 			}
diff --git a/Engine/Source/Programs/UnrealBuildTool/System/LinkEnvironment.cs b/Engine/Source/Programs/UnrealBuildTool/System/LinkEnvironment.cs
--- a/Engine/Source/Programs/UnrealBuildTool/System/LinkEnvironment.cs
+++ b/Engine/Source/Programs/UnrealBuildTool/System/LinkEnvironment.cs
@@ -237,9 +237,14 @@ namespace UnrealBuildTool
 		public List<FileItem> InputLibraries = new List<FileItem>();
 
 		/// <summary>
-		/// Compiled resources that should be linked into every binary
+		/// The default resource file to link in to every binary if a custom one is not provided
 		/// </summary>
-		public List<FileItem> InputDllResourceFiles = new List<FileItem>();
+		public List<FileItem> DefaultResourceFiles = new List<FileItem>();
+
+		/// <summary>
+		/// Resource files which should be compiled into every binary
+		/// </summary>
+		public List<FileItem> CommonResourceFiles = new List<FileItem>();
 
 		/// <summary>
 		/// The LinkEnvironmentConfiguration.
@@ -260,8 +265,9 @@ namespace UnrealBuildTool
 		{
 			InputFiles.AddRange(InCopyEnvironment.InputFiles);
 			InputLibraries.AddRange(InCopyEnvironment.InputLibraries);
-			InputDllResourceFiles.AddRange(InCopyEnvironment.InputDllResourceFiles);
-
+			DefaultResourceFiles.AddRange(InCopyEnvironment.DefaultResourceFiles);
+			CommonResourceFiles.AddRange(InCopyEnvironment.CommonResourceFiles);
+			
 			Config = new LinkEnvironmentConfiguration(InCopyEnvironment.Config);
 		}
 
