Merging CL 3217005 from //UE4/Dev-Platform: Fx to FixDylib script being run unnecessarily on too many modules. Mac build farm is oversaturated right now.

diff --git a/Engine/Source/Programs/UnrealBuildTool/Configuration/UEBuildTarget.cs b/Engine/Source/Programs/UnrealBuildTool/Configuration/UEBuildTarget.cs
--- a/Engine/Source/Programs/UnrealBuildTool/Configuration/UEBuildTarget.cs
+++ b/Engine/Source/Programs/UnrealBuildTool/Configuration/UEBuildTarget.cs
@@ -2775,12 +2775,30 @@ namespace UnrealBuildTool
 				}
 			}
 
-			// Create import libraries for all binaries ahead of time, since linking binaries often causes bottlenecks
 			if (!bCompileMonolithic)
 			{
-				foreach(UEBuildBinary Binary in AppBinaries)
+				if (GlobalLinkEnvironment.Config.Platform == CPPTargetPlatform.Win64 || GlobalLinkEnvironment.Config.Platform == CPPTargetPlatform.Win32)
 				{
-					Binary.SetCreateImportLibrarySeparately(true);
+					// On Windows create import libraries for all binaries ahead of time, since linking binaries often causes bottlenecks
+					foreach (UEBuildBinary Binary in AppBinaries)
+					{
+						Binary.SetCreateImportLibrarySeparately(true);
+					}
+				}
+				else
+				{
+					// On other platforms markup all the binaries containing modules with circular references
+					foreach (UEBuildModule Module in Modules.Values.Where(x => x.Binary != null))
+					{
+						foreach (string CircularlyReferencedModuleName in Module.Rules.CircularlyReferencedDependentModules)
+						{
+							UEBuildModule CircularlyReferencedModule;
+							if (Modules.TryGetValue(CircularlyReferencedModuleName, out CircularlyReferencedModule) && CircularlyReferencedModule.Binary != null)
+							{
+								CircularlyReferencedModule.Binary.SetCreateImportLibrarySeparately(true);
+							}
+						}
+					}
 				}
 			}
 
