UBT: Fix compile errors in targets which use the bAdaptiveUnityDisablesOptimizations setting (eg. OrionClient, OrionServer). The shared PCH functionality writes definitions from the compile environment into intermediate files, and includes them. bAdaptiveUnityDisablesOptimizations was clearing out the PCH setting, which resulted in all those definitions being lost. Now uses the base compile environment for these files instead.

diff --git a/Engine/Source/Programs/UnrealBuildTool/Configuration/UEBuildModule.cs b/Engine/Source/Programs/UnrealBuildTool/Configuration/UEBuildModule.cs
--- a/Engine/Source/Programs/UnrealBuildTool/Configuration/UEBuildModule.cs
+++ b/Engine/Source/Programs/UnrealBuildTool/Configuration/UEBuildModule.cs
@@ -1253,7 +1253,7 @@ namespace UnrealBuildTool
 			if (bModuleUsesUnityBuild)
 			{
 				CPPFilesToCompile = Unity.GenerateUnityCPPs(ToolChain, Target, CPPFilesToCompile, CPPCompileEnvironment, Name);
-				LinkInputFiles.AddRange(CompileUnityFilesWithToolChain(ToolChain, CPPCompileEnvironment, CPPFilesToCompile, Name, ActionGraph).ObjectFiles);
+				LinkInputFiles.AddRange(CompileUnityFilesWithToolChain(ToolChain, CPPCompileEnvironment, ModuleCompileEnvironment, CPPFilesToCompile, Name, ActionGraph).ObjectFiles);
 			}
 			else
 			{
@@ -1397,7 +1397,7 @@ namespace UnrealBuildTool
 		/// <summary>
 		/// Compiles the provided CPP unity files. Will
 		/// </summary>
-		private CPPOutput CompileUnityFilesWithToolChain(UEToolChain ToolChain, CPPEnvironment CompileEnvironment, List<FileItem> SourceFiles, string ModuleName, ActionGraph ActionGraph)
+		private CPPOutput CompileUnityFilesWithToolChain(UEToolChain ToolChain, CPPEnvironment CompileEnvironment, CPPEnvironment ModuleCompileEnvironment, List<FileItem> SourceFiles, string ModuleName, ActionGraph ActionGraph)
 		{
 			List<FileItem> NormalFiles = new List<FileItem>();
 			List<FileItem> AdaptiveFiles = new List<FileItem>();
@@ -1433,7 +1433,7 @@ namespace UnrealBuildTool
 			{
 				// Create an unoptmized compilation environment. Need to turn of PCH due to different
 				// compiler settings
-				CPPEnvironment UnoptimziedEnvironment = CompileEnvironment.DeepCopy();
+				CPPEnvironment UnoptimziedEnvironment = ModuleCompileEnvironment.DeepCopy();
 				UnoptimziedEnvironment.Config.bOptimizeCode = false;
 				UnoptimziedEnvironment.Config.PrecompiledHeaderAction = PrecompiledHeaderAction.None;
 
