UBT: Fixed missing DebugGame Editor configuration in generated project files. Regression from //UE4/Dev-Build integration.

diff --git a/Engine/Source/Programs/UnrealBuildTool/System/VCProject.cs b/Engine/Source/Programs/UnrealBuildTool/System/VCProject.cs
--- a/Engine/Source/Programs/UnrealBuildTool/System/VCProject.cs
+++ b/Engine/Source/Programs/UnrealBuildTool/System/VCProject.cs
@@ -135,10 +135,14 @@ namespace UnrealBuildTool
 			}
 
 			// Otherwise, get the default for the target type
-			if (Rules.Type == TargetRules.TargetType.Program || Rules.Type == TargetRules.TargetType.Editor)
+			if (Rules.Type == TargetRules.TargetType.Program)
 			{
 				return new[] { UnrealTargetConfiguration.Debug, UnrealTargetConfiguration.Development };
 			}
+			else if(Rules.Type == TargetRules.TargetType.Editor)
+			{
+				return new[] { UnrealTargetConfiguration.Debug, UnrealTargetConfiguration.DebugGame, UnrealTargetConfiguration.Development };
+			}
 			else
 			{
 				return ((UnrealTargetConfiguration[])Enum.GetValues(typeof(UnrealTargetConfiguration))).Where(x => x != UnrealTargetConfiguration.Unknown).ToArray();
