Fix non-unity build failures, some warnings about undeclared identifiers in conditional expressions, accessing OutputPath from deployment context on Android.

diff --git a/Engine/Source/Developer/DesktopPlatform/Private/Linux/DesktopPlatformLinux.cpp b/Engine/Source/Developer/DesktopPlatform/Private/Linux/DesktopPlatformLinux.cpp
--- a/Engine/Source/Developer/DesktopPlatform/Private/Linux/DesktopPlatformLinux.cpp
+++ b/Engine/Source/Developer/DesktopPlatform/Private/Linux/DesktopPlatformLinux.cpp
@@ -12,6 +12,10 @@
 #include "Misc/FeedbackContextMarkup.h"
 #include "HAL/ThreadHeartBeat.h"
 
+#ifndef WITH_LINUX_NATIVE_DIALOGS
+#define WITH_LINUX_NATIVE_DIALOGS 0
+#endif
+
 //#include "LinuxNativeFeedbackContext.h"
 // custom dialogs
 #if WITH_LINUX_NATIVE_DIALOGS
diff --git a/Engine/Source/Programs/UnrealBuildTool/Configuration/UEBuildDeploy.cs b/Engine/Source/Programs/UnrealBuildTool/Configuration/UEBuildDeploy.cs
--- a/Engine/Source/Programs/UnrealBuildTool/Configuration/UEBuildDeploy.cs
+++ b/Engine/Source/Programs/UnrealBuildTool/Configuration/UEBuildDeploy.cs
@@ -52,7 +52,17 @@ namespace UnrealBuildTool
 		/// <summary>
 		/// The output path
 		/// </summary>
-		public readonly FileReference OutputPath;
+		public FileReference OutputPath
+		{
+			get
+			{
+				if (OutputPaths.Count != 1)
+				{
+					throw new BuildException("Attempted to use UEBuildDeployTarget.OutputPath property, but there are multiple (or no) OutputPaths. You need to handle multiple in the code that called this (size = {0})", OutputPaths.Count);
+				}
+				return OutputPaths[0];
+			}
+		}
 
 		/// <summary>
 		/// The full list of output paths, for platforms that support building multiple binaries simultaneously
@@ -92,7 +102,6 @@ namespace UnrealBuildTool
 			this.TargetType = Target.TargetType;
 			this.Platform = Target.Platform;
 			this.Configuration = Target.Configuration;
-			this.OutputPath = Target.OutputPath;
 			this.OutputPaths = new List<FileReference>(Target.OutputPaths);
 			this.EngineIntermediateDirectory = Target.EngineIntermediateDirectory;
 			this.ProjectDirectory = Target.ProjectDirectory;
@@ -115,7 +124,6 @@ namespace UnrealBuildTool
 				TargetType = (TargetRules.TargetType)Reader.ReadInt32();
 				Platform = (UnrealTargetPlatform)Reader.ReadInt32();
 				Configuration = (UnrealTargetConfiguration)Reader.ReadInt32();
-				OutputPath = Reader.ReadFileReference();
 				OutputPaths = Reader.ReadList(x => x.ReadFileReference());
 				EngineIntermediateDirectory = Reader.ReadDirectoryReference();
 				ProjectDirectory = Reader.ReadDirectoryReference();
@@ -140,7 +148,6 @@ namespace UnrealBuildTool
 				Writer.Write((Int32)TargetType);
 				Writer.Write((Int32)Platform);
 				Writer.Write((Int32)Configuration);
-				Writer.Write(OutputPath);
 				Writer.Write(OutputPaths, (x, i) => x.Write(i));
 				Writer.Write(EngineIntermediateDirectory);
 				Writer.Write(ProjectDirectory);
