Fix deploy packagename to use TargetName instead of AppName
#ue4
#android

diff --git a/Engine/Source/Programs/UnrealBuildTool/Android/UEDeployAndroid.cs b/Engine/Source/Programs/UnrealBuildTool/Android/UEDeployAndroid.cs
--- a/Engine/Source/Programs/UnrealBuildTool/Android/UEDeployAndroid.cs
+++ b/Engine/Source/Programs/UnrealBuildTool/Android/UEDeployAndroid.cs
@@ -2359,7 +2359,7 @@ namespace UnrealBuildTool
 
 		public override bool PrepTargetForDeployment(UEBuildDeployTarget InTarget)
 		{
-			//Log.TraceInformation("$$$$$$$$$$$$$$ PrepTargetForDeployment $$$$$$$$$$$$$$$$$");
+			//Log.TraceInformation("$$$$$$$$$$$$$$ PrepTargetForDeployment $$$$$$$$$$$$$$$$$ {0}", InTarget.TargetName);
 			AndroidToolChain ToolChain = new AndroidToolChain(InTarget.ProjectFile); 
 
 			// we need to strip architecture from any of the output paths
@@ -2374,7 +2374,7 @@ namespace UnrealBuildTool
 			SetAndroidPluginData(ToolChain.GetAllArchitectures(), CollectPluginDataPaths(TargetReceipt.Read(ReceiptFilename)));
 
 			// make an apk at the end of compiling, so that we can run without packaging (debugger, cook on the fly, etc)
-			MakeApk(ToolChain, InTarget.AppName, InTarget.ProjectDirectory.FullName, BaseSoName, BuildConfiguration.RelativeEnginePath, bForDistribution: false, CookFlavor: "",
+			MakeApk(ToolChain, InTarget.TargetName, InTarget.ProjectDirectory.FullName, BaseSoName, BuildConfiguration.RelativeEnginePath, bForDistribution: false, CookFlavor: "",
 				bMakeSeparateApks: ShouldMakeSeparateApks(), bIncrementalPackage: true, bDisallowPackagingDataInApk: false, bDisallowExternalFilesDir: true);
 
 			// if we made any non-standard .apk files, the generated debugger settings may be wrong
