UBT: Do not forward the -2013/-2015/-2017 arguments to child instances of UBT created to rebuild UHT. Fixes locally compiling projects that force the compiler version to VS2013, which UHT does not support.

diff --git a/Engine/Source/Programs/UnrealBuildTool/System/ExternalExecution.cs b/Engine/Source/Programs/UnrealBuildTool/System/ExternalExecution.cs
--- a/Engine/Source/Programs/UnrealBuildTool/System/ExternalExecution.cs
+++ b/Engine/Source/Programs/UnrealBuildTool/System/ExternalExecution.cs
@@ -974,18 +974,6 @@ namespace UnrealBuildTool
 						}
 
 						// Propagate command-line options
-						if ( WindowsPlatform.Compiler == WindowsCompiler.VisualStudio2017 )
-						{
-							UBTArguments.Append(" -2017");
-						}
-						if ( WindowsPlatform.Compiler == WindowsCompiler.VisualStudio2015 )
-						{
-							UBTArguments.Append(" -2015");
-						}
-						if ( WindowsPlatform.Compiler == WindowsCompiler.VisualStudio2013 )
-						{
-							UBTArguments.Append(" -2013");
-						}
 						if ( BuildConfiguration.bIgnoreJunk )
 						{
 							UBTArguments.Append(" -ignorejunk");
