Fix build failures from //UE4/Dev-Build merge.

* UBT configuration files now read at UAT startup to fix exception packaging iOS/TVOS.
* Linux architecture specified as "" rather than null to fix exception stripping symbols on Linux for installed builds.
* Disabled automated tests.

diff --git a/Engine/Source/Programs/AutomationTool/AutomationUtils/Automation.cs b/Engine/Source/Programs/AutomationTool/AutomationUtils/Automation.cs
--- a/Engine/Source/Programs/AutomationTool/AutomationUtils/Automation.cs
+++ b/Engine/Source/Programs/AutomationTool/AutomationUtils/Automation.cs
@@ -481,7 +481,13 @@ AutomationTool.exe [-verbose] [-compileonly] [-p4] Command0 [-Arg0 -Arg1 -Arg2 
 			{
 				bIsEngineInstalled = GlobalCommandLine.InstalledEngine;
 			}
-			UnrealBuildTool.PlatformExports.SetIsEngineInstalled(bIsEngineInstalled.Value);
+
+			// Initialize UBT
+			if(!UnrealBuildTool.PlatformExports.Initialize(bIsEngineInstalled.Value))
+			{
+				Log.TraceInformation("Failed to initialize UBT");
+				return ExitCode.Error_Unknown;
+			}
 
 			// Change CWD to UE4 root.
 			Environment.CurrentDirectory = CommandUtils.CmdEnv.LocalRoot;
diff --git a/Engine/Source/Programs/AutomationTool/AutomationUtils/ScriptCompiler.cs b/Engine/Source/Programs/AutomationTool/AutomationUtils/ScriptCompiler.cs
--- a/Engine/Source/Programs/AutomationTool/AutomationUtils/ScriptCompiler.cs
+++ b/Engine/Source/Programs/AutomationTool/AutomationUtils/ScriptCompiler.cs
@@ -71,8 +71,6 @@ namespace AutomationTool
 				Environment.CurrentDirectory = UnrealBuildToolCWD;
 			}
 			// Register all the classes inside UBT
-			Log.TraceVerbose("Registering UBT Classes.");
-			UnrealBuildTool.PlatformExports.RegisterClasses();
 			Environment.CurrentDirectory = OldCWD;
 
 			// Compile only if not disallowed.
diff --git a/Engine/Source/Programs/UnrealBuildTool/HTML5/UEBuildHTML5.cs b/Engine/Source/Programs/UnrealBuildTool/HTML5/UEBuildHTML5.cs
--- a/Engine/Source/Programs/UnrealBuildTool/HTML5/UEBuildHTML5.cs
+++ b/Engine/Source/Programs/UnrealBuildTool/HTML5/UEBuildHTML5.cs
@@ -372,9 +372,8 @@ namespace UnrealBuildTool
 				bool bRegisterBuildPlatform = true;
 
 				// make sure we have the HTML5 files; if not, then this user doesn't really have HTML5 access/files, no need to compile HTML5!
-				string EngineSourcePath = Path.Combine(ProjectFileGenerator.EngineRelativePath, "Source");
-				string HTML5TargetPlatformFile = Path.Combine(EngineSourcePath, "Developer", "HTML5", "HTML5TargetPlatform", "HTML5TargetPlatform.Build.cs");
-				if ((File.Exists(HTML5TargetPlatformFile) == false))
+				FileReference HTML5TargetPlatformFile = FileReference.Combine(UnrealBuildTool.EngineSourceDirectory, "Developer", "HTML5", "HTML5TargetPlatform", "HTML5TargetPlatform.Build.cs");
+				if (!FileReference.Exists(HTML5TargetPlatformFile))
 				{
 					bRegisterBuildPlatform = false;
 					Log.TraceWarning("Missing required components (.... HTML5TargetPlatformFile, others here...). Check source control filtering, or try resyncing.");
diff --git a/Engine/Source/Programs/UnrealBuildTool/Linux/LinuxExports.cs b/Engine/Source/Programs/UnrealBuildTool/Linux/LinuxExports.cs
--- a/Engine/Source/Programs/UnrealBuildTool/Linux/LinuxExports.cs
+++ b/Engine/Source/Programs/UnrealBuildTool/Linux/LinuxExports.cs
@@ -18,7 +18,7 @@ namespace UnrealBuildTool
 		/// <param name="TargetFileName"></param>
 		public static void StripSymbols(string SourceFileName, string TargetFileName)
 		{
-			LinuxToolChain ToolChain = new LinuxToolChain(null);
+			LinuxToolChain ToolChain = new LinuxToolChain("");
 			ToolChain.StripSymbols(SourceFileName, TargetFileName);
 		}
 	}
diff --git a/Engine/Source/Programs/UnrealBuildTool/System/PlatformExports.cs b/Engine/Source/Programs/UnrealBuildTool/System/PlatformExports.cs
--- a/Engine/Source/Programs/UnrealBuildTool/System/PlatformExports.cs
+++ b/Engine/Source/Programs/UnrealBuildTool/System/PlatformExports.cs
@@ -89,20 +89,23 @@ namespace UnrealBuildTool
 		}
 
 		/// <summary>
-		/// Register all UBT platform classes
-		/// </summary>
-		public static void RegisterClasses()
-		{
-			UnrealBuildTool.RegisterAllUBTClasses(SDKOutputLevel.Quiet, false);
-		}
-
-		/// <summary>
-		/// Sets the IsEngineInstalled() flag from an external source
+		/// Initialize UBT in the context of another host process (presumably UAT)
 		/// </summary>
 		/// <param name="bIsEngineInstalled">Whether the engine is installed</param>
-		public static void SetIsEngineInstalled(bool bIsEngineInstalled)
+		/// <returns>True if initialization was successful</returns>
+		public static bool Initialize(bool bIsEngineInstalled)
 		{
 			UnrealBuildTool.SetIsEngineInstalled(bIsEngineInstalled);
+
+			// Read the XML configuration files
+			if(!XmlConfig.ReadConfigFiles())
+			{
+				return false;
+			}
+
+			// Register all the platform classes
+			UnrealBuildTool.RegisterAllUBTClasses(SDKOutputLevel.Quiet, false);
+			return true;
 		}
 	}
 }
