Fix linux using a blank architecture rather than the default. Should fix symbol stripping from UAT.

diff --git a/Engine/Source/Programs/UnrealBuildTool/Linux/LinuxExports.cs b/Engine/Source/Programs/UnrealBuildTool/Linux/LinuxExports.cs
--- a/Engine/Source/Programs/UnrealBuildTool/Linux/LinuxExports.cs
+++ b/Engine/Source/Programs/UnrealBuildTool/Linux/LinuxExports.cs
@@ -18,7 +18,7 @@ namespace UnrealBuildTool
 		/// <param name="TargetFileName"></param>
 		public static void StripSymbols(string SourceFileName, string TargetFileName)
 		{
-			LinuxToolChain ToolChain = new LinuxToolChain("");
+			LinuxToolChain ToolChain = new LinuxToolChain(LinuxPlatform.DefaultArchitecture);
 			ToolChain.StripSymbols(SourceFileName, TargetFileName);
 		}
 	}
