Fix Android CIS errors deploying multiple APKs at once for installed builds - store the target android architectures in the deployment data, and pass them back into the toolchain when deploying.

Also fix warnings generating project files on Mono due to different heuristic for parsing XML comments.

diff --git a/Engine/Source/Programs/UnrealBuildTool/Android/UEDeployAndroid.cs b/Engine/Source/Programs/UnrealBuildTool/Android/UEDeployAndroid.cs
--- a/Engine/Source/Programs/UnrealBuildTool/Android/UEDeployAndroid.cs
+++ b/Engine/Source/Programs/UnrealBuildTool/Android/UEDeployAndroid.cs
@@ -2398,7 +2398,7 @@ namespace UnrealBuildTool
 		public override bool PrepTargetForDeployment(UEBuildDeployTarget InTarget)
 		{
 			//Log.TraceInformation("$$$$$$$$$$$$$$ PrepTargetForDeployment $$$$$$$$$$$$$$$$$ {0}", InTarget.TargetName);
-			AndroidToolChain ToolChain = new AndroidToolChain(InTarget.ProjectFile, false, null, null); 
+			AndroidToolChain ToolChain = new AndroidToolChain(InTarget.ProjectFile, false, InTarget.AndroidArchitectures, InTarget.AndroidGPUArchitectures); 
 
 			// we need to strip architecture from any of the output paths
 			string BaseSoName = ToolChain.RemoveArchName(InTarget.OutputPaths[0].FullName);
diff --git a/Engine/Source/Programs/UnrealBuildTool/Configuration/UEBuildDeployTarget.cs b/Engine/Source/Programs/UnrealBuildTool/Configuration/UEBuildDeployTarget.cs
--- a/Engine/Source/Programs/UnrealBuildTool/Configuration/UEBuildDeployTarget.cs
+++ b/Engine/Source/Programs/UnrealBuildTool/Configuration/UEBuildDeployTarget.cs
@@ -92,6 +92,16 @@ namespace UnrealBuildTool
 		/// </summary>
 		public readonly bool bCreateStubIPA;
 
+		/// <summary>
+		/// Which architectures to deploy for on Android
+		/// </summary>
+		public readonly string[] AndroidArchitectures;
+
+		/// <summary>
+		/// Which GPU architectures to deploy for on Android
+		/// </summary>
+		public readonly string[] AndroidGPUArchitectures;
+
 		/// <summary>
 		/// Construct the deployment info from a target
 		/// </summary>
@@ -111,6 +121,8 @@ namespace UnrealBuildTool
 			this.BuildReceiptFileName = Target.BuildReceiptFileName;
 			this.bOutputToEngineBinaries = Target.Rules.bOutputToEngineBinaries;
 			this.bCreateStubIPA = Target.Rules.bCreateStubIPA;
+			this.AndroidArchitectures = Target.Rules.AndroidPlatform.Architectures.ToArray();
+			this.AndroidGPUArchitectures = Target.Rules.AndroidPlatform.GPUArchitectures.ToArray();
 		}
 
 		/// <summary>
@@ -134,6 +146,8 @@ namespace UnrealBuildTool
 				BuildReceiptFileName = Reader.ReadString();
 				bOutputToEngineBinaries = Reader.ReadBoolean();
 				bCreateStubIPA = Reader.ReadBoolean();
+				AndroidArchitectures = Reader.ReadArray(x => x.ReadString());
+				AndroidGPUArchitectures = Reader.ReadArray(x => x.ReadString());
 			}
 		}
 
@@ -159,6 +173,8 @@ namespace UnrealBuildTool
 				Writer.Write(BuildReceiptFileName);
 				Writer.Write(bOutputToEngineBinaries);
 				Writer.Write(bCreateStubIPA);
+				Writer.Write(AndroidArchitectures, (w, e) => w.Write(e));
+				Writer.Write(AndroidGPUArchitectures, (w, e) => w.Write(e));
 			}
 		}
 	}
diff --git a/Engine/Source/Programs/UnrealBuildTool/IOS/IOSToolChain.cs b/Engine/Source/Programs/UnrealBuildTool/IOS/IOSToolChain.cs
--- a/Engine/Source/Programs/UnrealBuildTool/IOS/IOSToolChain.cs
+++ b/Engine/Source/Programs/UnrealBuildTool/IOS/IOSToolChain.cs
@@ -86,12 +86,12 @@ namespace UnrealBuildTool
 			Settings = new Lazy<IOSToolChainSettings>(InCreateSettings);
 		}
 
-		/***********************************************************************
-		 * NOTE:
-		 *  Do NOT change the defaults to set your values, instead you should set the environment variables
-		 *  properly in your system, as other tools make use of them to work properly!
-		 *  The defaults are there simply for examples so you know what to put in your env vars...
-		 ***********************************************************************/
+		// ***********************************************************************
+		// * NOTE:
+		// *  Do NOT change the defaults to set your values, instead you should set the environment variables
+		// *  properly in your system, as other tools make use of them to work properly!
+		// *  The defaults are there simply for examples so you know what to put in your env vars...
+		// ***********************************************************************
 
 		// If you are looking for where to change the remote compile server name, look in RemoteToolChain.cs
 
