Fix null reference exception in UBT when using makefiles.

diff --git a/Engine/Source/Programs/UnrealBuildTool/System/ExternalExecution.cs b/Engine/Source/Programs/UnrealBuildTool/System/ExternalExecution.cs
--- a/Engine/Source/Programs/UnrealBuildTool/System/ExternalExecution.cs
+++ b/Engine/Source/Programs/UnrealBuildTool/System/ExternalExecution.cs
@@ -1032,7 +1032,7 @@ namespace UnrealBuildTool
                         CmdLine += " -UseMakefiles";
                     }
 
-					if (!Target.Rules.bCompileAgainstEngine)
+					if (Target.Rules != null && !Target.Rules.bCompileAgainstEngine)
 					{
 						CmdLine += " -NoEnginePlugins";
 					}
