Fix CIS failure staging on Linux due to assumption about CWD when initializing build platforms.

diff --git a/Engine/Source/Programs/UnrealBuildTool/Linux/UEBuildLinux.cs b/Engine/Source/Programs/UnrealBuildTool/Linux/UEBuildLinux.cs
--- a/Engine/Source/Programs/UnrealBuildTool/Linux/UEBuildLinux.cs
+++ b/Engine/Source/Programs/UnrealBuildTool/Linux/UEBuildLinux.cs
@@ -467,17 +467,8 @@ namespace UnrealBuildTool
 
 			if ((ProjectFileGenerator.bGenerateProjectFiles == true) || (SDK.HasRequiredSDKsInstalled() == SDKStatus.Valid))
 			{
-				bool bRegisterBuildPlatform = true;
-
-				string EngineSourcePath = Path.Combine(ProjectFileGenerator.RootRelativePath, "Engine", "Source");
-				string LinuxTargetPlatformFile = Path.Combine(EngineSourcePath, "Developer", "Linux", "LinuxTargetPlatform", "LinuxTargetPlatform.Build.cs");
-
-				if (File.Exists(LinuxTargetPlatformFile) == false)
-				{
-					bRegisterBuildPlatform = false;
-				}
-
-				if (bRegisterBuildPlatform == true)
+				FileReference LinuxTargetPlatformFile = FileReference.Combine(UnrealBuildTool.EngineSourceDirectory, "Developer", "Linux", "LinuxTargetPlatform", "LinuxTargetPlatform.Build.cs");
+				if (FileReference.Exists(LinuxTargetPlatformFile))
 				{
 					// Register this build platform for Linux
 					Log.TraceVerbose("        Registering for {0}", UnrealTargetPlatform.Linux.ToString());
