Make it easier to reproduce non-unity/non-pch errors when building locally - compile all modules by default for the UE4Editor target, and disable PCHs when doing a single file compile (currently assumes PCH is already built).

diff --git a/Engine/Source/Programs/UnrealBuildTool/Configuration/UEBuildTarget.cs b/Engine/Source/Programs/UnrealBuildTool/Configuration/UEBuildTarget.cs
--- a/Engine/Source/Programs/UnrealBuildTool/Configuration/UEBuildTarget.cs
+++ b/Engine/Source/Programs/UnrealBuildTool/Configuration/UEBuildTarget.cs
@@ -793,6 +793,7 @@ namespace UnrealBuildTool
 			{
 				RulesObject.bUseUnityBuild = false;
 				RulesObject.bForceUnityBuild = false;
+				RulesObject.bUsePCHFiles = false;
 				RulesObject.bDisableLinking = true;
 			}
 
diff --git a/Engine/Source/UE4Editor.Target.cs b/Engine/Source/UE4Editor.Target.cs
--- a/Engine/Source/UE4Editor.Target.cs
+++ b/Engine/Source/UE4Editor.Target.cs
@@ -10,6 +10,7 @@ public class UE4EditorTarget : TargetRules
 		Type = TargetType.Editor;
 		BuildEnvironment = TargetBuildEnvironment.Shared;
 		bBuildAllPlugins = true;
+		bPrecompile = true;
 		ExtraModuleNames.Add("UE4Game");
 	}
 
