UBT: Fix BuildConfiguration.xml settings not being applied to the target.

diff --git a/Engine/Source/Programs/UnrealBuildTool/Configuration/TargetRules.cs b/Engine/Source/Programs/UnrealBuildTool/Configuration/TargetRules.cs
--- a/Engine/Source/Programs/UnrealBuildTool/Configuration/TargetRules.cs
+++ b/Engine/Source/Programs/UnrealBuildTool/Configuration/TargetRules.cs
@@ -952,6 +952,9 @@ namespace UnrealBuildTool
 				ConfigCache.ReadSettings(DirectoryReference.FromFile(ProjectFile), Platform, ConfigurableObject);
 			}
 
+			// Read settings from the XML config files
+			XmlConfig.ApplyTo(this);
+
 			// Allow the build platform to set defaults for this target
 			if(Platform != UnrealTargetPlatform.Unknown)
 			{
