UBT: Fix missing BuildConfiguration option to allow disabling dSYM files on Mac for faster iteration.

diff --git a/Engine/Source/Programs/UnrealBuildTool/Configuration/TargetRules.cs b/Engine/Source/Programs/UnrealBuildTool/Configuration/TargetRules.cs
--- a/Engine/Source/Programs/UnrealBuildTool/Configuration/TargetRules.cs
+++ b/Engine/Source/Programs/UnrealBuildTool/Configuration/TargetRules.cs
@@ -917,6 +917,11 @@ namespace UnrealBuildTool
 		/// </summary>
 		public AndroidTargetRules AndroidPlatform = new AndroidTargetRules();
 
+		/// <summary>
+		/// Mac-specific target settings
+		/// </summary>
+		public MacTargetRules MacPlatform = new MacTargetRules();
+
 		/// <summary>
 		/// Windows-specific target settings
 		/// </summary>
@@ -1321,6 +1326,7 @@ namespace UnrealBuildTool
 		{
 			this.Inner = Inner;
 			AndroidPlatform = new ReadOnlyAndroidTargetRules(Inner.AndroidPlatform);
+			MacPlatform = new ReadOnlyMacTargetRules(Inner.MacPlatform);
 			WindowsPlatform = new ReadOnlyWindowsTargetRules(Inner.WindowsPlatform);
 		}
 
@@ -1897,6 +1903,12 @@ namespace UnrealBuildTool
 			private set;
 		}
 
+		public ReadOnlyMacTargetRules MacPlatform
+		{
+			get;
+			private set;
+		}
+
 		public ReadOnlyWindowsTargetRules WindowsPlatform
 		{
 			get;
diff --git a/Engine/Source/Programs/UnrealBuildTool/Mac/UEBuildMac.cs b/Engine/Source/Programs/UnrealBuildTool/Mac/UEBuildMac.cs
--- a/Engine/Source/Programs/UnrealBuildTool/Mac/UEBuildMac.cs
+++ b/Engine/Source/Programs/UnrealBuildTool/Mac/UEBuildMac.cs
@@ -9,6 +9,61 @@ using System.IO;
 
 namespace UnrealBuildTool
 {
+	/// <summary>
+	/// Mac-specific target settings
+	/// </summary>
+	public class MacTargetRules
+	{
+		/// <summary>
+		/// Whether to generate dSYM files
+		/// Lists Architectures that you want to build
+		/// </summary>
+		[XmlConfigFile(Category = "BuildConfiguration", Name = "bGeneratedSYMFile")]
+		public bool bGenerateDsymFile = true;
+
+		/// <summary>
+		/// Constructor
+		/// </summary>
+		public MacTargetRules()
+		{
+			XmlConfig.ApplyTo(this);
+		}
+	}
+
+	/// <summary>
+	/// Read-only wrapper for Mac-specific target settings
+	/// </summary>
+	public class ReadOnlyMacTargetRules
+	{
+		/// <summary>
+		/// The private mutable settings object
+		/// </summary>
+		private MacTargetRules Inner;
+
+		/// <summary>
+		/// Constructor
+		/// </summary>
+		/// <param name="Inner">The settings object to wrap</param>
+		public ReadOnlyMacTargetRules(MacTargetRules Inner)
+		{
+			this.Inner = Inner;
+		}
+
+		/// <summary>
+		/// Accessors for fields on the inner TargetRules instance
+		/// </summary>
+		#region Read-only accessor properties 
+		#pragma warning disable CS1591
+
+		public bool bGenerateDsymFile
+		{
+			get { return Inner.bGenerateDsymFile; }
+		}
+
+		#pragma warning restore CS1591
+		#endregion
+	}
+
 	class MacPlatformContext : UEBuildPlatformContext
 	{
 		FileReference ProjectFile;
@@ -157,8 +212,7 @@ namespace UnrealBuildTool
 				Target.bCompileLeanAndMeanUE = false;
 			}
 
-			bool bGeneratedSYMFile = true; // Always true on Mac
-			Target.bUsePDBFiles = !Target.bDisableDebugInfo && Target.Configuration != UnrealTargetConfiguration.Debug && Platform == UnrealTargetPlatform.Mac && bGeneratedSYMFile;
+			Target.bUsePDBFiles = !Target.bDisableDebugInfo && Target.Configuration != UnrealTargetConfiguration.Debug && Platform == UnrealTargetPlatform.Mac && Target.MacPlatform.bGenerateDsymFile;
 
 			// we always deploy - the build machines need to be able to copy the files back, which needs the full bundle
 			Target.bDeployAfterCompile = true;
