UBT: Fix some BuildConfiguration.xml settings for executors no longer being read. Only public fields were being checked for having the XmlConfigFile attribute, and constructors for executors were not applying the XML configuration to themselves.

diff --git a/Engine/Source/Programs/UnrealBuildTool/System/Distcc.cs b/Engine/Source/Programs/UnrealBuildTool/System/Distcc.cs
--- a/Engine/Source/Programs/UnrealBuildTool/System/Distcc.cs
+++ b/Engine/Source/Programs/UnrealBuildTool/System/Distcc.cs
@@ -45,6 +45,8 @@ namespace UnrealBuildTool
 
 		public Distcc()
 		{
+			XmlConfig.ApplyTo(this);
+
 			// The default for normal Mac users should be to use DistCode which installs as an Xcode plugin and provides dynamic host management
 			if (BuildHostPlatform.Current.Platform == UnrealTargetPlatform.Mac)
 			{
diff --git a/Engine/Source/Programs/UnrealBuildTool/System/LocalExecutor.cs b/Engine/Source/Programs/UnrealBuildTool/System/LocalExecutor.cs
--- a/Engine/Source/Programs/UnrealBuildTool/System/LocalExecutor.cs
+++ b/Engine/Source/Programs/UnrealBuildTool/System/LocalExecutor.cs
@@ -332,6 +332,11 @@ namespace UnrealBuildTool
 		[XmlConfigFile(Category = "BuildConfiguration")]
 		int MaxProcessorCount = int.MaxValue;
 
+		public LocalExecutor()
+		{
+			XmlConfig.ApplyTo(this);
+		}
+
 		public override string Name
 		{
 			get { return "Local"; }
diff --git a/Engine/Source/Programs/UnrealBuildTool/System/SNDBS.cs b/Engine/Source/Programs/UnrealBuildTool/System/SNDBS.cs
--- a/Engine/Source/Programs/UnrealBuildTool/System/SNDBS.cs
+++ b/Engine/Source/Programs/UnrealBuildTool/System/SNDBS.cs
@@ -40,6 +40,7 @@ namespace UnrealBuildTool
 
 		public SNDBS()
 		{
+			XmlConfig.ApplyTo(this);
 		}
 
 		public override string Name
diff --git a/Engine/Source/Programs/UnrealBuildTool/System/XGE.cs b/Engine/Source/Programs/UnrealBuildTool/System/XGE.cs
--- a/Engine/Source/Programs/UnrealBuildTool/System/XGE.cs
+++ b/Engine/Source/Programs/UnrealBuildTool/System/XGE.cs
@@ -36,6 +36,11 @@ namespace UnrealBuildTool
 
 		private const string ProgressMarkupPrefix = "@action";
 
+		public XGE()
+		{
+			XmlConfig.ApplyTo(this);
+		}
+
 		public override string Name
 		{
 			get { return "XGE"; }
diff --git a/Engine/Source/Programs/UnrealBuildTool/System/XmlConfig.cs b/Engine/Source/Programs/UnrealBuildTool/System/XmlConfig.cs
--- a/Engine/Source/Programs/UnrealBuildTool/System/XmlConfig.cs
+++ b/Engine/Source/Programs/UnrealBuildTool/System/XmlConfig.cs
@@ -122,7 +122,7 @@ namespace UnrealBuildTool
 		/// <returns>True if the type has a field with the XmlConfigFile attribute</returns>
 		static bool HasXmlConfigFileAttribute(Type Type)
 		{
-			foreach(FieldInfo Field in Type.GetFields())
+			foreach(FieldInfo Field in Type.GetFields(BindingFlags.Instance | BindingFlags.Static | BindingFlags.Public | BindingFlags.NonPublic))
 			{
 				foreach(CustomAttributeData CustomAttribute in Field.CustomAttributes)
 				{
@@ -290,7 +290,7 @@ namespace UnrealBuildTool
 		{
 			foreach(Type ConfigType in ConfigTypes)
 			{
-				foreach(FieldInfo FieldInfo in ConfigType.GetFields(BindingFlags.Instance | BindingFlags.Static | BindingFlags.GetField | BindingFlags.Public))
+				foreach(FieldInfo FieldInfo in ConfigType.GetFields(BindingFlags.Instance | BindingFlags.Static | BindingFlags.GetField | BindingFlags.Public | BindingFlags.NonPublic))
 				{
 					IEnumerable<XmlConfigFileAttribute> Attributes = FieldInfo.GetCustomAttributes<XmlConfigFileAttribute>();
 					foreach(XmlConfigFileAttribute Attribute in Attributes)
