Copying //UE4/Dev-Platform to //UE4/Dev-Main (Source: //UE4/Dev-Platform @ 3295257)

==========================
MAJOR FEATURES + CHANGES
==========================

Change 3235199 on 2016/12/14 by [REDACTED_PERSON]

	Fix new compile error for missing #define

Change 3235340 on 2016/12/14 by [REDACTED_PERSON]

	Linux: refactor of some PlatformMisc functions.

	- RootDir() removed since it was a no-op.
	- Old EngineDir() implementation removed in favor of more generic one that should handle foreign engine dir.
	- Change by CengizT,

Change 3237014 on 2016/12/15 by [REDACTED_PERSON]

	Fixed a crash in FChunkCacheWorker constructor

Change 3238305 on 2016/12/16 by [REDACTED_PERSON]

	- Added a None option to the FKey customization, unless the FKey property had NoClear on it

Change 3240823 on 2016/12/20 by [REDACTED_PERSON]

	Device profiles don't work for iPadPro 9.7 and 12.9

Change 3241103 on 2016/12/20 by [REDACTED_PERSON]

	Android support from Visual Studio
	#android

Change 3241357 on 2016/12/20 by [REDACTED_PERSON]

	Add gameActivityOnNewIntentAddtions section to Android UPL
	#PR #2969
	#ue4
	#android

Change 3241941 on 2016/12/21 by [REDACTED_PERSON]

	Build Fix

Change 3249832 on 2017/01/06 by [REDACTED_PERSON]

	refetch on timed out GET/POST requests

Change 3249837 on 2017/01/06 by [REDACTED_PERSON]

	black box issues fixed:

	use device pixel ratio during width and height checks
	and use canvas dimensions if in full screen -- otherwise store SDL_window dimensions for future use

Change 3249988 on 2017/01/06 by [REDACTED_PERSON]

	- Disable the HeartBeat() function on platforms that don't actually want to use the HeartbeatThread

Change 3253720 on 2017/01/11 by [REDACTED_PERSON]

	- Added support for a config class to use a specific platform's config hierarchy, so that the editor can read NDA'd platform default settings without needing the settings to in Base*.ini
	  - See SwitchRuntimeSettings.h / GetConfigOverridePlatform()
	- Addiontally made it so that NDAd platforms are saved to Project/Platform/Platform*.ini, instead of Project/Default*.ini (to keep samples .ini files free of NDAd platform settings).
	  - See UObject::GetDefaultConfigFilename()
	- Updated some minor ShooterGame switch settings while cleaning this up

Change 3254162 on 2017/01/11 by [REDACTED_PERSON]

	Avoid trying to load empty package names.
	Fixed issue with iterative ini files being unparseable if they inlcude a colon in them.
	#test Cook QAGame

Change 3255309 on 2017/01/12 by [REDACTED_PERSON]

	In the derived datacache commandlet wait for texture building to finish before we GC.
	#test DDC QAGame

Change 3255311 on 2017/01/12 by [REDACTED_PERSON]

	Removed debug logging for shader compilation.
	Issue hasn't occured in a while and the logging is annoying.
	#test Cook QAGame

Change 3257024 on 2017/01/13 by [REDACTED_PERSON]

	- Reread in the target RHIs array every time the list of shader types is needed, instead of caching, because the user could change the settings in the editor, then click cook.

Change 3259636 on 2017/01/16 by [REDACTED_PERSON]

	- Fixed split screen render issue with payer 2 getting no geometry

Change 3260159 on 2017/01/17 by [REDACTED_PERSON]

	Added extra logging when deleting a directory fails during ReconcileWorkspace.

Change 3260300 on 2017/01/17 by [REDACTED_PERSON]

	More logging for cleaning workspaces.

Change 3261056 on 2017/01/17 by [REDACTED_PERSON]

	Cook on the fly builds now resolve string asset references.
	#test Trivial

Change 3262803 on 2017/01/18 by [REDACTED_PERSON]

	Added missing support for compiling plugins external to Engine/Plugins & Game/Plugins

Change 3262852 on 2017/01/18 by [REDACTED_PERSON]

	Fixed the bad robomerge
	Don't try to regenerate projects when adding a content only plugin to a content only project

Change 3264930 on 2017/01/19 by [REDACTED_PERSON]

	#include some header files needed when building UFE.

Change 3265728 on 2017/01/20 by [REDACTED_PERSON]

	PlatformShowcase - Added TestBed_MobileFeatures .umap and related test content.

Change 3267188 on 2017/01/21 by [REDACTED_PERSON]

	Merging //UE4/Dev-Main to Dev-Platform (//UE4/Dev-Platform)

Change 3267439 on 2017/01/22 by [REDACTED_PERSON]

	Fix Dev-Platform build.

	- Fixed just to have it compile; perhaps a proper fix is needed.
	- Seems to be caused by CL 3265587 (delegate was changed to return an array of search results instead of a single one).

Change 3267556 on 2017/01/23 by [REDACTED_PERSON]

	Linux: fix MoveFile to work across file systems.

	- PR #3141 with slight changes.

Change 3267843 on 2017/01/23 by [REDACTED_PERSON]

	Remove name collision (macro vs libc++).

	- Redoing CL 3259310.

Change 3267850 on 2017/01/23 by [REDACTED_PERSON]

	Fix wrong always true condition.

	- PLATFORM_LINUX is always defined, but can be 0.

Change 3268048 on 2017/01/23 by [REDACTED_PERSON]

	Integrated fix for rebuild lighting commandlet from Greg Korman @ Impulse Gear.
	#test Rebuild lighting Paragon

Change 3268403 on 2017/01/23 by [REDACTED_PERSON]

	#BUILDUPGRADENOTES
	- Moved XboxOne and PS4 settings into platform specific .ini files (after using GetConfigOverridePlatform() in their class delcarations)
	- Licensee games that have PS4, XboxOne, Switch settings in DefaultEngine.ini will have those settings saved in the platform version next time the project settings are edited. DOCUMENT THIS!

Change 3272441 on 2017/01/25 by [REDACTED_PERSON]

	Fixed documentation error in UnrealPluginLanguage
	#ue4
	#android

Change 3272478 on 2017/01/25 by [REDACTED_PERSON]

	Fix another documentation error in UnrealPluginLanguage
	#ue4

Change 3272826 on 2017/01/25 by [REDACTED_PERSON]

	Google Cloud Messaging plugin for Android
	#ue4
	#android

Change 3272839 on 2017/01/25 by [REDACTED_PERSON]

	Fix name of Google Cloud Messaging Sender ID
	#ue4
	#android

Change 3273837 on 2017/01/26 by [REDACTED_PERSON]

	Added check to ensure editor never saves source texture data which has had ReleaseSourceMemory called on it.
	Instead crash as this is a loss of content situation.
	#test Cook paragon cook qagame

Change 3274122 on 2017/01/26 by [REDACTED_PERSON]

	Runtime permissions support on Android
	- Removing certain permissions
	#android

Change 3274311 on 2017/01/26 by [REDACTED_PERSON]

	Merging //UE4/Dev-Main to Dev-Platform (//UE4/Dev-Platform)

Change 3274794 on 2017/01/27 by [REDACTED_PERSON]

	Linux: fix installed SDK check (UE-40392).

	- Pull request #3111 by rubu.

Change 3274803 on 2017/01/27 by [REDACTED_PERSON]

	Linux: added few more exceptions to .gitignore (UE-39612).

	- Pull request #3026 by ardneran.

Change 3276247 on 2017/01/27 by [REDACTED_PERSON]

	HTML5 HeapSize settings - make use of it from UE4 Editor:Platforms:HTML5:Memory:HeapSize

	note: emscripten says this is really no longer needed when using [ -s ALLOW_MEMORY_GROWTH=1 ] -- but tests have shown when using that, the game load/compile times takes longer

Change 3276347 on 2017/01/27 by dan.reynolds

	Android Streaming Test Content

Change 3276682 on 2017/01/29 by [REDACTED_PERSON]

	HTML5 thirdparty build scripts

	- fix up what looks like a bad merge
	- allow linux to also build these  libs
	- fixed harfbuzz to use freetype2-2.6 when building HTML5 libs

	- tested on mac, linux, and windows (git-bash)

Change 3276796 on 2017/01/29 by [REDACTED_PERSON]

	HTML5 thirdparty (python) build scripts

	- linux patches from mozilla's jukka

	- tested on mac and, linux, and windows (git-bash)

	part of:

Change 3276803 on 2017/01/29 by [REDACTED_PERSON]

	HTML5 thirdparty build scripts

	- getting ready to build with (new toolchain that has) wasm support
	- minor fix to handle whitespace in project path

	- tested on mac and, linux, and windows (git-bash)

	part of:

Change 3278007 on 2017/01/30 by [REDACTED_PERSON]

	SteamVR: whitelist for Linux.

	- Makes Blueprint functions available in Linux builds, even if stubbed.
	- Can be probably whitelisted for Mac too.

Change 3278172 on 2017/01/30 by [REDACTED_PERSON]

	Do not rebuild UnrealPak locally (UE-41285).

Change 3279873 on 2017/01/31 by [REDACTED_PERSON]

	 + Implement streaming in Vorbis
	 + Add streaming to Android audio
	 + Fix audio streaming chunk race condition

Change 3280063 on 2017/01/31 by [REDACTED_PERSON]

	GitHub 2949 : Fix for crashes when backgrounding/sleeping on iOS metal devices

	#2949

Change 3280072 on 2017/01/31 by [REDACTED_PERSON]

	PR #2889: Add -distribution when iOS distribution Packaging. with IPhonePackage.exe (Contributed by [REDACTED_PERSON])
	https://github.com/EpicGames/UnrealEngine/pull/2889

	#2889

Change 3280091 on 2017/01/31 by [REDACTED_PERSON]

	Linux: fix "unable to make writable" toast (UE-37228).

	- Also fixed other platforms that returned inverted the error result.

Change 3280624 on 2017/01/31 by [REDACTED_PERSON]

	PR #2891: iOS IDFV string allocation fix (Contributed by [REDACTED_PERSON])
	https://github.com/EpicGames/UnrealEngine/pull/2891

	#2891

Change 3280625 on 2017/01/31 by [REDACTED_PERSON]

	GitHub 2576 - Fix UIImagePickerController crash

	#2576

Change 3281618 on 2017/02/01 by [REDACTED_PERSON]

	- Fixed hopeful compile error with missing inlcude

Change 3282277 on 2017/02/01 by [REDACTED_PERSON]

	- Support 0.12.16 and 1.1.1 (the first versions that can share Oasis)

Change 3282441 on 2017/02/01 by [REDACTED_PERSON]

	Fix Linux editor splash screen (UE-28123).

Change 3282580 on 2017/02/01 by [REDACTED_PERSON]

	HTML5 - fix "firefox nighly" issue with:

	failed to compile wasm module: CompileError: at offset XXX: initial memory size too big:

	WARNING: this greatly impacts (in browser) compile times

Change 3285991 on 2017/02/03 by [REDACTED_PERSON]

	Fix executable path for stripping Android debug symbols (handle non-Windows properly)
	#ue4
	#android

Change 3286406 on 2017/02/03 by [REDACTED_PERSON]

	Save and restore texture filtering for movie playback in all cases
	#ue4
	#android

Change 3286800 on 2017/02/04 by [REDACTED_PERSON]

	Fix executable path for stripping Android debug symbols (handle non-Windows properly)
	#ue4
	#android

Change 3288598 on 2017/02/06 by [REDACTED_PERSON]

	CodeLite fixes.

	- Use *-Linux-Debug binary for Debug configuration.
	- Fix virtual paths.

Change 3288864 on 2017/02/06 by [REDACTED_PERSON]

	Merging //UE4/Dev-Main to Dev-Platform (//UE4/Dev-Platform)
	- Note, Switch is known to not boot with this, fix coming next

Change 3289364 on 2017/02/06 by [REDACTED_PERSON]

	[BUILDUPGRADENOTES] - Fixed the "type" of the desktop device profiles to be Windows, not WindowsNoEditor, etc. It should be the platform, not a random string.
	- Updated how DeviceProfiles are loaded, especially in the editor, so that we can have NDAd platforms have their default DP values in platform-hidden files
	- This makes use of the ability for a class to override the platform hierarchy in the editor (like we do with other editor-exposed platform objects)
	- Added Config/[PS4|XboxOne|Switch]/ConfidentialPlatform.ini files so that the DP loading code knows to look in their directories for DPs. See FGenericPlatformMisc::GetConfidentialPlatforms() for more information
	- Note that saving still saves the entire DP to the .ini. Next DP change is to have them properly save against their 2(!) parents - the .ini file earlier in the hierarchy, and the parent DP object. Makes it tricky, for sure.
	- Added FConfigFile::GetArray (previous was only on FConfigCacheIni)

Change 3289796 on 2017/02/07 by [REDACTED_PERSON]

	Linux: remove leftover CEF build script.

Change 3289872 on 2017/02/07 by [REDACTED_PERSON]

	Linux: install MIME types (UE-40954).

	- Pull request #3154 by RicardoEPRodrigues.

Change 3289915 on 2017/02/07 by [REDACTED_PERSON]

	- Fixed CIS warnings

Change 3289916 on 2017/02/07 by [REDACTED_PERSON]

	Linux: remove -opengl4 from the default invocation.

Change 3290009 on 2017/02/07 by [REDACTED_PERSON]

	UE4 - Fixed boot time EDL causing some issues even when it wasn't being used.

Change 3290120 on 2017/02/07 by [REDACTED_PERSON]

	Merging //UE4/Dev-Main to Dev-Platform (//UE4/Dev-Platform)

Change 3290948 on 2017/02/07 by [REDACTED_PERSON]

	Linux: fix crash when clicking on question mark (UE-41634).

	- Symbol interposition problem (proper fix is still to be investigated).

	(Edigrating part of CL 3290683 from Release-4.15 to Dev-Platform)

Change 3291074 on 2017/02/07 by [REDACTED_PERSON]

	Speculative build fix.

Change 3292028 on 2017/02/08 by [REDACTED_PERSON]

	- Fixed Incremental CIS build failures

Change 3292105 on 2017/02/08 by [REDACTED_PERSON]

	emcc.py - change warning to info

Change 3292201 on 2017/02/08 by JohnHenry.Carawon

	Change comment to fix XML warning when generating project files on Linux

Change 3292242 on 2017/02/08 by [REDACTED_PERSON]

	Linux: avoid unnecessary dependency on CEF (UE-41634).

	- Do not apply CEF workaround to monolithic builds (eg. stock Game/Server targets).
	- Also disable CEF compilation for ShaderCompileWorker.
	- Based on CL 3292077 in 4.15.

Change 3292559 on 2017/02/08 by [REDACTED_PERSON]

	- Added more platforms to disable the file handle caching (all the ones that use MANAGED_FILE_HANDLES)

Change 3294333 on 2017/02/09 by [REDACTED_PERSON]

	Merging //UE4/Dev-Main to Dev-Platform (//UE4/Dev-Platform)

Change 3294506 on 2017/02/09 by [REDACTED_PERSON]

	- Fixed GoogleCloudMessaging.uplugin to fix the Installed flag. Every other plugin had false, this one had true, which caused various checks to go haywire

Change 3294984 on 2017/02/09 by [REDACTED_PERSON]

	- Worked around the remote compiling issue with code-based projects on a different drive than the engine

Change 3295056 on 2017/02/09 by [REDACTED_PERSON]

	- Fixed the remote compiling issue by unconverting the path back to host when reading from the module filename

Change 3295161 on 2017/02/09 by [REDACTED_PERSON]

	- Fixed new bug when buildin native ios that was caused by a remote compile break

Change 3295229 on 2017/02/09 by [REDACTED_PERSON]

	- Fixed a crash in clothing on platforms that don't support clothing

diff --git a/Engine/Build/Linux/uproject.xml b/Engine/Build/Linux/uproject.xml
new file mode 100644
--- /dev/null
+++ b/Engine/Build/Linux/uproject.xml
@@ -0,0 +1,10 @@
+<?xml version="1.0" encoding="utf-8"?>
+<mime-info xmlns="http://www.freedesktop.org/standards/shared-mime-info" >
+    <mime-type type="application/uproject">
+        <comment>Unreal Engine Project</comment>
+        <glob-deleteall/>
+        <glob pattern="*.uproject"/>
+	    <icon name="ue4editor"/>
+    </mime-type>
+</mime-info>
+
diff --git a/Engine/Plugins/Runtime/GoogleCloudMessaging/GoogleCloudMessaging.uplugin b/Engine/Plugins/Runtime/GoogleCloudMessaging/GoogleCloudMessaging.uplugin
new file mode 100644
--- /dev/null
+++ b/Engine/Plugins/Runtime/GoogleCloudMessaging/GoogleCloudMessaging.uplugin
@@ -0,0 +1,21 @@
+{
+	"FileVersion" : 3,
+	"Version" : 1,
+	"VersionName" : "1.0",
+	"FriendlyName" : "Google Cloud Messaging",
+	"Description" : "Support for remote notifications using Google Cloud Messaging",
+	"Category" : "Notifications",
+	"CreatedBy" : "Epic Games, Inc.",
+	"CreatedByURL" : "http://epicgames.com",
+	"EnabledByDefault" : true,
+    "Installed" : false,
+	"Modules" :
+	[
+        {
+            "Name": "GoogleCloudMessaging",
+            "Type": "Runtime",
+            "LoadingPhase": "PreDefault",
+            "WhitelistPlatforms": [ "Android" ]
+        }
+	]
+}
\ No newline at end of file
diff --git a/Engine/Plugins/Runtime/GoogleCloudMessaging/Source/GoogleCloudMessaging/GoogleCloudMessaging.Build.cs b/Engine/Plugins/Runtime/GoogleCloudMessaging/Source/GoogleCloudMessaging/GoogleCloudMessaging.Build.cs
new file mode 100644
--- /dev/null
+++ b/Engine/Plugins/Runtime/GoogleCloudMessaging/Source/GoogleCloudMessaging/GoogleCloudMessaging.Build.cs
@@ -0,0 +1,38 @@
+// Copyright 1998-2016 Epic Games, Inc. All Rights Reserved.
+
+using System.IO;
+
+namespace UnrealBuildTool.Rules
+{
+	public class GoogleCloudMessaging : ModuleRules
+	{
+		public GoogleCloudMessaging(ReadOnlyTargetRules Target) : base(Target)
+		{
+			PrivateIncludePaths.AddRange(new string[]
+				{
+					"GoogleCloudMessaging/Private",
+				});
+
+			PrivateDependencyModuleNames.AddRange(new string[]
+				{
+					"Core",
+					"CoreUObject",
+					"Engine",
+					"Launch"
+				});
+
+			if (UEBuildConfiguration.bBuildEditor)
+			{
+				PrivateDependencyModuleNames.Add("UnrealEd");
+			}
+
+			if (Target.Platform == UnrealTargetPlatform.Android)
+			{
+				string PluginPath = Utils.MakePathRelativeTo(ModuleDirectory, BuildConfiguration.RelativeEnginePath);
+				AdditionalPropertiesForReceipt.Add(new ReceiptProperty("AndroidPlugin", Path.Combine(PluginPath, "GoogleCloudMessaging_UPL.xml")));
+			}
+
+			PublicIncludePathModuleNames.Add("Launch");
+		}
+	}
+}
diff --git a/Engine/Plugins/Runtime/GoogleCloudMessaging/Source/GoogleCloudMessaging/GoogleCloudMessaging_UPL.xml b/Engine/Plugins/Runtime/GoogleCloudMessaging/Source/GoogleCloudMessaging/GoogleCloudMessaging_UPL.xml
new file mode 100644
--- /dev/null
+++ b/Engine/Plugins/Runtime/GoogleCloudMessaging/Source/GoogleCloudMessaging/GoogleCloudMessaging_UPL.xml
@@ -0,0 +1,177 @@
+<?xml version="1.0" encoding="utf-8"?>
+<!-- Google Cloud Messaging plugin additions-->
+<root xmlns:android="http://schemas.android.com/apk/res/android">
+	<!-- init section is always evaluated once per architecture -->
+	<init>
+		<log text="Google Cloud Messaging init"/>
+		<!-- client sender ID -->
+		<setStringFromProperty result="GCMClientSenderID" ini="Engine" section="/Script/AndroidRuntimeSettings.AndroidRuntimeSettings" property="GCMClientSenderID" default="" />
+		<setBoolIsEqual result="GCMClientSenderInvalid" arg1="$S(GCMClientSenderID)" arg2=""/>
+		<setBoolNot result="GCMEnabled" source="$B(GCMClientSenderInvalid)"/>
+		<log text="GCMClientSenderID set: $B(GCMEnabled)"/>
+	</init>
+	
+	<AARImports>
+		<!-- only modify AAR imports if enabled -->
+		<if condition="GCMEnabled">
+			<true>
+				<insert>com.google.android.gms,play-services-gcm,9.2.0</insert>
+				<insertNewline/>
+			</true>
+		</if>
+	</AARImports>
+	
+	<!-- optional updates applied to AndroidManifest.xml -->
+	<androidManifestUpdates>
+		<!-- only modify manifest if enabled -->
+		<if condition="GCMEnabled">
+			<true>
+				<!-- get the package name from manifest -->
+				<setStringFromAttribute result="AppPackageName" tag="manifest" name="package"/>
+
+				<addPermission android:name="com.google.android.c2dm.permission.RECEIVE" />
+				<addElements tag="application">
+      <service
+        android:name="com.epicgames.ue4.RemoteNotificationsListener"
+        android:exported="false" >
+        <intent-filter>
+          <action android:name="com.google.android.c2dm.intent.RECEIVE" />
+        </intent-filter>
+      </service>
+      <service
+        android:name="com.epicgames.ue4.RemoteNotificationsInstanceIDListener"
+        android:exported="false">
+        <intent-filter>
+          <action android:name="com.google.android.gms.iid.InstanceID"/>
+        </intent-filter>
+      </service>
+      <service
+        android:name="com.epicgames.ue4.RemoteNotificationsRegistrationIntentService"
+        android:exported="false">
+      </service>
+				</addElements>
+
+				<!-- create receiver -->
+				<setElement result="GCMReceiver" value="receiver" />
+				<addAttribute tag="$GCMReceiver" name="android:name" value="com.google.android.gms.gcm.GcmReceiver" />
+				<addAttribute tag="$GCMReceiver" name="android:exported" value="true" />
+				<addAttribute tag="$GCMReceiver" name="android:permission" value="com.google.android.c2dm.permission.SEND" />
+
+				<!-- create intent filter -->
+				<setElement result="GCMReceiverIntentFilter" value="intent-filter" />
+
+				<!-- create action -->
+				<setElement result="GCMIntentFilterAction" value="action" />
+				<addAttribute tag="$GCMIntentFilterAction" name="android:name" value="com.google.android.c2dm.intent.RECEIVE" />
+
+				<!-- create category -->
+				<setElement result="GCMIntentFilterCategory" value="category" />
+				<addAttribute tag="$GCMIntentFilterCategory" name="android:name" value="$S(AppPackageName)" />
+
+				<!-- add action & category to intent-filter -->
+				<addElement tag="$GCMReceiverIntentFilter" name="GCMIntentFilterAction" />
+				<addElement tag="$GCMReceiverIntentFilter" name="GCMIntentFilterCategory" />
+
+				<!-- add intent-filter to receiver -->
+				<addElement tag="$GCMReceiver" name="GCMReceiverIntentFilter" />
+
+				<!-- add receiver to application -->
+				<addElement tag="application" name="GCMReceiver" />
+				
+			</true>
+		</if>
+	</androidManifestUpdates>
+
+	<!-- optional additions to proguard
+	<proguardAdditions>
+	</proguardAdditions>
+	-->
+	
+	<!-- optional files or directories to copy or delete from Intermediate/Android/APK before ndk-build -->
+	<prebuildCopies>
+		<if condition="GCMEnabled">
+			<true>
+				<copyDir src="$S(PluginDir)/../Java" dst="$S(BuildDir)/src/com/epicgames/ue4" />
+			</true>
+		</if>
+	</prebuildCopies>
+
+	<!-- optional files or directories to copy to Intermediate/Android/APK
+	<resourceCopies>
+	</resourceCopies>
+	-->
+
+	<!-- optional additions to the GameActivity imports in GameActivity.java
+	<gameActivityImportAdditions>
+	</gameActivityImportAdditions>
+	-->
+
+	<!-- optional additions to the GameActivity class in GameActivity.java -->
+	<gameActivityClassAdditions>
+		<if condition="GCMEnabled">
+			<true>
+				<insertValue value="public final String GCMSenderId = &quot;$S(GCMClientSenderID)&quot;;" />
+				<insertNewline />
+				<insert>
+	// Methods that notify the native code about Google Cloud Messaging interactions
+	public native void nativeGCMRegisteredForRemoteNotifications(String token);
+	public native void nativeGCMFailedToRegisterForRemoteNotifications(String errorMessage);
+	public native void nativeGCMReceivedRemoteNotification(String message);
+				</insert>
+			</true>
+		</if>
+	</gameActivityClassAdditions>
+	
+	<!-- optional additions to GameActivity onCreate metadata reading in GameActivity.java
+	<gameActivityReadMetadataAdditions>
+	</gameActivityReadMetadataAdditions>
+	-->
+	
+	<!-- optional additions to GameActivity onCreate in GameActivity.java
+	<gameActivityOnCreateAdditions>
+	</gameActivityOnCreateAdditions>
+	-->
+
+	<!-- optional additions to GameActivity onDestroy in GameActivity.java
+	<gameActivityOnDestroyAdditions>
+	</gameActivityOnDestroyAdditions>
+	-->
+	
+	<!-- optional additions to GameActivity onStart in GameActivity.java -->
+	<gameActivityOnStartAdditions>
+		<if condition="GCMEnabled">
+			<true>
+				<insert>
+		// start IntentService to register this application with Google Cloud Messaging
+		Intent intent = new Intent(_activity.getApplicationContext(), RemoteNotificationsRegistrationIntentService.class);
+		startService(intent);
+				</insert>
+			</true>
+		</if>
+	</gameActivityOnStartAdditions>
+
+	<!-- optional additions to GameActivity onStop in GameActivity.java
+	<gameActivityOnStopAdditions>
+	</gameActivityOnStopAdditions>
+	-->
+
+	<!-- optional additions to GameActivity onPause in GameActivity.java
+	<gameActivityOnPauseAdditions>
+	</gameActivityOnPauseAdditions>
+	-->
+
+	<!-- optional additions to GameActivity onResume in GameActivity.java
+	<gameActivityOnResumeAdditions>
+	</gameActivityOnResumeAdditions>
+	-->
+
+	<!-- optional additions to GameActivity onActivityResult in GameActivity.java
+	<gameActivityOnActivityResultAdditions>
+	</gameActivityOnActivityResultAdditions>
+	-->
+
+	<!-- optional libraries to load in GameActivity.java before libUE4.so
+	<soLoadLibrary>
+	</soLoadLibrary>
+	-->
+</root>
diff --git a/Engine/Plugins/Runtime/GoogleCloudMessaging/Source/GoogleCloudMessaging/Private/GoogleCloudMessaging.cpp b/Engine/Plugins/Runtime/GoogleCloudMessaging/Source/GoogleCloudMessaging/Private/GoogleCloudMessaging.cpp
new file mode 100644
--- /dev/null
+++ b/Engine/Plugins/Runtime/GoogleCloudMessaging/Source/GoogleCloudMessaging/Private/GoogleCloudMessaging.cpp
@@ -0,0 +1,34 @@
+// Copyright 1998-2016 Epic Games, Inc. All Rights Reserved.
+
+#include "GoogleCloudMessaging.h"
+
+void IGoogleCloudMessagingModuleInterface::StartupModule()
+{
+	// This code will execute after your module is loaded into memory (but after global variables are initialized, of course.)
+}
+
+
+void IGoogleCloudMessagingModuleInterface::ShutdownModule()
+{
+	// This function may be called during shutdown to clean up your module.  For modules that support dynamic reloading,
+	// we call this function before unloading the module.
+}
+
+#if WITH_EDITOR
+
+DEFINE_LOG_CATEGORY( LogGoogleCloudMessaging );
+
+class FEditorGoogleCloudMessaging : public IGoogleCloudMessagingModuleInterface
+{
+	virtual void RegisterForRemoteNotifications() override;
+};
+
+void FEditorGoogleCloudMessaging::RegisterForRemoteNotifications()
+{
+	UE_LOG( LogGoogleCloudMessaging, Warning, TEXT("Not implemented in the editor") );
+}
+
+IMPLEMENT_MODULE( FEditorGoogleCloudMessaging, GoogleCloudMessaging )
+
+#endif
+
diff --git a/Engine/Plugins/Runtime/GoogleCloudMessaging/Source/GoogleCloudMessaging/Private/GoogleCloudMessagingAndroid.cpp b/Engine/Plugins/Runtime/GoogleCloudMessaging/Source/GoogleCloudMessaging/Private/GoogleCloudMessagingAndroid.cpp
new file mode 100644
--- /dev/null
+++ b/Engine/Plugins/Runtime/GoogleCloudMessaging/Source/GoogleCloudMessaging/Private/GoogleCloudMessagingAndroid.cpp
@@ -0,0 +1,97 @@
+// Copyright 1998-2016 Epic Games, Inc. All Rights Reserved.
+
+#include "GoogleCloudMessaging.h"
+
+#if PLATFORM_ANDROID
+
+#include "Async/TaskGraphInterfaces.h"
+#include "Misc/CoreDelegates.h"
+
+#include "Android/AndroidJNI.h"
+#include "Android/AndroidApplication.h"
+#include "Android/AndroidEventManager.h"
+
+DEFINE_LOG_CATEGORY( LogGoogleCloudMessaging );
+
+class FAndroigGoogleCloudMessaging : public IGoogleCloudMessagingModuleInterface
+{
+	virtual void RegisterForRemoteNotifications() override;
+};
+
+IMPLEMENT_MODULE( FAndroigGoogleCloudMessaging, GoogleCloudMessaging )
+
+void FAndroigGoogleCloudMessaging::RegisterForRemoteNotifications()
+{
+	if( JNIEnv* Env = FAndroidApplication::GetJavaEnv() )
+	{
+		jmethodID jRegisterForRemoteNotifications = FJavaWrapper::FindMethod( Env, FJavaWrapper::GameActivityClassID, "AndroidThunkJava_RegisterForRemoteNotifications", "()V", false );
+		FJavaWrapper::CallVoidMethod( Env, FJavaWrapper::GameActivityThis, jRegisterForRemoteNotifications, false );
+	}
+}
+
+// registered for remote notifications
+extern "C" void Java_com_epicgames_ue4_GameActivity_nativeGCMRegisteredForRemoteNotifications( JNIEnv* jenv, jobject thiz, jstring jGCMToken )
+{
+	auto GCMTokenLength = jenv->GetStringUTFLength( jGCMToken );
+	const char* GCMTokenChars = jenv->GetStringUTFChars( jGCMToken, 0 );
+
+	TArray<uint8> GCMTokenBytes;
+	GCMTokenBytes.AddUninitialized( GCMTokenLength );
+	FMemory::Memcpy( GCMTokenBytes.GetData(), GCMTokenChars, GCMTokenLength * sizeof( uint8 ) );
+
+	FString GCMToken;
+	GCMToken = FString( UTF8_TO_TCHAR( GCMTokenChars ) );
+
+	jenv->ReleaseStringUTFChars( jGCMToken, GCMTokenChars );
+
+	FSimpleDelegateGraphTask::CreateAndDispatchWhenReady( FSimpleDelegateGraphTask::FDelegate::CreateLambda( [=]() {
+		UE_LOG(LogGoogleCloudMessaging, Display, TEXT("GCM Registration Token: %s"), *GCMToken);
+
+		FCoreDelegates::ApplicationRegisteredForRemoteNotificationsDelegate.Broadcast( GCMTokenBytes );
+	}), TStatId(), nullptr, ENamedThreads::GameThread );
+}
+
+// failed to register for remote notifications
+extern "C" void Java_com_epicgames_ue4_GameActivity_nativeGCMFailedToRegisterForRemoteNotifications( JNIEnv* jenv, jobject thiz, jstring jErrorMessage )
+{
+	FString GCMErrorMessage;
+	const char* GCMErrorMessageChars = jenv->GetStringUTFChars( jErrorMessage, 0 );
+	GCMErrorMessage = FString( UTF8_TO_TCHAR( GCMErrorMessageChars ) );
+	jenv->ReleaseStringUTFChars( jErrorMessage, GCMErrorMessageChars );
+
+	FSimpleDelegateGraphTask::CreateAndDispatchWhenReady( FSimpleDelegateGraphTask::FDelegate::CreateLambda( [=]() {
+		UE_LOG(LogGoogleCloudMessaging, Display, TEXT("GCM Registration Error: %s"), *GCMErrorMessage);
+
+		FCoreDelegates::ApplicationFailedToRegisterForRemoteNotificationsDelegate.Broadcast( GCMErrorMessage );
+	}), TStatId(), nullptr, ENamedThreads::GameThread );
+}
+
+// received message
+extern "C" void Java_com_epicgames_ue4_GameActivity_nativeGCMReceivedRemoteNotification( JNIEnv* jenv, jobject thiz, jstring jMessage )
+{
+	FString Message;
+	const char* MessageChars = jenv->GetStringUTFChars( jMessage, 0 );
+	Message = FString( UTF8_TO_TCHAR( MessageChars ) );
+	jenv->ReleaseStringUTFChars( jMessage, MessageChars );
+
+	int AppState = 3; // EApplicationState::Active;
+	if (FAppEventManager::GetInstance()->IsGamePaused())
+	{
+		if (FAppEventManager::GetInstance()->IsGameInFocus())
+		{
+			AppState = 1; // EApplicationState::Inactive;
+		}
+		else
+		{
+			AppState = 2; // EApplicationState::Background;
+		}
+	}
+
+	FSimpleDelegateGraphTask::CreateAndDispatchWhenReady( FSimpleDelegateGraphTask::FDelegate::CreateLambda( [=]() {
+		UE_LOG(LogGoogleCloudMessaging, Display, TEXT("GCM AppState = %d, Message : %s"), AppState, *Message);
+
+		FCoreDelegates::ApplicationReceivedRemoteNotificationDelegate.Broadcast( Message /* , AppState */ );
+	}), TStatId(), nullptr, ENamedThreads::GameThread );
+}
+
+#endif
diff --git a/Engine/Plugins/Runtime/GoogleCloudMessaging/Source/GoogleCloudMessaging/Public/GoogleCloudMessaging.h b/Engine/Plugins/Runtime/GoogleCloudMessaging/Source/GoogleCloudMessaging/Public/GoogleCloudMessaging.h
new file mode 100644
--- /dev/null
+++ b/Engine/Plugins/Runtime/GoogleCloudMessaging/Source/GoogleCloudMessaging/Public/GoogleCloudMessaging.h
@@ -0,0 +1,46 @@
+// Copyright 1998-2016 Epic Games, Inc. All Rights Reserved.
+
+#pragma once
+
+#include "CoreMinimal.h"
+#include "ModuleManager.h"
+
+DECLARE_LOG_CATEGORY_EXTERN(LogGoogleCloudMessaging, Log, All);
+
+/**
+ * The public interface to this module.  In most cases, this interface is only public to sibling modules 
+ * within this plugin.
+ */
+class IGoogleCloudMessagingModuleInterface : public IModuleInterface
+{
+
+public:
+
+	virtual void RegisterForRemoteNotifications() = 0;
+
+	/**
+	 * Singleton-like access to this module's interface.  This is just for convenience!
+	 * Beware of calling this during the shutdown phase, though.  Your module might have been unloaded already.
+	 *
+	 * @return Returns singleton instance, loading the module on demand if needed
+	 */
+	static inline IGoogleCloudMessagingModuleInterface& Get()
+	{
+		return FModuleManager::LoadModuleChecked< IGoogleCloudMessagingModuleInterface >( "GoogleCloudMessaging" );
+	}
+
+	/**
+	 * Checks to see if this module is loaded and ready.  It is only valid to call Get() if IsAvailable() returns true.
+	 *
+	 * @return True if the module is loaded and ready to use
+	 */
+	static inline bool IsAvailable()
+	{
+		return FModuleManager::Get().IsModuleLoaded( "GoogleCloudMessaging" );
+	}
+
+    /** IModuleInterface implementation */
+    virtual void StartupModule() override;
+    virtual void ShutdownModule() override;
+};
+
diff --git a/Engine/Plugins/Runtime/GoogleCloudMessaging/Source/Java/RemoteNotificationsInstanceIDListener.java b/Engine/Plugins/Runtime/GoogleCloudMessaging/Source/Java/RemoteNotificationsInstanceIDListener.java
new file mode 100644
--- /dev/null
+++ b/Engine/Plugins/Runtime/GoogleCloudMessaging/Source/Java/RemoteNotificationsInstanceIDListener.java
@@ -0,0 +1,16 @@
+// Copyright 1998-2016 Epic Games, Inc. All Rights Reserved.
+package com.epicgames.ue4;
+
+import android.content.Intent;
+import com.google.android.gms.iid.InstanceID;
+import com.google.android.gms.iid.InstanceIDListenerService;
+
+public class RemoteNotificationsInstanceIDListener extends InstanceIDListenerService
+{
+	@Override
+	public void onTokenRefresh()
+	{
+		Intent intent = new Intent( GameActivity._activity.getApplicationContext(), RemoteNotificationsRegistrationIntentService.class );
+		startService( intent );
+	}
+}
diff --git a/Engine/Plugins/Runtime/GoogleCloudMessaging/Source/Java/RemoteNotificationsListener.java b/Engine/Plugins/Runtime/GoogleCloudMessaging/Source/Java/RemoteNotificationsListener.java
new file mode 100644
--- /dev/null
+++ b/Engine/Plugins/Runtime/GoogleCloudMessaging/Source/Java/RemoteNotificationsListener.java
@@ -0,0 +1,15 @@
+// Copyright 1998-2016 Epic Games, Inc. All Rights Reserved.
+package com.epicgames.ue4;
+
+import android.os.Bundle;
+import com.google.android.gms.gcm.GcmListenerService;
+
+public class RemoteNotificationsListener extends GcmListenerService
+{
+  @Override
+  public void onMessageReceived( String from, Bundle data ) 
+  {
+    String message = data.getString( "message" );
+    GameActivity._activity.nativeGCMReceivedRemoteNotification( message );
+  }
+}
diff --git a/Engine/Plugins/Runtime/GoogleCloudMessaging/Source/Java/RemoteNotificationsRegistrationIntentService.java b/Engine/Plugins/Runtime/GoogleCloudMessaging/Source/Java/RemoteNotificationsRegistrationIntentService.java
new file mode 100644
--- /dev/null
+++ b/Engine/Plugins/Runtime/GoogleCloudMessaging/Source/Java/RemoteNotificationsRegistrationIntentService.java
@@ -0,0 +1,46 @@
+// Copyright 1998-2016 Epic Games, Inc. All Rights Reserved.
+package com.epicgames.ue4;
+
+import java.io.IOException;
+import android.app.IntentService;
+import android.content.Intent;
+import com.google.android.gms.iid.InstanceID;
+import com.google.android.gms.gcm.GoogleCloudMessaging;
+import com.google.android.gms.gcm.GcmPubSub;
+
+public class RemoteNotificationsRegistrationIntentService extends IntentService
+{
+  public RemoteNotificationsRegistrationIntentService()
+  {
+    super( "com.epicgames.ue4.RemoteNotificationsRegistrationIntentService" );
+  }
+
+  public RemoteNotificationsRegistrationIntentService( String name )
+  {
+    super( name );
+  }
+  
+  @Override
+  protected void onHandleIntent( Intent intent )
+  {
+    try
+    {
+      InstanceID instanceID = InstanceID.getInstance( GameActivity._activity.getApplicationContext() );
+      String Token = instanceID.getToken( GameActivity._activity.GCMSenderId, GoogleCloudMessaging.INSTANCE_ID_SCOPE, null );
+
+      // subscribe to default topic: /topics/global
+      subscribeTopics( Token );
+
+      GameActivity._activity.nativeGCMRegisteredForRemoteNotifications( Token );
+    }
+    catch( Exception e )
+    {
+      GameActivity._activity.nativeGCMFailedToRegisterForRemoteNotifications( "Failed to complete token refresh" );
+    }
+  }
+
+  private void subscribeTopics( String token ) throws IOException {
+    GcmPubSub pubSub = GcmPubSub.getInstance( GameActivity._activity.getApplicationContext() );
+    pubSub.subscribe( token, "/topics/global", null );
+  }
+}
diff --git a/Engine/Source/Developer/AudioFormatADPCM/Private/AudioFormatADPCM.cpp b/Engine/Source/Developer/AudioFormatADPCM/Private/AudioFormatADPCM.cpp
--- a/Engine/Source/Developer/AudioFormatADPCM/Private/AudioFormatADPCM.cpp
+++ b/Engine/Source/Developer/AudioFormatADPCM/Private/AudioFormatADPCM.cpp
@@ -403,10 +403,10 @@ public:
 	{
 		uint8 const*	SrcData = SrcBuffer.GetData();
 		uint32			SrcSize = SrcBuffer.Num();
-		uint32			bytesProcessed = 0;
+		uint32			BytesProcessed = 0;
 		
-		FWaveModInfo	waveInfo;
-		waveInfo.ReadWaveInfo((uint8*)SrcData, SrcSize);
+		FWaveModInfo	WaveInfo;
+		WaveInfo.ReadWaveInfo((uint8*)SrcData, SrcSize);
 		
 		// Choose a chunk size that is much larger then the number of samples that the os audio render callback will ask for so that the chunk system has time to load new data before its needed
 		// The audio render callback will typically ask for a power of 2 number of samples. However, there is not a power of 2 number of samples in the uncompressed block.
@@ -416,21 +416,23 @@ public:
 		// Also, the first 78 odd bytes of the buffer is for the header so the first chunk is 78 bytes bigger then the rest. This needs to be accounted for when using chunk 0.
 
 		// The incoming data is organized by channel (all samples for channel 0 then all samples for channel 1, etc) but for streaming we need the channels interlaced by compressed blocks
-		// so that each chunk contains an even number of copmressed blocks per channel, ie channels can not span one chunk
+		// so that each chunk contains an even number of compressed blocks per channel, ie channels can not span one chunk
 		
-		const int32	NumChannels = *waveInfo.pChannels;
+		const int32	NumChannels = *WaveInfo.pChannels;
 		
-		if(*waveInfo.pFormatTag == WAVE_FORMAT_ADPCM)
+		if(*WaveInfo.pFormatTag == WAVE_FORMAT_ADPCM)
 		{
-			int32 BlockSize = *waveInfo.pBlockAlign;
-			const int32 NumBlocksPerChannel = (waveInfo.SampleDataSize + BlockSize * NumChannels - 1) / (BlockSize * NumChannels);
+			int32 BlockSize = *WaveInfo.pBlockAlign;
+			const int32 NumBlocksPerChannel = (WaveInfo.SampleDataSize + BlockSize * NumChannels - 1) / (BlockSize * NumChannels);
 			
 			// Ensure chunk size is an even multiple of block size and channel number, ie (ChunkSize % (BlockSize * NumChannels)) == 0
-			int32 ChunkSize = ((MONO_PCM_BUFFER_SIZE + BlockSize * NumChannels - 1) / (BlockSize * NumChannels)) * (BlockSize * NumChannels);
+			// Use a base chunk size of MONO_PCM_BUFFER_SIZE * NumChannels * 2 because Android uses MONO_PCM_BUFFER_SIZE to submit buffers to the os, the streaming system has more scheduling flexability when the chunk size is
+			//	larger the the buffer size submitted to the OS
+			int32 ChunkSize = ((MONO_PCM_BUFFER_SIZE * NumChannels * 2 + BlockSize * NumChannels - 1) / (BlockSize * NumChannels)) * (BlockSize * NumChannels);
 			
 			// Add the first chunk and the header data
-			AddNewChunk(OutBuffers, ChunkSize + waveInfo.SampleDataStart - SrcData);	// Add the first chunk with enough reserve room for the header data
-			AddChunkData(OutBuffers, SrcData, waveInfo.SampleDataStart - SrcData);
+			AddNewChunk(OutBuffers, ChunkSize + WaveInfo.SampleDataStart - SrcData);	// Add the first chunk with enough reserve room for the header data
+			AddChunkData(OutBuffers, SrcData, WaveInfo.SampleDataStart - SrcData);
 			
 			int32	CurChunkDataSize = 0;	// Don't include the header size here since we want the first chunk to include both the header data and a full ChunkSize of data
 			
@@ -444,12 +446,12 @@ public:
 						CurChunkDataSize = 0;
 					}
 					
-					AddChunkData(OutBuffers, waveInfo.SampleDataStart + (channelItr * NumBlocksPerChannel + blockItr) * BlockSize, BlockSize);
+					AddChunkData(OutBuffers, WaveInfo.SampleDataStart + (channelItr * NumBlocksPerChannel + blockItr) * BlockSize, BlockSize);
 					CurChunkDataSize += BlockSize;
 				}
 			}
 		}
-		else if(*waveInfo.pFormatTag == WAVE_FORMAT_LPCM)
+		else if(*WaveInfo.pFormatTag == WAVE_FORMAT_LPCM)
 		{
 			int32 ChunkSize = MONO_PCM_BUFFER_SIZE * 4;	// Use an extra big buffer for uncompressed data since uncompressed uses about x4 amount of data
 			int32	FrameSize = sizeof(uint16) * NumChannels;
@@ -458,9 +460,9 @@ public:
 			
 			// Add the first chunk and the header data
 			AddNewChunk(OutBuffers, ChunkSize + 128);	// Add the first chunk with enough reserve room for the header data
-			AddChunkData(OutBuffers, SrcData, waveInfo.SampleDataStart - SrcData);
-			SrcSize -= waveInfo.SampleDataStart - SrcData;
-			SrcData = waveInfo.SampleDataStart;
+			AddChunkData(OutBuffers, SrcData, WaveInfo.SampleDataStart - SrcData);
+			SrcSize -= WaveInfo.SampleDataStart - SrcData;
+			SrcData = WaveInfo.SampleDataStart;
 			
 			while(SrcSize > 0)
 			{
diff --git a/Engine/Source/Developer/AudioFormatOgg/Private/AudioFormatOgg.cpp b/Engine/Source/Developer/AudioFormatOgg/Private/AudioFormatOgg.cpp
--- a/Engine/Source/Developer/AudioFormatOgg/Private/AudioFormatOgg.cpp
+++ b/Engine/Source/Developer/AudioFormatOgg/Private/AudioFormatOgg.cpp
@@ -434,6 +434,53 @@ public:
 
 		return CompressedDataStore.Num();
 	}
+
+	virtual bool SplitDataForStreaming(const TArray<uint8>& SrcBuffer, TArray<TArray<uint8>>& OutBuffers) const override
+	{
+		// Just chunk purely on MONO_PCM_BUFFER_SIZE
+		uint8 const*	SrcData = SrcBuffer.GetData();
+		uint32			SrcSize = SrcBuffer.Num();
+
+		// Load the audio quality info to get the number of channels
+		FVorbisAudioInfo	AudioInfo;
+		FSoundQualityInfo	QualityInfo;
+		if (!AudioInfo.ReadCompressedInfo((uint8*)SrcData, SrcSize, &QualityInfo))
+		{
+			return false;
+		}
+
+		// Use a base chunk size of MONO_PCM_BUFFER_SIZE * 2 because Android uses MONO_PCM_BUFFER_SIZE to submit buffers to the os, the streaming system has more scheduling flexability when the chunk size is
+		//	larger the the buffer size submitted to the OS
+		uint32	ChunkSize = MONO_PCM_BUFFER_SIZE * 2 * QualityInfo.NumChannels;
+		
+		while(SrcSize > 0)
+		{
+			int32	CurChunkDataSize = FMath::Min<uint32>(SrcSize, ChunkSize);
+				
+			AddNewChunk(OutBuffers, CurChunkDataSize);
+			AddChunkData(OutBuffers, SrcData, CurChunkDataSize);
+				
+			SrcSize -= CurChunkDataSize;
+			SrcData += CurChunkDataSize;
+		}
+
+		return true;
+	}
+	
+	// Add a new chunk and reserve ChunkSize bytes in it
+	void AddNewChunk(TArray<TArray<uint8>>& OutBuffers, int32 ChunkReserveSize) const
+	{
+		TArray<uint8>& NewBuffer = *new (OutBuffers) TArray<uint8>;
+		NewBuffer.Empty(ChunkReserveSize);
+	}
+	
+	// Add data to the current chunk
+	void AddChunkData(TArray<TArray<uint8>>& OutBuffers, const uint8* ChunkData, int32 ChunkDataSize) const
+	{
+		TArray<uint8>& TargetBuffer = OutBuffers[OutBuffers.Num() - 1];
+		TargetBuffer.Append(ChunkData, ChunkDataSize);
+	}
+
 };
 
 
diff --git a/Engine/Source/Developer/Linux/LinuxTargetPlatform/Private/LinuxTargetPlatform.h b/Engine/Source/Developer/Linux/LinuxTargetPlatform/Private/LinuxTargetPlatform.h
--- a/Engine/Source/Developer/Linux/LinuxTargetPlatform/Private/LinuxTargetPlatform.h
+++ b/Engine/Source/Developer/Linux/LinuxTargetPlatform/Private/LinuxTargetPlatform.h
@@ -56,22 +56,6 @@ public:
 		TextureLODSettings = nullptr;
 		StaticMeshLODSettings.Initialize(EngineSettings);
 
-		// Get the Target RHIs for this platform, we do not always want all those that are supported.
-		GConfig->GetArray(TEXT("/Script/LinuxTargetPlatform.LinuxTargetSettings"), TEXT("TargetedRHIs"), TargetedShaderFormats, GEngineIni);
-
-		// Gather the list of Target RHIs and filter out any that may be invalid.
-		TArray<FName> PossibleShaderFormats;
-		GetAllPossibleShaderFormats(PossibleShaderFormats);
-
-		for(int32 ShaderFormatIdx = TargetedShaderFormats.Num()-1; ShaderFormatIdx >= 0; ShaderFormatIdx--)
-		{
-			FString ShaderFormat = TargetedShaderFormats[ShaderFormatIdx];
-			if(PossibleShaderFormats.Contains(FName(*ShaderFormat)) == false)
-			{
-				TargetedShaderFormats.RemoveAt(ShaderFormatIdx);
-			}
-		}
-
 		InitDevicesFromConfig();
 #endif // WITH_ENGINE
 	}
@@ -191,7 +175,7 @@ public:
 			
 			// else check for legacy LINUX_ROOT
 			ToolchainRoot[ 0 ] = 0;
-			FPlatformMisc::GetEnvironmentVariable(TEXT("LINUX_ROOT"), ToolchainRoot, ARRAY_COUNT(ToolchainRoot));            
+			FPlatformMisc::GetEnvironmentVariable(TEXT("LINUX_ROOT"), ToolchainRoot, ARRAY_COUNT(ToolchainRoot));
 			FString ToolchainCompiler = ToolchainRoot;
 			if (PLATFORM_WINDOWS)
 			{
@@ -251,6 +235,23 @@ public:
 
 	virtual void GetAllTargetedShaderFormats( TArray<FName>& OutFormats ) const override
 	{
+		// Get the Target RHIs for this platform, we do not always want all those that are supported. (reload in case user changed in the editor)
+		TArray<FString>TargetedShaderFormats;
+		GConfig->GetArray(TEXT("/Script/LinuxTargetPlatform.LinuxTargetSettings"), TEXT("TargetedRHIs"), TargetedShaderFormats, GEngineIni);
+
+		// Gather the list of Target RHIs and filter out any that may be invalid.
+		TArray<FName> PossibleShaderFormats;
+		GetAllPossibleShaderFormats(PossibleShaderFormats);
+
+		for (int32 ShaderFormatIdx = TargetedShaderFormats.Num() - 1; ShaderFormatIdx >= 0; ShaderFormatIdx--)
+		{
+			FString ShaderFormat = TargetedShaderFormats[ShaderFormatIdx];
+			if (PossibleShaderFormats.Contains(FName(*ShaderFormat)) == false)
+			{
+				TargetedShaderFormats.RemoveAt(ShaderFormatIdx);
+			}
+		}
+
 		for(const FString& ShaderFormat : TargetedShaderFormats)
 		{
 			OutFormats.AddUnique(FName(*ShaderFormat));
@@ -463,8 +464,6 @@ private:
 	// Holds static mesh LOD settings.
 	FStaticMeshLODSettings StaticMeshLODSettings;
 
-	// List of shader formats specified as targets
-	TArray<FString> TargetedShaderFormats;
 #endif // WITH_ENGINE
 
 private:
diff --git a/Engine/Source/Developer/Mac/MacTargetPlatform/Private/GenericMacTargetPlatform.h b/Engine/Source/Developer/Mac/MacTargetPlatform/Private/GenericMacTargetPlatform.h
--- a/Engine/Source/Developer/Mac/MacTargetPlatform/Private/GenericMacTargetPlatform.h
+++ b/Engine/Source/Developer/Mac/MacTargetPlatform/Private/GenericMacTargetPlatform.h
@@ -46,34 +46,6 @@ public:
 			FConfigCacheIni::LoadLocalIniFile(EngineSettings, TEXT("Engine"), true, *this->PlatformName());
 			TextureLODSettings = nullptr;
 			StaticMeshLODSettings.Initialize(EngineSettings);
-		
-			// Get the Target RHIs for this platform, we do not always want all those that are supported.
-			GConfig->GetArray(TEXT("/Script/MacTargetPlatform.MacTargetSettings"), TEXT("TargetedRHIs"), TargetedShaderFormats, GEngineIni);
-		
-			// Get the cached shader formats for this platform, we do not always want all those that are supported.
-			GConfig->GetArray(TEXT("/Script/MacTargetPlatform.MacTargetSettings"), TEXT("CachedShaderFormats"), CachedShaderFormats, GEngineIni);
-		
-			// Gather the list of Target RHIs and filter out any that may be invalid.
-			TArray<FName> PossibleShaderFormats;
-			GetAllPossibleShaderFormats(PossibleShaderFormats);
-			
-			for(int32 ShaderFormatIdx = TargetedShaderFormats.Num()-1; ShaderFormatIdx >= 0; ShaderFormatIdx--)
-			{
-				FString ShaderFormat = TargetedShaderFormats[ShaderFormatIdx];
-				if(PossibleShaderFormats.Contains(FName(*ShaderFormat)) == false)
-				{
-					TargetedShaderFormats.RemoveAt(ShaderFormatIdx);
-				}
-			}
-		
-			for(int32 ShaderFormatIdx = CachedShaderFormats.Num()-1; ShaderFormatIdx >= 0; ShaderFormatIdx--)
-			{
-				FString ShaderFormat = CachedShaderFormats[ShaderFormatIdx];
-				if(PossibleShaderFormats.Contains(FName(*ShaderFormat)) == false)
-				{
-					CachedShaderFormats.RemoveAt(ShaderFormatIdx);
-				}
-			}
 		#endif
 	}
 
@@ -166,6 +138,23 @@ return TSuper::SupportsFeature(Feature);
 
 	virtual void GetAllTargetedShaderFormats(TArray<FName>& OutFormats) const override
 	{
+		// Get the Target RHIs for this platform, we do not always want all those that are supported.
+		TArray<FString>TargetedShaderFormats;
+		GConfig->GetArray(TEXT("/Script/MacTargetPlatform.MacTargetSettings"), TEXT("TargetedRHIs"), TargetedShaderFormats, GEngineIni);
+
+		// Gather the list of Target RHIs and filter out any that may be invalid.
+		TArray<FName> PossibleShaderFormats;
+		GetAllPossibleShaderFormats(PossibleShaderFormats);
+
+		for (int32 ShaderFormatIdx = TargetedShaderFormats.Num() - 1; ShaderFormatIdx >= 0; ShaderFormatIdx--)
+		{
+			FString ShaderFormat = TargetedShaderFormats[ShaderFormatIdx];
+			if (PossibleShaderFormats.Contains(FName(*ShaderFormat)) == false)
+			{
+				TargetedShaderFormats.RemoveAt(ShaderFormatIdx);
+			}
+		}
+
 		for(const FString& ShaderFormat : TargetedShaderFormats)
 		{
 			OutFormats.AddUnique(FName(*ShaderFormat));
@@ -174,6 +163,23 @@ return TSuper::SupportsFeature(Feature);
 	
 	virtual void GetAllCachedShaderFormats( TArray<FName>& OutFormats ) const override
 	{
+		// Get the cached shader formats for this platform, we do not always want all those that are supported.
+		TArray<FString> CachedShaderFormats;
+		GConfig->GetArray(TEXT("/Script/MacTargetPlatform.MacTargetSettings"), TEXT("CachedShaderFormats"), CachedShaderFormats, GEngineIni);
+
+		// Gather the list of Target RHIs and filter out any that may be invalid.
+		TArray<FName> PossibleShaderFormats;
+		GetAllPossibleShaderFormats(PossibleShaderFormats);
+
+		for (int32 ShaderFormatIdx = CachedShaderFormats.Num() - 1; ShaderFormatIdx >= 0; ShaderFormatIdx--)
+		{
+			FString ShaderFormat = CachedShaderFormats[ShaderFormatIdx];
+			if (PossibleShaderFormats.Contains(FName(*ShaderFormat)) == false)
+			{
+				CachedShaderFormats.RemoveAt(ShaderFormatIdx);
+			}
+		}
+
 		for(const FString& ShaderFormat : CachedShaderFormats)
 		{
 			OutFormats.AddUnique(FName(*ShaderFormat));
@@ -289,12 +295,7 @@ private:
 
 	// Holds the static mesh LOD settings.
 	FStaticMeshLODSettings StaticMeshLODSettings;
-	
-	// List of shader formats specified as targets
-	TArray<FString> TargetedShaderFormats;
-	
-	// List of shader formats specified to cache
-	TArray<FString> CachedShaderFormats;
+
 #endif // WITH_ENGINE
 
 private:
diff --git a/Engine/Source/Developer/Windows/WindowsTargetPlatform/Private/GenericWindowsTargetPlatform.h b/Engine/Source/Developer/Windows/WindowsTargetPlatform/Private/GenericWindowsTargetPlatform.h
--- a/Engine/Source/Developer/Windows/WindowsTargetPlatform/Private/GenericWindowsTargetPlatform.h
+++ b/Engine/Source/Developer/Windows/WindowsTargetPlatform/Private/GenericWindowsTargetPlatform.h
@@ -41,24 +41,13 @@ public:
 			TextureLODSettings = nullptr; // These are registered by the device profile system.
 			StaticMeshLODSettings.Initialize(EngineSettings);
 
-			// Get the Target RHIs for this platform, we do not always want all those that are supported.
-			GConfig->GetArray(TEXT("/Script/WindowsTargetPlatform.WindowsTargetSettings"), TEXT("TargetedRHIs"), TargetedShaderFormats, GEngineIni);
-			
-			// Gather the list of Target RHIs and filter out any that may be invalid.
-			TArray<FName> PossibleShaderFormats;
-			GetAllPossibleShaderFormats(PossibleShaderFormats);
 
-			for(int32 ShaderFormatIdx = TargetedShaderFormats.Num()-1; ShaderFormatIdx >= 0; ShaderFormatIdx--)
-			{
-				FString ShaderFormat = TargetedShaderFormats[ShaderFormatIdx];
-				if(PossibleShaderFormats.Contains(FName(*ShaderFormat)) == false)
-				{
-					TargetedShaderFormats.RemoveAt(ShaderFormatIdx);
-				}
-			}
+			// Get the Target RHIs for this platform, we do not always want all those that are supported.
+			TArray<FName> TargetedShaderFormats;
+			GetAllTargetedShaderFormats(TargetedShaderFormats);
 
 			// If we're targeting only DX11 we can use DX11 texture formats. Otherwise we'd have to compress fallbacks and increase the size of cooked content significantly.
-			static FString NAME_PCD3D_SM5(TEXT("PCD3D_SM5"));
+			static FName NAME_PCD3D_SM5(TEXT("PCD3D_SM5"));
 			bSupportDX11TextureFormats = TargetedShaderFormats.Num() == 1
 				&& TargetedShaderFormats[0] == NAME_PCD3D_SM5;
 		#endif
@@ -169,6 +158,23 @@ public:
 
 	virtual void GetAllTargetedShaderFormats( TArray<FName>& OutFormats ) const override
 	{
+		// Get the Target RHIs for this platform, we do not always want all those that are supported. (reload in case user changed in the editor)
+		TArray<FString>TargetedShaderFormats;
+		GConfig->GetArray(TEXT("/Script/WindowsTargetPlatform.WindowsTargetSettings"), TEXT("TargetedRHIs"), TargetedShaderFormats, GEngineIni);
+
+		// Gather the list of Target RHIs and filter out any that may be invalid.
+		TArray<FName> PossibleShaderFormats;
+		GetAllPossibleShaderFormats(PossibleShaderFormats);
+		
+		for (int32 ShaderFormatIdx = TargetedShaderFormats.Num() - 1; ShaderFormatIdx >= 0; ShaderFormatIdx--)
+		{
+			FString ShaderFormat = TargetedShaderFormats[ShaderFormatIdx];
+			if (PossibleShaderFormats.Contains(FName(*ShaderFormat)) == false)
+			{
+				TargetedShaderFormats.RemoveAt(ShaderFormatIdx);
+			}
+		}
+
 		for(const FString& ShaderFormat : TargetedShaderFormats)
 		{
 			OutFormats.AddUnique(FName(*ShaderFormat));
@@ -280,9 +286,6 @@ private:
 	// Holds static mesh LOD settings.
 	FStaticMeshLODSettings StaticMeshLODSettings;
 
-	// List of shader formats specified as targets
-	TArray<FString> TargetedShaderFormats;
-
 	// True if the project supports non-DX11 texture formats.
 	bool bSupportDX11TextureFormats;
 #endif // WITH_ENGINE
diff --git a/Engine/Source/Developer/iOS/IOSTargetPlatform/Private/IOSTargetPlatform.cpp b/Engine/Source/Developer/iOS/IOSTargetPlatform/Private/IOSTargetPlatform.cpp
--- a/Engine/Source/Developer/iOS/IOSTargetPlatform/Private/IOSTargetPlatform.cpp
+++ b/Engine/Source/Developer/iOS/IOSTargetPlatform/Private/IOSTargetPlatform.cpp
@@ -170,6 +170,9 @@ int32 FIOSTargetPlatform::CheckRequirements(const FString& ProjectPath, bool bPr
 
 	// shell to IPP and get the status of the provision and cert
 
+	bool bForDistribtion = false;
+	GConfig->GetBool(TEXT("/Script/UnrealEd.ProjectPackagingSettings"), TEXT("ForDistribution"), bForDistribtion, GGameIni);
+
 	FString BundleIdentifier;
 	GConfig->GetString(TEXT("/Script/IOSRuntimeSettings.IOSRuntimeSettings"), TEXT("BundleIdentifier"), BundleIdentifier, GEngineIni);
 	BundleIdentifier = BundleIdentifier.Replace(TEXT("[PROJECT_NAME]"), FApp::GetGameName());
@@ -178,10 +181,10 @@ int32 FIOSTargetPlatform::CheckRequirements(const FString& ProjectPath, bool bPr
     FString CmdExe = TEXT("/bin/sh");
     FString ScriptPath = FPaths::ConvertRelativePathToFull(FPaths::EngineDir() / TEXT("Build/BatchFiles/Mac/RunMono.sh"));
     FString IPPPath = FPaths::ConvertRelativePathToFull(FPaths::EngineDir() / TEXT("Binaries/DotNet/IOS/IPhonePackager.exe"));
-	FString CommandLine = FString::Printf(TEXT("\"%s\" \"%s\" Validate Engine -project \"%s\" -bundlename \"%s\""), *ScriptPath, *IPPPath, *ProjectPath, *(BundleIdentifier));
+	FString CommandLine = FString::Printf(TEXT("\"%s\" \"%s\" Validate Engine -project \"%s\" -bundlename \"%s\" %s"), *ScriptPath, *IPPPath, *ProjectPath, *(BundleIdentifier), (bForDistribtion ? TEXT("-distribution") : TEXT("")) );
 #else
 	FString CmdExe = FPaths::ConvertRelativePathToFull(FPaths::EngineDir() / TEXT("Binaries/DotNet/IOS/IPhonePackager.exe"));
-	FString CommandLine = FString::Printf(TEXT("Validate Engine -project \"%s\" -bundlename \"%s\""), *ProjectPath, *(BundleIdentifier));
+	FString CommandLine = FString::Printf(TEXT("Validate Engine -project \"%s\" -bundlename \"%s\" %s"), *ProjectPath, *(BundleIdentifier), (bForDistribtion ? TEXT("-distribution") : TEXT("")) );
 	FString RemoteServerName;
 	GConfig->GetString(TEXT("/Script/IOSRuntimeSettings.IOSRuntimeSettings"), TEXT("RemoteServerName"), RemoteServerName, GEngineIni);
 	if (RemoteServerName.Len() == 0)
diff --git a/Engine/Source/Editor/DetailCustomizations/Private/KeyStructCustomization.cpp b/Engine/Source/Editor/DetailCustomizations/Private/KeyStructCustomization.cpp
--- a/Engine/Source/Editor/DetailCustomizations/Private/KeyStructCustomization.cpp
+++ b/Engine/Source/Editor/DetailCustomizations/Private/KeyStructCustomization.cpp
@@ -3,6 +3,7 @@
 #include "KeyStructCustomization.h"
 #include "DetailWidgetRow.h"
 #include "SKeySelector.h"
+#include "UnrealType.h"
 
 #define LOCTEXT_NAMESPACE "FKeyStructCustomization"
 
@@ -36,6 +37,7 @@ void FKeyStructCustomization::CustomizeHeader( TSharedRef<class IPropertyHandle>
 		.CurrentKey(this, &FKeyStructCustomization::GetCurrentKey)
 		.OnKeyChanged(this, &FKeyStructCustomization::OnKeyChanged)
 		.Font(StructCustomizationUtils.GetRegularFont())
+		.AllowClear(!StructPropertyHandle->GetProperty()->HasAnyPropertyFlags(CPF_NoClear))
 	];
 }
 
diff --git a/Engine/Source/Editor/DeviceProfileServices/Private/DeviceProfileServicesUIManager.cpp b/Engine/Source/Editor/DeviceProfileServices/Private/DeviceProfileServicesUIManager.cpp
--- a/Engine/Source/Editor/DeviceProfileServices/Private/DeviceProfileServicesUIManager.cpp
+++ b/Engine/Source/Editor/DeviceProfileServices/Private/DeviceProfileServicesUIManager.cpp
@@ -86,7 +86,7 @@ void FDeviceProfileServicesUIManager::CreatePlatformMap()
 	TArray<ITargetPlatform*> Platforms = GetTargetPlatformManager()->GetTargetPlatforms();
 	for (int32 Index = 0; Index < Platforms.Num(); ++Index)
 	{
-		PlatformList.Add(MakeShareable(new FString(Platforms[Index]->PlatformName())));
+		PlatformList.Add(MakeShareable(new FString(Platforms[Index]->IniPlatformName())));
 		DeviceTypeToIconMap.Add( Platforms[Index]->PlatformName(), Platforms[Index]->GetPlatformInfo().GetIconStyleName(PlatformInfo::EPlatformIconSize::Normal) );
 	}
 }
diff --git a/Engine/Source/Editor/LevelEditor/Private/SLevelViewportToolBar.cpp b/Engine/Source/Editor/LevelEditor/Private/SLevelViewportToolBar.cpp
--- a/Engine/Source/Editor/LevelEditor/Private/SLevelViewportToolBar.cpp
+++ b/Engine/Source/Editor/LevelEditor/Private/SLevelViewportToolBar.cpp
@@ -204,7 +204,7 @@ void SLevelViewportToolBar::Construct( const FArguments& InArgs )
 					.LabelIcon( this, &SLevelViewportToolBar::GetDevicePreviewMenuLabelIcon )
 					.OnGetMenuContent( this, &SLevelViewportToolBar::GenerateDevicePreviewMenu )
 					//@todo rendering: mobile preview in view port is not functional yet - remove this once it is.
-					.Visibility(EVisibility::Collapsed)
+					//.Visibility(EVisibility::Collapsed)
 				]
 				+ SHorizontalBox::Slot()
 				.Padding( ToolbarSlotPadding )
diff --git a/Engine/Source/Editor/UnrealEd/Private/Commandlets/DerivedDataCacheCommandlet.cpp b/Engine/Source/Editor/UnrealEd/Private/Commandlets/DerivedDataCacheCommandlet.cpp
--- a/Engine/Source/Editor/UnrealEd/Private/Commandlets/DerivedDataCacheCommandlet.cpp
+++ b/Engine/Source/Editor/UnrealEd/Private/Commandlets/DerivedDataCacheCommandlet.cpp
@@ -16,6 +16,7 @@ DerivedDataCacheCommandlet.cpp: Commandlet for DDC maintenence
 #include "ShaderCompiler.h"
 #include "DistanceFieldAtlas.h"
 #include "Misc/RedirectCollector.h"
+#include "Engine/Texture.h"
 
 DEFINE_LOG_CATEGORY_STATIC(LogDerivedDataCacheCommandlet, Log, All);
 
@@ -77,6 +78,15 @@ int32 UDerivedDataCacheCommandlet::Main( const FString& Params )
 		UE_LOG(LogDerivedDataCacheCommandlet, Display, TEXT("Done waiting for shaders to finish."));
 	};
 
+
+	auto WaitForCurrentTextureBuildingToFinish = []()
+	{
+		for ( TObjectIterator<UTexture> Texture; Texture; ++Texture )
+		{
+			Texture->FinishCachePlatformData();
+		}
+	};
+
 	if (!bStartupOnly && bFillCache)
 	{
 		FCoreUObjectDelegates::PackageCreatedForLoad.AddUObject(this, &UDerivedDataCacheCommandlet::MaybeMarkPackageAsAlreadyLoaded);
@@ -267,6 +277,9 @@ int32 UDerivedDataCacheCommandlet::Main( const FString& Params )
 
 			if (NumProcessedSinceLastGC >= GCInterval || FileIndex < 0 || bLastPackageWasMap)
 			{
+				WaitForCurrentShaderCompilationToFinish();
+				WaitForCurrentTextureBuildingToFinish();
+
 				const double StartGCTime = FPlatformTime::Seconds();
 				if (NumProcessedSinceLastGC >= GCInterval || FileIndex < 0)
 				{
@@ -282,12 +295,12 @@ int32 UDerivedDataCacheCommandlet::Main( const FString& Params )
 				GCTime += FPlatformTime::Seconds() - StartGCTime;
 
 				bLastPackageWasMap = false;
-				WaitForCurrentShaderCompilationToFinish();
 			}
 		}
 	}
 
 	WaitForCurrentShaderCompilationToFinish();
+	WaitForCurrentTextureBuildingToFinish();
 	GetDerivedDataCacheRef().WaitForQuiescence(true);
 
 	UE_LOG(LogDerivedDataCacheCommandlet, Display, TEXT("%.2lfs spent looking for processed packages, %.2lfs spent on GC."), FindProcessedPackagesTime, GCTime);
diff --git a/Engine/Source/Editor/UnrealEd/Private/CookOnTheFlyServer.cpp b/Engine/Source/Editor/UnrealEd/Private/CookOnTheFlyServer.cpp
--- a/Engine/Source/Editor/UnrealEd/Private/CookOnTheFlyServer.cpp
+++ b/Engine/Source/Editor/UnrealEd/Private/CookOnTheFlyServer.cpp
@@ -3544,7 +3544,8 @@ void UCookOnTheFlyServer::ProcessAccessedIniSettings(const FConfigFile* Config,
 				AccessedKey.Empty();
 				for ( const auto& ValueArrayEntry : ValueArray )
 				{
-					AccessedKey.Add(ValueArrayEntry.GetSavedValue());
+					FString RemovedColon = ValueArrayEntry.GetSavedValue().Replace(TEXT(":"), TEXT(""));
+					AccessedKey.Add(RemovedColon);
 				}
 			}
 			
@@ -3666,7 +3667,8 @@ bool UCookOnTheFlyServer::IniSettingsOutOfDate(const ITargetPlatform* TargetPlat
 				}
 				for ( int Index = 0; Index < CurrentValues.Num(); ++Index )
 				{
-					if ( CurrentValues[Index].GetSavedValue() != OldIniValue.Value[Index] )
+					const FString FilteredCurrentValue = CurrentValues[Index].GetSavedValue().Replace(TEXT(":"), TEXT(""));
+					if ( FilteredCurrentValue != OldIniValue.Value[Index] )
 					{
 						UE_LOG(LogCook, Display, TEXT("Inisetting is different for %s, value %s != %s invalidating cook"),  *FString::Printf(TEXT("%s %s %s %s %d"),*PlatformName, *Filename, *SectionName.ToString(), *ValueName.ToString(), Index), *CurrentValues[Index].GetSavedValue(), *OldIniValue.Value[Index] );
 						return true;
@@ -6171,6 +6173,9 @@ void UCookOnTheFlyServer::StartCookByTheBook( const FCookByTheBookStartupOptions
 	// add all the files for the requested platform to the cook list
 	for ( const FName& FileFName : FilesInPath )
 	{
+		if ( FileFName == NAME_None )
+			continue;
+
 		const FName PackageFileFName = GetCachedStandardPackageFileFName(FileFName);
 		
 		if (PackageFileFName != NAME_None)
diff --git a/Engine/Source/Editor/UnrealEd/Private/Kismet2/BlueprintEditorUtils.cpp b/Engine/Source/Editor/UnrealEd/Private/Kismet2/BlueprintEditorUtils.cpp
--- a/Engine/Source/Editor/UnrealEd/Private/Kismet2/BlueprintEditorUtils.cpp
+++ b/Engine/Source/Editor/UnrealEd/Private/Kismet2/BlueprintEditorUtils.cpp
@@ -954,7 +954,7 @@ struct FRegenerationHelper
 		FLinkerLoad* Linker = Obj->GetLinker();
 		if (Linker && !Obj->HasAnyFlags(RF_LoadCompleted))
 		{
-			check(!GEventDrivenLoaderEnabled);
+			check(!GEventDrivenLoaderEnabled || !EVENT_DRIVEN_ASYNC_LOAD_ACTIVE_AT_RUNTIME);
 			Obj->SetFlags(RF_NeedLoad);
 			Linker->Preload(Obj);
 			return true;
diff --git a/Engine/Source/Editor/UnrealEd/Private/SKeySelector.cpp b/Engine/Source/Editor/UnrealEd/Private/SKeySelector.cpp
--- a/Engine/Source/Editor/UnrealEd/Private/SKeySelector.cpp
+++ b/Engine/Source/Editor/UnrealEd/Private/SKeySelector.cpp
@@ -106,6 +106,12 @@ void SKeySelector::Construct(const FArguments& InArgs)
 		}
 	}
 
+	// if we allow NoClear, add a "None" option to be able to clear out a binding
+	if (InArgs._AllowClear)
+	{
+		new (KeyTreeRoot) FKeyTreeItem(MakeShareable(new FKeyTreeInfo(FText(), MakeShareable(new FKey(NAME_None)))));
+	}
+
 	TreeViewWidth = InArgs._TreeViewWidth;
 	TreeViewHeight = InArgs._TreeViewHeight;
 	CategoryFont = FSlateFontInfo( FPaths::EngineContentDir() / TEXT("Slate/Fonts/Roboto-Bold.ttf"), 9 );
diff --git a/Engine/Source/Editor/UnrealEd/Public/SKeySelector.h b/Engine/Source/Editor/UnrealEd/Public/SKeySelector.h
--- a/Engine/Source/Editor/UnrealEd/Public/SKeySelector.h
+++ b/Engine/Source/Editor/UnrealEd/Public/SKeySelector.h
@@ -37,6 +37,7 @@ public:
 		, _TreeViewHeight(400.f)
 		, _Font( FEditorStyle::GetFontStyle( TEXT("NormalFont") ) )
 		, _FilterBlueprintBindable( true )
+		, _AllowClear( true )
 		{}
 		SLATE_ATTRIBUTE( TOptional<FKey>, CurrentKey )
 		SLATE_ATTRIBUTE( FOptionalSize, TreeViewWidth )
@@ -44,6 +45,7 @@ public:
 		SLATE_EVENT( FOnKeyChanged, OnKeyChanged )
 		SLATE_ATTRIBUTE( FSlateFontInfo, Font )
 		SLATE_ARGUMENT( bool, FilterBlueprintBindable )
+		SLATE_ARGUMENT( bool, AllowClear )
 	SLATE_END_ARGS()
 public:
 	void Construct(const FArguments& InArgs);
diff --git a/Engine/Source/Programs/AutomationTool/AutomationUtils/DeploymentContext.cs b/Engine/Source/Programs/AutomationTool/AutomationUtils/DeploymentContext.cs
--- a/Engine/Source/Programs/AutomationTool/AutomationUtils/DeploymentContext.cs
+++ b/Engine/Source/Programs/AutomationTool/AutomationUtils/DeploymentContext.cs
@@ -372,8 +372,26 @@ public class DeploymentContext //: ProjectParams
 			{
 				OutputPath = CommandUtils.CombinePaths(InputPath.Substring(LocalRoot.Length).TrimStart('/', '\\'));
 			}
-			else
-			{
+            else if (InputPath.EndsWith(".uplugin", StringComparison.InvariantCultureIgnoreCase))
+            {
+                // This is a plugin that lives outside of the Engine/Plugins or Game/Plugins directory so needs to be remapped for staging/packaging
+                // We need to remap C:\SomePath\PluginName\PluginName.uplugin to RemappedPlugins\PluginName\PluginName.uplugin
+                int Index = InputPath.LastIndexOf(Path.DirectorySeparatorChar);
+                if (Index != -1)
+                {
+                    int PluginDirIndex = InputPath.LastIndexOf(Path.DirectorySeparatorChar, Index - 1);
+                    if (PluginDirIndex != -1)
+                    {
+                        OutputPath = CommandUtils.CombinePaths("RemappedPlugins", InputPath.Substring(PluginDirIndex));
+                    }
+                }
+                if (OutputPath == null)
+                {
+                    throw new AutomationException("Can't deploy {0} because the plugin path is non-standard, so could not be remapped", InputPath);
+                }
+            }
+            else
+            {
 				throw new AutomationException("Can't deploy {0} because it doesn't start with {1} or {2}", InputPath, ProjectRoot, LocalRoot);
 			}
 		}
diff --git a/Engine/Source/Programs/AutomationTool/HTML5/HTML5Platform.Automation.cs b/Engine/Source/Programs/AutomationTool/HTML5/HTML5Platform.Automation.cs
--- a/Engine/Source/Programs/AutomationTool/HTML5/HTML5Platform.Automation.cs
+++ b/Engine/Source/Programs/AutomationTool/HTML5/HTML5Platform.Automation.cs
@@ -119,29 +119,7 @@ public class HTML5Platform : Platform
 		string TemplateFile = Path.Combine(CombinePaths(CmdEnv.LocalRoot, "Engine"), "Build", "HTML5", TemplateFileName);
 		string OutputFile = Path.Combine(PackagePath, (Params.ClientConfigsToBuild[0].ToString() != "Development" ? (Params.ShortProjectName + "-HTML5-" + Params.ClientConfigsToBuild[0].ToString()) : Params.ShortProjectName)) + ".html";
 
-		// find Heap Size.
-		ulong HeapSize;
-
-		int ConfigHeapSize = 0;
-		// Valuer set by Editor UI
-		var bGotHeapSize = ConfigCache.GetInt32("/Script/HTML5PlatformEditor.HTML5TargetSettings", "HeapSize" + Params.ClientConfigsToBuild[0].ToString(), out ConfigHeapSize);
-
-		// Fallback if the previous method failed
-		if (!bGotHeapSize && !ConfigCache.GetInt32("/Script/BuildSettings.BuildSettings", "HeapSize" + Params.ClientConfigsToBuild[0].ToString(), out ConfigHeapSize)) // in Megs.
-		{
-			// we couldn't find a per config heap size, look for a common one.
-			if (!ConfigCache.GetInt32("/Script/BuildSettings.BuildSettings", "HeapSize", out ConfigHeapSize))
-			{
-				ConfigHeapSize = Params.IsCodeBasedProject ? 1024 : 512;
-				Log("Could not find Heap Size setting in .ini for Client config {0}", Params.ClientConfigsToBuild[0].ToString());
-			}
-		}
-
-		HeapSize = (ulong)ConfigHeapSize * 1024L * 1024L; // convert to bytes.
-		Log("Setting Heap size to {0} Mb ", ConfigHeapSize);
-
-
-		GenerateFileFromTemplate(TemplateFile, OutputFile, Params.ShortProjectName, Params.ClientConfigsToBuild[0].ToString(), Params.StageCommandline, !Params.IsCodeBasedProject, HeapSize);
+		GenerateFileFromTemplate(TemplateFile, OutputFile, Params.ShortProjectName, Params.ClientConfigsToBuild[0].ToString(), Params.StageCommandline, !Params.IsCodeBasedProject, HTML5SDKInfo.HeapSize(ConfigCache, Params.ClientConfigsToBuild[0].ToString()));
 
 		string MacBashTemplateFile = Path.Combine(CombinePaths(CmdEnv.LocalRoot, "Engine"), "Build", "HTML5", "RunMacHTML5LaunchHelper.command.template");
 		string MacBashOutputFile = Path.Combine(PackagePath, "RunMacHTML5LaunchHelper.command");
@@ -236,7 +214,7 @@ public class HTML5Platform : Platform
 	}
 
 
-	protected void GenerateFileFromTemplate(string InTemplateFile, string InOutputFile, string InGameName, string InGameConfiguration, string InArguments, bool IsContentOnly, ulong HeapSize)
+	protected void GenerateFileFromTemplate(string InTemplateFile, string InOutputFile, string InGameName, string InGameConfiguration, string InArguments, bool IsContentOnly, int HeapSize)
 	{
 		StringBuilder outputContents = new StringBuilder();
 		using (StreamReader reader = new StreamReader(InTemplateFile))
@@ -252,7 +230,7 @@ public class HTML5Platform : Platform
 
 				if (LineStr.Contains("%HEAPSIZE%"))
 				{
-					LineStr = LineStr.Replace("%HEAPSIZE%", HeapSize.ToString());
+					LineStr = LineStr.Replace("%HEAPSIZE%", HeapSize.ToString() + " * 1024 * 1024");
 				}
 
 				if (LineStr.Contains("%CONFIG%"))
diff --git a/Engine/Source/Programs/ShaderCompileWorker/ShaderCompileWorker.Target.cs b/Engine/Source/Programs/ShaderCompileWorker/ShaderCompileWorker.Target.cs
--- a/Engine/Source/Programs/ShaderCompileWorker/ShaderCompileWorker.Target.cs
+++ b/Engine/Source/Programs/ShaderCompileWorker/ShaderCompileWorker.Target.cs
@@ -34,6 +34,7 @@ public class ShaderCompileWorkerTarget : TargetRules
 		UEBuildConfiguration.bCompileAgainstEngine = false;
 		UEBuildConfiguration.bCompileAgainstCoreUObject = false;
 		UEBuildConfiguration.bBuildWithEditorOnlyData = true;
+		UEBuildConfiguration.bCompileCEF3 = false;
 
 		// Never use malloc profiling in ShaderCompileWorker.
 		BuildConfiguration.bUseMallocProfiler = false;
diff --git a/Engine/Source/Programs/UnrealBuildTool/Android/AndroidProjectGenerator.cs b/Engine/Source/Programs/UnrealBuildTool/Android/AndroidProjectGenerator.cs
--- a/Engine/Source/Programs/UnrealBuildTool/Android/AndroidProjectGenerator.cs
+++ b/Engine/Source/Programs/UnrealBuildTool/Android/AndroidProjectGenerator.cs
@@ -15,7 +15,39 @@ namespace UnrealBuildTool
 	{
 		static bool CheckedForNsight = false;		// whether we have checked for a recent enough version of Nsight yet
 		static bool NsightInstalled = false;		// true if a recent enough version of Nsight is installed
-		static int NsightVersionCode = 0;			// version code matching detected Nsight
+		static int NsightVersionCode = 0;           // version code matching detected Nsight
+
+		public static bool VSDebugCommandLineOptionPresent = false;    //User must put -vsdebugandroid commandline option to build the debug projects
+		static bool VSDebuggingEnabled = false;      // When set to true, allows debugging with built in MS Cross Platform Android tools.  
+													 //  It adds separate projects ending in .androidproj and a file VSAndroidUnreal.props for the engine and all game projects
+
+		static bool VSSupportChecked = false;       // Don't want to check multiple times
+
+		static bool VSPropsFileWritten = false;     // This is for the file VSAndroidUnreal.props which only needs to be written once
+
+
+		bool IsVSAndroidSupportInstalled()
+		{
+			if (VSSupportChecked)
+			{
+				return VSDebuggingEnabled;
+			}
+			VSSupportChecked = true;
+			
+			//check to make sure Cross Platform Tools are installed for MS
+			string Path = Microsoft.Win32.Registry.GetValue("HKEY_LOCAL_MACHINE\\SOFTWARE\\Wow6432Node\\Android SDK Tools", "Path", null) as string;
+			if (!String.IsNullOrEmpty(Path) && VSDebugCommandLineOptionPresent)
+			{
+				VSDebuggingEnabled = true;
+			}
+			else if (VSDebugCommandLineOptionPresent)
+			{
+				Log.TraceWarning("Android SDK tools have to be installed to use this.  Please Install Visual C++ for Cross-Platform Mobile Development https://msdn.microsoft.com/en-us/library/dn707598.aspx");
+			}
+
+			return VSDebuggingEnabled;
+
+		}
 
 		/// <summary>
 		/// Check to see if a recent enough version of Nsight is installed.
@@ -112,7 +144,7 @@ namespace UnrealBuildTool
 		public override bool HasVisualStudioSupport(UnrealTargetPlatform InPlatform, UnrealTargetConfiguration InConfiguration, VCProjectFileFormat ProjectFileFormat)
 		{
 			// Debugging, etc. are dependent on the TADP being installed
-			return IsNsightInstalled(ProjectFileFormat);
+			return IsNsightInstalled(ProjectFileFormat) || IsVSAndroidSupportInstalled();
 		}
 
 		/// <summary>
@@ -123,12 +155,21 @@ namespace UnrealBuildTool
 		/// <returns>string    The name of the platform that VisualStudio recognizes</returns>
 		public override string GetVisualStudioPlatformName(UnrealTargetPlatform InPlatform, UnrealTargetConfiguration InConfiguration)
 		{
+			string PlatformName = InPlatform.ToString();
+
 			if (InPlatform == UnrealTargetPlatform.Android)
 			{
-				return "Tegra-Android";
+				if (IsVSAndroidSupportInstalled())
+				{
+					PlatformName = "ARM";
+				}
+				else
+				{
+					PlatformName = "Tegra-Android";
+				}
 			}
 
-			return InPlatform.ToString();
+			return PlatformName;
 		}
 
 		/// <summary>
@@ -139,7 +180,7 @@ namespace UnrealBuildTool
 		/// <returns>string    The custom property import lines for the project file; Empty string if it doesn't require one</returns>
 		public override string GetAdditionalVisualStudioPropertyGroups(UnrealTargetPlatform InPlatform, VCProjectFileFormat ProjectFileFormat)
 		{
-			if (!IsNsightInstalled(ProjectFileFormat))
+			if(IsVSAndroidSupportInstalled() || !IsNsightInstalled(ProjectFileFormat))
 			{
 				return base.GetAdditionalVisualStudioPropertyGroups(InPlatform, ProjectFileFormat);
 			}
@@ -157,12 +198,22 @@ namespace UnrealBuildTool
 		/// <returns>string    The custom property import lines for the project file; Empty string if it doesn't require one</returns>
 		public override string GetVisualStudioPlatformConfigurationType(UnrealTargetPlatform InPlatform, VCProjectFileFormat ProjectFileFormat)
 		{
-			if (!IsNsightInstalled(ProjectFileFormat))
+			string ConfigurationType = "";
+
+			if(IsVSAndroidSupportInstalled())
 			{
-				return base.GetVisualStudioPlatformConfigurationType(InPlatform, ProjectFileFormat);
+				ConfigurationType = "Makefile";
 			}
-
-			return "ExternalBuildSystem";
+			else if (IsNsightInstalled(ProjectFileFormat))
+			{
+				ConfigurationType = "ExternalBuildSystem";				
+			}
+			else
+			{
+				ConfigurationType = base.GetVisualStudioPlatformConfigurationType(InPlatform, ProjectFileFormat);
+			}
+			
+			return ConfigurationType;
 		}
 
 		/// <summary>
@@ -174,7 +225,7 @@ namespace UnrealBuildTool
 		/// <returns>string    The custom configuration section for the project file; Empty string if it doesn't require one</returns>
 		public override string GetVisualStudioPlatformToolsetString(UnrealTargetPlatform InPlatform, UnrealTargetConfiguration InConfiguration, VCProjectFileFormat InProjectFileFormat)
 		{
-			if (!IsNsightInstalled(InProjectFileFormat))
+			if (IsVSAndroidSupportInstalled() || !IsNsightInstalled(InProjectFileFormat))
 			{
 				return "\t\t<PlatformToolset>" + VCProjectFileGenerator.GetProjectFilePlatformToolsetVersionString(InProjectFileFormat) + "</PlatformToolset>" + ProjectFileGenerator.NewLine;
 			}
@@ -197,56 +248,388 @@ namespace UnrealBuildTool
 		/// <returns>The custom path lines for the project file; Empty string if it doesn't require one</returns>
 		public override string GetVisualStudioPathsEntries(UnrealTargetPlatform InPlatform, UnrealTargetConfiguration InConfiguration, TargetType TargetType, FileReference TargetRulesPath, FileReference ProjectFilePath, FileReference NMakeOutputPath, VCProjectFileFormat InProjectFileFormat)
 		{
-			if (!IsNsightInstalled(InProjectFileFormat))
+			string PathsLines = "";
+
+			if (!IsVSAndroidSupportInstalled() && IsNsightInstalled(InProjectFileFormat))
+			{
+
+				// NOTE: We are intentionally overriding defaults for these paths with empty strings.  We never want Visual Studio's
+				//       defaults for these fields to be propagated, since they are version-sensitive paths that may not reflect
+				//       the environment that UBT is building in.  We'll set these environment variables ourselves!
+				// NOTE: We don't touch 'ExecutablePath' because that would result in Visual Studio clobbering the system "Path"
+				//       environment variable
+
+				//@todo android: clean up debug path generation
+				string GameName = TargetRulesPath.GetFileNameWithoutExtension();
+				GameName = Path.GetFileNameWithoutExtension(GameName);
+
+
+				// intermediate path for Engine or Game's intermediate
+				string IntermediateDirectoryPath;
+				IntermediateDirectoryPath = Path.GetDirectoryName(NMakeOutputPath.FullName) + "/../../Intermediate/Android/APK";
+
+				// string for <OverrideAPKPath>
+				string APKPath = Path.Combine(
+					Path.GetDirectoryName(NMakeOutputPath.FullName),
+					Path.GetFileNameWithoutExtension(NMakeOutputPath.FullName) + "-armv7-es2.apk");
+
+				// string for <BuildXmlPath> and <AndroidManifestPath>
+				string BuildXmlPath = IntermediateDirectoryPath;
+				string AndroidManifestPath = Path.Combine(IntermediateDirectoryPath, "AndroidManifest.xml");
+
+				// string for <AdditionalLibraryDirectories>
+				string AdditionalLibDirs = "";
+				AdditionalLibDirs += IntermediateDirectoryPath + @"\obj\local\armeabi-v7a";
+				AdditionalLibDirs += ";" + IntermediateDirectoryPath + @"\obj\local\x86";
+				AdditionalLibDirs += @";$(AdditionalLibraryDirectories)";
+
+				PathsLines =
+					"		<IncludePath />" + ProjectFileGenerator.NewLine +
+					"		<ReferencePath />" + ProjectFileGenerator.NewLine +
+					"		<LibraryPath />" + ProjectFileGenerator.NewLine +
+					"		<LibraryWPath />" + ProjectFileGenerator.NewLine +
+					"		<SourcePath />" + ProjectFileGenerator.NewLine +
+					"		<ExcludePath />" + ProjectFileGenerator.NewLine +
+					"		<AndroidAttach>False</AndroidAttach>" + ProjectFileGenerator.NewLine +
+					"		<DebuggerFlavor>AndroidDebugger</DebuggerFlavor>" + ProjectFileGenerator.NewLine +
+					"		<OverrideAPKPath>" + APKPath + "</OverrideAPKPath>" + ProjectFileGenerator.NewLine +
+					"		<AdditionalLibraryDirectories>" + AdditionalLibDirs + "</AdditionalLibraryDirectories>" + ProjectFileGenerator.NewLine +
+					"		<BuildXmlPath>" + BuildXmlPath + "</BuildXmlPath>" + ProjectFileGenerator.NewLine +
+					"		<AndroidManifestPath>" + AndroidManifestPath + "</AndroidManifestPath>" + ProjectFileGenerator.NewLine;
+			}
+			else 
 			{
-				return base.GetVisualStudioPathsEntries(InPlatform, InConfiguration, TargetType, TargetRulesPath, ProjectFilePath, NMakeOutputPath, InProjectFileFormat);
+				PathsLines = base.GetVisualStudioPathsEntries(InPlatform, InConfiguration, TargetType, TargetRulesPath, ProjectFilePath, NMakeOutputPath, InProjectFileFormat);
 			}
 
-			// NOTE: We are intentionally overriding defaults for these paths with empty strings.  We never want Visual Studio's
-			//       defaults for these fields to be propagated, since they are version-sensitive paths that may not reflect
-			//       the environment that UBT is building in.  We'll set these environment variables ourselves!
-			// NOTE: We don't touch 'ExecutablePath' because that would result in Visual Studio clobbering the system "Path"
-			//       environment variable
+			return PathsLines;
 
-			//@todo android: clean up debug path generation
-			string GameName = TargetRulesPath.GetFileNameWithoutExtension();
-			GameName = Path.GetFileNameWithoutExtension(GameName);
+		}
+		
+		/// <summary>
+		/// Return any custom property settings. These will be included right after Global properties to make values available to all other imports.
+		/// </summary>
+		/// <param name="InPlatform">  The UnrealTargetPlatform being built</param>
+		/// <returns>string    The custom property import lines for the project file; Empty string if it doesn't require one</returns>
+		public override string GetVisualStudioGlobalProperties(UnrealTargetPlatform InPlatform)
+		{
+			if (IsVSAndroidSupportInstalled())
+			{
+				return "	<Import Project=\"$(AndroidTargetsPath)\\Android.Tools.props\"/>" + ProjectFileGenerator.NewLine;
+			}
+			return base.GetVisualStudioGlobalProperties(InPlatform);
+		}
 
+		/// <summary>
+		/// Return any custom property settings. These will be included in the ImportGroup section
+		/// </summary>
+		/// <param name="InPlatform">  The UnrealTargetPlatform being built</param>
+		/// <returns>string    The custom property import lines for the project file; Empty string if it doesn't require one</returns>
+		public override string GetVisualStudioImportGroupProperties(UnrealTargetPlatform InPlatform)
+		{
+			if (IsVSAndroidSupportInstalled())
+			{
+				return "		<Import Project=\"VSAndroidUnreal.props\" /> " + ProjectFileGenerator.NewLine;
+			}
+			return base.GetVisualStudioImportGroupProperties(InPlatform);
+		}
 
-			// intermediate path for Engine or Game's intermediate
-			string IntermediateDirectoryPath;
-			IntermediateDirectoryPath = Path.GetDirectoryName(NMakeOutputPath.FullName) + "/../../Intermediate/Android/APK";
+		/// <summary>
+		/// For Additional Project Property file  VSAndroidUnreal.props file that need to be written out.  This is currently used only on Android. 
+		/// </summary>
+		public override void WriteAdditionalPropFile()
+		{
+			if(!IsVSAndroidSupportInstalled())
+			{
+				return;
+			}
+
+			//This file only needs to be written once and doesn't change.
+			if(VSPropsFileWritten)
+			{
+				return;
+			}
+			VSPropsFileWritten = true;
+
+
+			string FileName = "VSAndroidUnreal.props";
+
+			string FileText = "<?xml version=\"1.0\" encoding=\"utf-8\"?>" + ProjectFileGenerator.NewLine +
+								"<Project ToolsVersion=\"4.0\" xmlns=\"http://schemas.microsoft.com/developer/msbuild/2003\">" + ProjectFileGenerator.NewLine +
+								"	<ImportGroup Label=\"PropertySheets\" />" + ProjectFileGenerator.NewLine +
+								"	<PropertyGroup Label=\"UserMacros\">" + ProjectFileGenerator.NewLine +
+								"		<ANDROID_HOME>$(VS_AndroidHome)</ANDROID_HOME>" + ProjectFileGenerator.NewLine +
+								"		<JAVA_HOME>$(VS_JavaHome)</JAVA_HOME>" + ProjectFileGenerator.NewLine +
+								"		<ANT_HOME>$(VS_AntHome)</ANT_HOME>" + ProjectFileGenerator.NewLine +
+								"		<NDKROOT>$(VS_NdkRoot)</NDKROOT>" + ProjectFileGenerator.NewLine +
+								"	</PropertyGroup>" + ProjectFileGenerator.NewLine +
+								"	<PropertyGroup />" + ProjectFileGenerator.NewLine +
+								"	<ItemDefinitionGroup />" + ProjectFileGenerator.NewLine +
+								"	<ItemGroup>" + ProjectFileGenerator.NewLine +
+								"		<BuildMacro Include=\"ANDROID_HOME\">" + ProjectFileGenerator.NewLine +
+								"			<Value>$(ANDROID_HOME)</Value>" + ProjectFileGenerator.NewLine +
+								"			<EnvironmentVariable>true</EnvironmentVariable>" + ProjectFileGenerator.NewLine +
+								"		</BuildMacro>" + ProjectFileGenerator.NewLine +
+								"		<BuildMacro Include=\"JAVA_HOME\">" + ProjectFileGenerator.NewLine +
+								"			<Value>$(JAVA_HOME)</Value>" + ProjectFileGenerator.NewLine +
+								"			<EnvironmentVariable>true</EnvironmentVariable>" + ProjectFileGenerator.NewLine +
+								"		</BuildMacro>" + ProjectFileGenerator.NewLine +
+								"		<BuildMacro Include=\"ANT_HOME\">" + ProjectFileGenerator.NewLine +
+								"			<Value>$(ANT_HOME)</Value>" + ProjectFileGenerator.NewLine +
+								"			<EnvironmentVariable>true</EnvironmentVariable>" + ProjectFileGenerator.NewLine +
+								"		</BuildMacro>" + ProjectFileGenerator.NewLine +
+								"		<BuildMacro Include=\"NDKROOT\">" + ProjectFileGenerator.NewLine +
+								"			<Value>$(NDKROOT)</Value>" + ProjectFileGenerator.NewLine +
+								"			<EnvironmentVariable>true</EnvironmentVariable>" + ProjectFileGenerator.NewLine +
+								"		</BuildMacro>" + ProjectFileGenerator.NewLine +
+								"	</ItemGroup>" + ProjectFileGenerator.NewLine +
+								"</Project>";
+
+			ProjectFileGenerator.WriteFileIfChanged(ProjectFileGenerator.IntermediateProjectFilesPath + "\\" + FileName, FileText);
+
+		}
+
+
+		/// <summary>
+		/// For additional Project file *PROJECTNAME*-AndroidRun.androidproj.user that needs to be written out.  This is currently used only on Android. 
+		/// </summary>
+		/// <param name="ProjectFile">ProjectFile object</param>
+		public override void WriteAdditionalProjUserFile(ProjectFile ProjectFile)
+		{
+			if (!IsVSAndroidSupportInstalled() || ProjectFile.SourceFiles.Count == 0)
+			{
+				return;
+			}
 
-			// string for <OverrideAPKPath>
-			string APKPath = Path.Combine(
-				Path.GetDirectoryName(NMakeOutputPath.FullName),
-				Path.GetFileNameWithoutExtension(NMakeOutputPath.FullName) + "-armv7-es2.apk");
+			string BaseDirectory = ProjectFile.SourceFiles[0].BaseFolder.FullName;
 
-			// string for <BuildXmlPath> and <AndroidManifestPath>
-			string BuildXmlPath = IntermediateDirectoryPath;
-			string AndroidManifestPath = Path.Combine(IntermediateDirectoryPath, "AndroidManifest.xml");
+			string ProjectName = ProjectFile.ProjectFilePath.GetFileNameWithoutExtension();
 
-			// string for <AdditionalLibraryDirectories>
-			string AdditionalLibDirs = "";
-			AdditionalLibDirs += IntermediateDirectoryPath + @"\obj\local\armeabi-v7a";
-			AdditionalLibDirs += ";" + IntermediateDirectoryPath + @"\obj\local\x86";
-			AdditionalLibDirs += @";$(AdditionalLibraryDirectories)";
+			string FileName = ProjectName + "-AndroidRun.androidproj.user";
 
-			string PathsLines =
-				"		<IncludePath />" + ProjectFileGenerator.NewLine +
-				"		<ReferencePath />" + ProjectFileGenerator.NewLine +
-				"		<LibraryPath />" + ProjectFileGenerator.NewLine +
-				"		<LibraryWPath />" + ProjectFileGenerator.NewLine +
-				"		<SourcePath />" + ProjectFileGenerator.NewLine +
-				"		<ExcludePath />" + ProjectFileGenerator.NewLine +
-				"		<AndroidAttach>False</AndroidAttach>" + ProjectFileGenerator.NewLine +
-				"		<DebuggerFlavor>AndroidDebugger</DebuggerFlavor>" + ProjectFileGenerator.NewLine +
-				"		<OverrideAPKPath>" + APKPath + "</OverrideAPKPath>" + ProjectFileGenerator.NewLine +
-				"		<AdditionalLibraryDirectories>" + AdditionalLibDirs + "</AdditionalLibraryDirectories>" + ProjectFileGenerator.NewLine +
-				"		<BuildXmlPath>" + BuildXmlPath + "</BuildXmlPath>" + ProjectFileGenerator.NewLine +
-				"		<AndroidManifestPath>" + AndroidManifestPath + "</AndroidManifestPath>" + ProjectFileGenerator.NewLine;
+			string FileText =   "<?xml version=\"1.0\" encoding=\"utf-8\"?>" + ProjectFileGenerator.NewLine +
+								"<Project ToolsVersion=\"14.0\" xmlns=\"http://schemas.microsoft.com/developer/msbuild/2003\">" + ProjectFileGenerator.NewLine +
+								"	<PropertyGroup Condition=\"\'$(Configuration)|$(Platform)\'==\'Debug|ARM\'\">" + ProjectFileGenerator.NewLine +
+								"		<PackagePath>" + BaseDirectory + "\\Binaries\\Android\\" + ProjectName + "-armv7-es2.apk</PackagePath>" + ProjectFileGenerator.NewLine +
+								"		<LaunchActivity>" + ProjectFileGenerator.NewLine +
+								"		</LaunchActivity>" + ProjectFileGenerator.NewLine +
+								"		<AdditionalSymbolSearchPaths>" + BaseDirectory + "\\Intermediate\\Android\\APK\\obj\\local\\armeabi-v7a;" + BaseDirectory + "\\Intermediate\\Android\\APK\\jni\\armeabi-v7a;" + BaseDirectory + "\\Binaries\\Android;$(AdditionalSymbolSearchPaths)</AdditionalSymbolSearchPaths>" + ProjectFileGenerator.NewLine +
+								"		<DebuggerFlavor>AndroidDebugger</DebuggerFlavor>" + ProjectFileGenerator.NewLine +
+								"	</PropertyGroup>" + ProjectFileGenerator.NewLine +
+								"</Project>";
 
-			return PathsLines;
+			ProjectFileGenerator.WriteFileIfChanged(ProjectFileGenerator.IntermediateProjectFilesPath + "\\" + FileName, FileText);
+
+		}
+
+
+		/// <summary>
+		/// For additional Project file *PROJECTNAME*-AndroidRun.androidproj that needs to be written out.  This is currently used only on Android. 
+		/// </summary>
+		/// <param name="ProjectFile">ProjectFile object</param>
+		public override Tuple<ProjectFile, string> WriteAdditionalProjFile(ProjectFile ProjectFile)
+		{
+			if (!IsVSAndroidSupportInstalled())
+			{
+				return null;
+			}
+
+			string ProjectName = ProjectFile.ProjectFilePath.GetFileNameWithoutExtension() + "-AndroidRun";
+			
+			string FileName = ProjectName + ".androidproj";
+
+			string FileText = "<?xml version=\"1.0\" encoding=\"utf-8\"?> " + ProjectFileGenerator.NewLine +
+								"<Project DefaultTargets=\"Build\" ToolsVersion=\"14.0\" xmlns=\"http://schemas.microsoft.com/developer/msbuild/2003\"> " + ProjectFileGenerator.NewLine +
+								"	<ItemGroup Label=\"ProjectConfigurations\"> " + ProjectFileGenerator.NewLine +
+								"		<ProjectConfiguration Include=\"Debug|ARM\"> " + ProjectFileGenerator.NewLine +
+								"			<Configuration>Debug</Configuration> " + ProjectFileGenerator.NewLine +
+								"			<Platform>ARM</Platform> " + ProjectFileGenerator.NewLine +
+								"		</ProjectConfiguration> " + ProjectFileGenerator.NewLine +
+								"		<ProjectConfiguration Include=\"Release|ARM\"> " + ProjectFileGenerator.NewLine +
+								"			<Configuration>Release</Configuration> " + ProjectFileGenerator.NewLine +
+								"			<Platform>ARM</Platform> " + ProjectFileGenerator.NewLine +
+								"		</ProjectConfiguration> " + ProjectFileGenerator.NewLine +
+								"		<ProjectConfiguration Include=\"Debug|ARM64\"> " + ProjectFileGenerator.NewLine +
+								"			<Configuration>Debug</Configuration> " + ProjectFileGenerator.NewLine +
+								"			<Platform>ARM64</Platform> " + ProjectFileGenerator.NewLine +
+								"		</ProjectConfiguration> " + ProjectFileGenerator.NewLine +
+								"		<ProjectConfiguration Include=\"Release|ARM64\"> " + ProjectFileGenerator.NewLine +
+								"			<Configuration>Release</Configuration> " + ProjectFileGenerator.NewLine +
+								"			<Platform>ARM64</Platform> " + ProjectFileGenerator.NewLine +
+								"		</ProjectConfiguration> " + ProjectFileGenerator.NewLine +
+								"		<ProjectConfiguration Include=\"Debug|x64\"> " + ProjectFileGenerator.NewLine +
+								"			<Configuration>Debug</Configuration> " + ProjectFileGenerator.NewLine +
+								"			<Platform>x64</Platform> " + ProjectFileGenerator.NewLine +
+								"		</ProjectConfiguration> " + ProjectFileGenerator.NewLine +
+								"		<ProjectConfiguration Include=\"Release|x64\"> " + ProjectFileGenerator.NewLine +
+								"			<Configuration>Release</Configuration> " + ProjectFileGenerator.NewLine +
+								"			<Platform>x64</Platform> " + ProjectFileGenerator.NewLine +
+								"		</ProjectConfiguration> " + ProjectFileGenerator.NewLine +
+								"		<ProjectConfiguration Include=\"Debug|x86\"> " + ProjectFileGenerator.NewLine +
+								"			<Configuration>Debug</Configuration> " + ProjectFileGenerator.NewLine +
+								"			<Platform>x86</Platform> " + ProjectFileGenerator.NewLine +
+								"		</ProjectConfiguration> " + ProjectFileGenerator.NewLine +
+								"		<ProjectConfiguration Include=\"Release|x86\"> " + ProjectFileGenerator.NewLine +
+								"			<Configuration>Release</Configuration> " + ProjectFileGenerator.NewLine +
+								"			<Platform>x86</Platform> " + ProjectFileGenerator.NewLine +
+								"		</ProjectConfiguration> " + ProjectFileGenerator.NewLine +
+								"	</ItemGroup> " + ProjectFileGenerator.NewLine +
+								"	<PropertyGroup Label=\"Globals\"> " + ProjectFileGenerator.NewLine +
+								"		<RootNamespace>" + ProjectName + "</RootNamespace> " + ProjectFileGenerator.NewLine +
+								"		<MinimumVisualStudioVersion>14.0</MinimumVisualStudioVersion> " + ProjectFileGenerator.NewLine +
+								"		<ProjectVersion>1.0</ProjectVersion> " + ProjectFileGenerator.NewLine +
+								
+								//Set the project guid 
+								"		<ProjectGuid>" + System.Guid.NewGuid().ToString("B").ToUpper() + "</ProjectGuid> " + ProjectFileGenerator.NewLine +
+								
+								"		<ConfigurationType>Application</ConfigurationType> " + ProjectFileGenerator.NewLine +
+								"		<_PackagingProjectWithoutNativeComponent>true</_PackagingProjectWithoutNativeComponent> " + ProjectFileGenerator.NewLine +
+								"		<LaunchActivity Condition=\"\'$(LaunchActivity)\' == \'\'\">com." + ProjectName + "." + ProjectName + "</LaunchActivity> " + ProjectFileGenerator.NewLine +
+								"		<JavaSourceRoots>src</JavaSourceRoots> " + ProjectFileGenerator.NewLine +
+								"	</PropertyGroup> " + ProjectFileGenerator.NewLine +
+								"	<Import Project=\"$(AndroidTargetsPath)\\Android.Default.props\" /> " + ProjectFileGenerator.NewLine +
+								"	<PropertyGroup Condition=\"\'$(Configuration)|$(Platform)\'==\'Debug|ARM\'\" Label=\"Configuration\"> " + ProjectFileGenerator.NewLine +
+								"		<UseDebugLibraries>true</UseDebugLibraries> " + ProjectFileGenerator.NewLine +
+								"		<TargetName>$(RootNamespace)</TargetName> " + ProjectFileGenerator.NewLine +
+								"	</PropertyGroup> " + ProjectFileGenerator.NewLine +
+								"	<PropertyGroup Condition=\"\'$(Configuration)|$(Platform)\'==\'Release|ARM\'\" Label=\"Configuration\"> " + ProjectFileGenerator.NewLine +
+								"		<UseDebugLibraries>false</UseDebugLibraries> " + ProjectFileGenerator.NewLine +
+								"		<TargetName>$(RootNamespace)</TargetName> " + ProjectFileGenerator.NewLine +
+								"	</PropertyGroup> " + ProjectFileGenerator.NewLine +
+								"	<PropertyGroup Condition=\"\'$(Configuration)|$(Platform)\'==\'Debug|ARM64\'\" Label=\"Configuration\"> " + ProjectFileGenerator.NewLine +
+								"		<UseDebugLibraries>true</UseDebugLibraries> " + ProjectFileGenerator.NewLine +
+								"		<TargetName>$(RootNamespace)</TargetName> " + ProjectFileGenerator.NewLine +
+								"	</PropertyGroup> " + ProjectFileGenerator.NewLine +
+								"	<PropertyGroup Condition=\"\'$(Configuration)|$(Platform)\'==\'Release|ARM64\'\" Label=\"Configuration\"> " + ProjectFileGenerator.NewLine +
+								"		<UseDebugLibraries>false</UseDebugLibraries> " + ProjectFileGenerator.NewLine +
+								"		<TargetName>$(RootNamespace)</TargetName> " + ProjectFileGenerator.NewLine +
+								"	</PropertyGroup> " + ProjectFileGenerator.NewLine +
+								"	<PropertyGroup Condition=\"\'$(Configuration)|$(Platform)\'==\'Debug|x64\'\" Label=\"Configuration\"> " + ProjectFileGenerator.NewLine +
+								"		<UseDebugLibraries>true</UseDebugLibraries> " + ProjectFileGenerator.NewLine +
+								"		<TargetName>$(RootNamespace)</TargetName> " + ProjectFileGenerator.NewLine +
+								"	</PropertyGroup> " + ProjectFileGenerator.NewLine +
+								"	<PropertyGroup Condition=\"\'$(Configuration)|$(Platform)\'==\'Release|x64\'\" Label=\"Configuration\"> " + ProjectFileGenerator.NewLine +
+								"		<UseDebugLibraries>false</UseDebugLibraries> " + ProjectFileGenerator.NewLine +
+								"		<TargetName>$(RootNamespace)</TargetName> " + ProjectFileGenerator.NewLine +
+								"	</PropertyGroup> " + ProjectFileGenerator.NewLine +
+								"	<PropertyGroup Condition=\"\'$(Configuration)|$(Platform)\'==\'Debug|x86\'\" Label=\"Configuration\"> " + ProjectFileGenerator.NewLine +
+								"		<UseDebugLibraries>true</UseDebugLibraries> " + ProjectFileGenerator.NewLine +
+								"		<TargetName>$(RootNamespace)</TargetName> " + ProjectFileGenerator.NewLine +
+								"	</PropertyGroup> " + ProjectFileGenerator.NewLine +
+								"	<PropertyGroup Condition=\"\'$(Configuration)|$(Platform)\'==\'Release|x86\'\" Label=\"Configuration\"> " + ProjectFileGenerator.NewLine +
+								"		<UseDebugLibraries>false</UseDebugLibraries> " + ProjectFileGenerator.NewLine +
+								"		<TargetName>$(RootNamespace)</TargetName> " + ProjectFileGenerator.NewLine +
+								"	</PropertyGroup> " + ProjectFileGenerator.NewLine +
+								"	<Import Project=\"$(AndroidTargetsPath)\\Android.props\" /> " + ProjectFileGenerator.NewLine +
+								"	<ImportGroup Label=\"ExtensionSettings\" /> " + ProjectFileGenerator.NewLine +
+								"	<PropertyGroup Label=\"UserMacros\" /> " + ProjectFileGenerator.NewLine +
+								"	<PropertyGroup Condition=\"\'$(Configuration)|$(Platform)\'==\'Debug|ARM\'\"> " + ProjectFileGenerator.NewLine +
+								"	</PropertyGroup> " + ProjectFileGenerator.NewLine +
+								"	<PropertyGroup Condition=\"\'$(Configuration)|$(Platform)\'==\'Release|ARM\'\"> " + ProjectFileGenerator.NewLine +
+								"	</PropertyGroup> " + ProjectFileGenerator.NewLine +
+								"	<PropertyGroup Condition=\"\'$(Configuration)|$(Platform)\'==\'Debug|ARM64\'\"> " + ProjectFileGenerator.NewLine +
+								"	</PropertyGroup> " + ProjectFileGenerator.NewLine +
+								"	<PropertyGroup Condition=\"\'$(Configuration)|$(Platform)\'==\'Release|ARM64\'\"> " + ProjectFileGenerator.NewLine +
+								"	</PropertyGroup> " + ProjectFileGenerator.NewLine +
+								"	<PropertyGroup Condition=\"\'$(Configuration)|$(Platform)\'==\'Debug|x64\'\"> " + ProjectFileGenerator.NewLine +
+								"	</PropertyGroup> " + ProjectFileGenerator.NewLine +
+								"	<PropertyGroup Condition=\"\'$(Configuration)|$(Platform)\'==\'Release|x64\'\"> " + ProjectFileGenerator.NewLine +
+								"	</PropertyGroup> " + ProjectFileGenerator.NewLine +
+								"	<PropertyGroup Condition=\"\'$(Configuration)|$(Platform)\'==\'Debug|x86\'\"> " + ProjectFileGenerator.NewLine +
+								"	</PropertyGroup> " + ProjectFileGenerator.NewLine +
+								"	<PropertyGroup Condition=\"\'$(Configuration)|$(Platform)\'==\'Release|x86\'\"> " + ProjectFileGenerator.NewLine +
+								"	</PropertyGroup> " + ProjectFileGenerator.NewLine +
+								"	<ItemDefinitionGroup Condition=\"\'$(Configuration)|$(Platform)\'==\'Debug|ARM\'\"> " + ProjectFileGenerator.NewLine +
+								"		<AntPackage> " + ProjectFileGenerator.NewLine +
+								"			<AndroidAppLibName /> " + ProjectFileGenerator.NewLine +
+								"		</AntPackage> " + ProjectFileGenerator.NewLine +
+								"	</ItemDefinitionGroup> " + ProjectFileGenerator.NewLine +
+								"	<ItemDefinitionGroup Condition=\"\'$(Configuration)|$(Platform)\'==\'Release|ARM\'\"> " + ProjectFileGenerator.NewLine +
+								"		<AntPackage> " + ProjectFileGenerator.NewLine +
+								"			<AndroidAppLibName /> " + ProjectFileGenerator.NewLine +
+								"		</AntPackage> " + ProjectFileGenerator.NewLine +
+								"	</ItemDefinitionGroup> " + ProjectFileGenerator.NewLine +
+								"	<ItemDefinitionGroup Condition=\"\'$(Configuration)|$(Platform)\'==\'Debug|ARM64\'\"> " + ProjectFileGenerator.NewLine +
+								"		<AntPackage> " + ProjectFileGenerator.NewLine +
+								"			<AndroidAppLibName /> " + ProjectFileGenerator.NewLine +
+								"		</AntPackage> " + ProjectFileGenerator.NewLine +
+								"	</ItemDefinitionGroup> " + ProjectFileGenerator.NewLine +
+								"	<ItemDefinitionGroup Condition=\"\'$(Configuration)|$(Platform)\'==\'Release|ARM64\'\"> " + ProjectFileGenerator.NewLine +
+								"		<AntPackage> " + ProjectFileGenerator.NewLine +
+								"			<AndroidAppLibName /> " + ProjectFileGenerator.NewLine +
+								"		</AntPackage> " + ProjectFileGenerator.NewLine +
+								"	</ItemDefinitionGroup> " + ProjectFileGenerator.NewLine +
+								"	<ItemDefinitionGroup Condition=\"\'$(Configuration)|$(Platform)\'==\'Debug|x64\'\"> " + ProjectFileGenerator.NewLine +
+								"		<AntPackage> " + ProjectFileGenerator.NewLine +
+								"			<AndroidAppLibName /> " + ProjectFileGenerator.NewLine +
+								"		</AntPackage> " + ProjectFileGenerator.NewLine +
+								"	</ItemDefinitionGroup> " + ProjectFileGenerator.NewLine +
+								"	<ItemDefinitionGroup Condition=\"\'$(Configuration)|$(Platform)\'==\'Release|x64\'\"> " + ProjectFileGenerator.NewLine +
+								"		<AntPackage> " + ProjectFileGenerator.NewLine +
+								"			<AndroidAppLibName /> " + ProjectFileGenerator.NewLine +
+								"		</AntPackage> " + ProjectFileGenerator.NewLine +
+								"	</ItemDefinitionGroup> " + ProjectFileGenerator.NewLine +
+								"	<ItemDefinitionGroup Condition=\"\'$(Configuration)|$(Platform)\'==\'Debug|x86\'\"> " + ProjectFileGenerator.NewLine +
+								"		<AntPackage> " + ProjectFileGenerator.NewLine +
+								"			<AndroidAppLibName /> " + ProjectFileGenerator.NewLine +
+								"		</AntPackage> " + ProjectFileGenerator.NewLine +
+								"	</ItemDefinitionGroup> " + ProjectFileGenerator.NewLine +
+								"	<ItemDefinitionGroup Condition=\"\'$(Configuration)|$(Platform)\'==\'Release|x86\'\"> " + ProjectFileGenerator.NewLine +
+								"		<AntPackage> " + ProjectFileGenerator.NewLine +
+								"			<AndroidAppLibName /> " + ProjectFileGenerator.NewLine +
+								"		</AntPackage> " + ProjectFileGenerator.NewLine +
+								"	</ItemDefinitionGroup> " + ProjectFileGenerator.NewLine +
+								"	<Import Project=\"$(AndroidTargetsPath)\\Android.targets\" /> " + ProjectFileGenerator.NewLine +
+								"	<ImportGroup Label=\"ExtensionTargets\" /> " + ProjectFileGenerator.NewLine +
+								"</Project>";
+			
+			bool Success = ProjectFileGenerator.WriteFileIfChanged(ProjectFileGenerator.IntermediateProjectFilesPath + "\\" + FileName, FileText);
+
+			FileReference ProjectFilePath = FileReference.Combine(ProjectFileGenerator.IntermediateProjectFilesPath, FileName);
+			AndroidDebugProjectFile Project = new AndroidDebugProjectFile(ProjectFilePath);
+			Project.ShouldBuildForAllSolutionTargets = false;
+			Project.ShouldBuildByDefaultForSolutionTargets = false;
+
+			return Success ? new Tuple<ProjectFile, string>(Project, "UE4 Android Debug Projects") : null;
+
+		}
+
+	}
+
+
+	/// <summary>
+	/// An Android Debug Project
+	/// </summary>
+	class AndroidDebugProjectFile : MSBuildProjectFile
+	{
+		/// <summary>
+		/// Constructs a new project file object
+		/// </summary>
+		/// <param name="InitFilePath">The path to the project file on disk</param>
+		public AndroidDebugProjectFile(FileReference InitFilePath)
+			: base(InitFilePath)
+		{
+		}
+
+		// This is the GUID that Visual Studio uses to identify an Android project file in the solution
+		public override string ProjectTypeGUID
+		{
+			get { return "{39E2626F-3545-4960-A6E8-258AD8476CE5}"; }
 		}
+
+		/// <summary>
+		/// The only valid configuration for these to be run in is Debug|ARM
+		/// </summary>
+		/// <param name="Platform">Actual platform</param>
+		/// <param name="Configuration">Actual configuration</param>
+		/// <param name="TargetConfigurationName">The configuration name from the target rules, or null if we don't have one</param>
+		/// <param name="ProjectPlatformName">Name of platform string to use for Visual Studio project</param>
+		/// <param name="ProjectConfigurationName">Name of configuration string to use for Visual Studio project</param>
+		public override void MakeProjectPlatformAndConfigurationNames(UnrealTargetPlatform Platform, UnrealTargetConfiguration Configuration, string TargetConfigurationName, out string ProjectPlatformName, out string ProjectConfigurationName)
+		{
+			ProjectConfigurationName = UnrealTargetConfiguration.Debug.ToString();
+			ProjectPlatformName = "ARM";
+		}
+
 	}
 }
diff --git a/Engine/Source/Programs/UnrealBuildTool/Android/AndroidToolChain.cs b/Engine/Source/Programs/UnrealBuildTool/Android/AndroidToolChain.cs
--- a/Engine/Source/Programs/UnrealBuildTool/Android/AndroidToolChain.cs
+++ b/Engine/Source/Programs/UnrealBuildTool/Android/AndroidToolChain.cs
@@ -1335,27 +1335,32 @@ namespace UnrealBuildTool
 			ProcessStartInfo StartInfo = new ProcessStartInfo();
 			if (SourceFileName.Contains("-armv7"))
 			{
-				StartInfo.FileName = ArPathArm.Replace("-ar.exe", "-strip.exe");
+				StartInfo.FileName = ArPathArm;
 			}
 			else
 			if (SourceFileName.Contains("-arm64"))
             {
-				StartInfo.FileName = ArPathArm64.Replace("-ar.exe", "-strip.exe");
+				StartInfo.FileName = ArPathArm64;
 			}
 			else
 			if (SourceFileName.Contains("-x86"))
             {
-				StartInfo.FileName = ArPathx86.Replace("-ar.exe", "-strip.exe");
+				StartInfo.FileName = ArPathx86;
 			}
 			else
 			if (SourceFileName.Contains("-x64"))
             {
-				StartInfo.FileName = ArPathx64.Replace("-ar.exe", "-strip.exe");
+				StartInfo.FileName = ArPathx64;
 			}
 			else
 			{
 				throw new BuildException("Couldn't determine Android architecture to strip symbols from {0}", SourceFileName);
 			}
+
+			// fix the executable (replace the last -ar with -strip and keep any extension)
+			int ArIndex = StartInfo.FileName.LastIndexOf("-ar");
+			StartInfo.FileName = StartInfo.FileName.Substring(0, ArIndex) + "-strip" + StartInfo.FileName.Substring(ArIndex + 3);
+
 			StartInfo.Arguments = "--strip-debug " + TargetFileName;
 			StartInfo.UseShellExecute = false;
 			StartInfo.CreateNoWindow = true;
diff --git a/Engine/Source/Programs/UnrealBuildTool/Android/UEDeployAndroid.cs b/Engine/Source/Programs/UnrealBuildTool/Android/UEDeployAndroid.cs
--- a/Engine/Source/Programs/UnrealBuildTool/Android/UEDeployAndroid.cs
+++ b/Engine/Source/Programs/UnrealBuildTool/Android/UEDeployAndroid.cs
@@ -1668,11 +1668,11 @@ namespace UnrealBuildTool
 				Text.AppendLine("\t<uses-permission android:name=\"android.permission.WRITE_EXTERNAL_STORAGE\"/>");
 				Text.AppendLine("\t<uses-permission android:name=\"android.permission.ACCESS_NETWORK_STATE\"/>");
 				Text.AppendLine("\t<uses-permission android:name=\"android.permission.WAKE_LOCK\"/>");
-				Text.AppendLine("\t<uses-permission android:name=\"android.permission.READ_PHONE_STATE\"/>");
+			//	Text.AppendLine("\t<uses-permission android:name=\"android.permission.READ_PHONE_STATE\"/>");
 				Text.AppendLine("\t<uses-permission android:name=\"com.android.vending.CHECK_LICENSE\"/>");
 				Text.AppendLine("\t<uses-permission android:name=\"android.permission.ACCESS_WIFI_STATE\"/>");
 				Text.AppendLine("\t<uses-permission android:name=\"android.permission.MODIFY_AUDIO_SETTINGS\"/>");
-				Text.AppendLine("\t<uses-permission android:name=\"android.permission.GET_ACCOUNTS\"/>");
+			//	Text.AppendLine("\t<uses-permission android:name=\"android.permission.GET_ACCOUNTS\"/>");
 				Text.AppendLine("\t<uses-permission android:name=\"android.permission.VIBRATE\"/>");
 				//			Text.AppendLine("\t<uses-permission android:name=\"android.permission.DISABLE_KEYGUARD\"/>");
 
@@ -2583,7 +2583,8 @@ namespace UnrealBuildTool
 				{ "//$${gameActivityOnStopAdditions}$$", UPL.ProcessPluginNode(NDKArch, "gameActivityOnStopAdditions", "")},
 				{ "//$${gameActivityOnPauseAdditions}$$", UPL.ProcessPluginNode(NDKArch, "gameActivityOnPauseAdditions", "")},
 				{ "//$${gameActivityOnResumeAdditions}$$", UPL.ProcessPluginNode(NDKArch, "gameActivityOnResumeAdditions", "")},
-				{ "//$${gameActivityOnActivityResultAdditions}$$", UPL.ProcessPluginNode(NDKArch, "gameActivityOnActivityResultAdditions", "")},
+				{ "//$${gameActivityOnNewIntentAdditions}$$", UPL.ProcessPluginNode(NDKArch, "gameActivityOnNewIntentAdditions", "")},
+  				{ "//$${gameActivityOnActivityResultAdditions}$$", UPL.ProcessPluginNode(NDKArch, "gameActivityOnActivityResultAdditions", "")},
 				{ "//$${soLoadLibrary}$$", UPL.ProcessPluginNode(NDKArch, "soLoadLibrary", LoadLibraryDefaults)}
 			};
 
diff --git a/Engine/Source/Programs/UnrealBuildTool/Configuration/UEBuildModuleCPP.cs b/Engine/Source/Programs/UnrealBuildTool/Configuration/UEBuildModuleCPP.cs
--- a/Engine/Source/Programs/UnrealBuildTool/Configuration/UEBuildModuleCPP.cs
+++ b/Engine/Source/Programs/UnrealBuildTool/Configuration/UEBuildModuleCPP.cs
@@ -772,7 +772,7 @@ namespace UnrealBuildTool
 			// for header dependencies
 			CPPFile.CachedIncludePaths = ModuleCompileEnvironment.IncludePaths;
 
-			FileItem PCHFile = ModuleCompileEnvironment.Headers.CachePCHUsageForCPPFile(CPPFile, ModuleCompileEnvironment.IncludePaths);
+			FileItem PCHFile = ModuleCompileEnvironment.Headers.CachePCHUsageForCPPFile(CPPFile, ModuleCompileEnvironment.IncludePaths, ModuleCompileEnvironment.Platform);
 
 			if (UnrealBuildTool.bPrintPerformanceInfo)
 			{
@@ -849,7 +849,7 @@ namespace UnrealBuildTool
 						CPPFile.CachedIncludePaths = ModuleCompileEnvironment.IncludePaths;
 
 						// Find headers used by the source file.
-						FileItem PCH = ModuleCompileEnvironment.Headers.CachePCHUsageForCPPFile(CPPFile, ModuleCompileEnvironment.IncludePaths);
+						FileItem PCH = ModuleCompileEnvironment.Headers.CachePCHUsageForCPPFile(CPPFile, ModuleCompileEnvironment.IncludePaths, ModuleCompileEnvironment.Platform);
 						if (PCH == null)
 						{
 							throw new BuildException("Source file \"{0}\" is not including any headers.  We expect all modules to include a header file for precompiled header generation.  Please add an #include statement.", CPPFile.AbsolutePath);
diff --git a/Engine/Source/Programs/UnrealBuildTool/HTML5/HTML5SDKInfo.cs b/Engine/Source/Programs/UnrealBuildTool/HTML5/HTML5SDKInfo.cs
--- a/Engine/Source/Programs/UnrealBuildTool/HTML5/HTML5SDKInfo.cs
+++ b/Engine/Source/Programs/UnrealBuildTool/HTML5/HTML5SDKInfo.cs
@@ -404,6 +404,34 @@ namespace UnrealBuildTool
 			return Directory.Exists(EMSCRIPTEN_ROOT) && File.Exists(NODE_JS) && Directory.Exists(LLVM_ROOT) && File.Exists(PYTHON);
 		}
 
+		/// <summary>
+		/// 
+		/// </summary>
+		/// <param name="ConfigCache"></param>
+		/// <param name="BuildType"></param>
+		/// <returns></returns>
+		public static int HeapSize(ConfigHierarchy ConfigCache, string BuildType)
+		{
+			int ConfigHeapSize = 0;
+
+			// Valuer set by Editor UI
+			var bGotHeapSize = ConfigCache.GetInt32("/Script/HTML5PlatformEditor.HTML5TargetSettings", "HeapSize" + BuildType, out ConfigHeapSize);
+	
+			// Fallback if the previous method failed
+			if (!bGotHeapSize && !ConfigCache.GetInt32("/Script/BuildSettings.BuildSettings", "HeapSize" + BuildType, out ConfigHeapSize))
+			{
+				// we couldn't find a per config heap size, look for a common one.
+				if (!ConfigCache.GetInt32("/Script/BuildSettings.BuildSettings", "HeapSize", out ConfigHeapSize))
+				{
+					ConfigHeapSize = BuildType == "Development" ? 1024 : 512;
+					Log.TraceInformation("Could not find Heap Size setting in .ini for Client config {0}", BuildType);
+				}
+			}
+	
+			Log.TraceInformation("Setting Heap size to {0} Mb ", ConfigHeapSize);
+			return ConfigHeapSize;
+		}
+
 		// this script is used at:
 		// HTML5ToolChain.cs
 		// UEBuildHTML5.cs
diff --git a/Engine/Source/Programs/UnrealBuildTool/HTML5/HTML5ToolChain.cs b/Engine/Source/Programs/UnrealBuildTool/HTML5/HTML5ToolChain.cs
--- a/Engine/Source/Programs/UnrealBuildTool/HTML5/HTML5ToolChain.cs
+++ b/Engine/Source/Programs/UnrealBuildTool/HTML5/HTML5ToolChain.cs
@@ -83,9 +83,12 @@ namespace UnrealBuildTool
 			// JavsScript option overrides (see src/settings.js)
 
 			// we have to specify the full amount of memory with Asm.JS (1.5 G)
-			// I wonder if there's a per game way to change this.
-			int TotalMemory = 256 * 1024 * 1024;
-			Result += " -s TOTAL_MEMORY=" + TotalMemory.ToString();
+//			ConfigHierarchy cc = ConfigCache.ReadHierarchy(ConfigHierarchyType.Engine, DirectoryReference.FromFile(ProjectFile), UnrealTargetPlatform.HTML5);
+//			int TotalMemory = HTML5SDKInfo.HeapSize(cc, TargetConfiguration.ToString()) * 1024 * 1024;
+//			Result += " -s TOTAL_MEMORY=" + TotalMemory.ToString();
+			// and, have to specify size smaller than used and must use "grow memory" !!!
+			Result += " -s TOTAL_MEMORY=" + 128 * 1024 * 1024;
+			Result += " -s ALLOW_MEMORY_GROWTH=1";
 
 			// no need for exceptions
 			Result += " -s DISABLE_EXCEPTION_CATCHING=1";
@@ -192,7 +195,7 @@ namespace UnrealBuildTool
 			return Result;
 		}
 
-		static string GetLinkArguments(LinkEnvironment LinkEnvironment)
+		string GetLinkArguments(LinkEnvironment LinkEnvironment)
 		{
 			string Result = GetSharedArguments_Global(LinkEnvironment.Configuration, LinkEnvironment.Architecture, false, false, false, false);
 
diff --git a/Engine/Source/Programs/UnrealBuildTool/Linux/UEBuildLinux.cs b/Engine/Source/Programs/UnrealBuildTool/Linux/UEBuildLinux.cs
--- a/Engine/Source/Programs/UnrealBuildTool/Linux/UEBuildLinux.cs
+++ b/Engine/Source/Programs/UnrealBuildTool/Linux/UEBuildLinux.cs
@@ -64,6 +64,14 @@ namespace UnrealBuildTool
 					Rules.DynamicallyLoadedModuleNames.Add("ShaderFormatOpenGL");
 				}
 			}
+			else if (ModuleName == "Launch")
+			{
+				// this is a hack to influence symbol resolution on Linux that results in global delete being called from within CEF
+				if (Target.LinkType != TargetLinkType.Monolithic && Target.bCompileCEF3)
+				{
+					Rules.AddEngineThirdPartyPrivateStaticDependencies(Target, "CEF3");
+				}
+			}
 		}
 
 		/// <summary>
diff --git a/Engine/Source/Programs/UnrealBuildTool/System/CPPHeaders.cs b/Engine/Source/Programs/UnrealBuildTool/System/CPPHeaders.cs
--- a/Engine/Source/Programs/UnrealBuildTool/System/CPPHeaders.cs
+++ b/Engine/Source/Programs/UnrealBuildTool/System/CPPHeaders.cs
@@ -453,7 +453,7 @@ namespace UnrealBuildTool
 			}
 		}
 
-		public FileItem CachePCHUsageForCPPFile(FileItem CPPFile, CppIncludePaths IncludePaths)
+		public FileItem CachePCHUsageForCPPFile(FileItem CPPFile, CppIncludePaths IncludePaths, CppPlatform Platform)
 		{
 			// @todo ubtmake: We don't really need to scan every file looking for PCH headers, just need one.  The rest is just for error checking.
 			// @todo ubtmake: We don't need all of the direct includes either.  We just need the first, unless we want to check for errors.
@@ -483,10 +483,19 @@ namespace UnrealBuildTool
 			}
 
 			// search the include paths to resolve the file.
-			FileItem PrecompiledHeaderIncludeFile = CPPHeaders.FindIncludedFile(CPPFile.Reference, FirstInclude.IncludeName, IncludePaths);
+			string FirstIncludeName = FirstInclude.IncludeName;
+			UEBuildPlatform BuildPlatform = UEBuildPlatform.GetBuildPlatformForCPPTargetPlatform(Platform);
+			// convert back from relative to host path if needed
+			if (!BuildPlatform.UseAbsolutePathsInUnityFiles())
+			{
+				FirstIncludeName = RemoteExports.UnconvertPath(FirstIncludeName);
+			}
+
+			FileItem PrecompiledHeaderIncludeFile = CPPHeaders.FindIncludedFile(CPPFile.Reference, FirstIncludeName, IncludePaths);
 			if (PrecompiledHeaderIncludeFile == null)
 			{
-				throw new BuildException("The first include statement in source file '{0}' is trying to include the file '{1}' as the precompiled header, but that file could not be located in any of the module's include search paths.", CPPFile.AbsolutePath, FirstInclude.IncludeName);
+				FirstIncludeName = RemoteExports.UnconvertPath(FirstInclude.IncludeName);
+				throw new BuildException("The first include statement in source file '{0}' is trying to include the file '{1}' as the precompiled header, but that file could not be located in any of the module's include search paths.", CPPFile.AbsolutePath, FirstIncludeName);
 			}
 
 			IncludeDependencyCache.CacheResolvedIncludeFullPath(CPPFile, 0, PrecompiledHeaderIncludeFile.Reference, bUseIncludeDependencyResolveCache, bTestIncludeDependencyResolveCache);
diff --git a/Engine/Source/Programs/UnrealBuildTool/System/CodeLiteProject.cs b/Engine/Source/Programs/UnrealBuildTool/System/CodeLiteProject.cs
--- a/Engine/Source/Programs/UnrealBuildTool/System/CodeLiteProject.cs
+++ b/Engine/Source/Programs/UnrealBuildTool/System/CodeLiteProject.cs
@@ -115,6 +115,7 @@ namespace UnrealBuildTool
 					// TODO It seems that the full pathname doesn't work for some files like .ini, .usf 
 					if ((ProjectTargetType == TargetType.Client) ||
 						(ProjectTargetType == TargetType.Editor) ||
+						(ProjectTargetType == TargetType.Game) ||
 						(ProjectTargetType == TargetType.Server) )
 					{
 						if(ProjectName.Contains("UE4"))
@@ -142,16 +143,9 @@ namespace UnrealBuildTool
 						int Idx = Path.GetDirectoryName(CurrentFile.Reference.FullName).IndexOf(ProjectName) + ProjectName.Length;
 						CurrentFilePath = Path.GetDirectoryName(CurrentFile.Reference.FullName).Substring(Idx);
 					}
-					else if (ProjectTargetType == TargetType.Game)
-					{
-//						int lengthOfProjectRootPath = Path.GetFullPath(ProjectFileGenerator.MasterProjectRelativePath).Length;
-//						CurrentFilePath = Path.GetDirectoryName(Path.GetFullPath(CurrentFile.FilePath)).Substring(lengthOfProjectRootPath);
-					//	int lengthOfProjectRootPath = EngineRootDirectory.Length;
-						int Idx = Path.GetDirectoryName(CurrentFile.Reference.FullName).IndexOf(ProjectName) + ProjectName.Length;
-						CurrentFilePath = Path.GetDirectoryName(CurrentFile.Reference.FullName).Substring(Idx);
-					}
 
-					string [] SplitFolders = CurrentFilePath.Split('/');
+					char[] Delimiters = new char[] { '/', '\\' };
+					string [] SplitFolders = CurrentFilePath.Split(Delimiters, StringSplitOptions.RemoveEmptyEntries);
 					//
 					// Set the CodeLite root folder again.
 					//
@@ -256,7 +250,8 @@ namespace UnrealBuildTool
 							case UnrealTargetPlatform.Linux:
 							{
 								ExecutableToRun = "./" + ProjectName;
-								if ( (ProjectTargetType == TargetType.Game) || (ProjectTargetType == TargetType.Program))
+								if ((ProjectTargetType == TargetType.Game) || 
+									(ProjectTargetType == TargetType.Program))
 								{
 									if (CurConf != UnrealTargetConfiguration.Development)
 									{
@@ -266,6 +261,12 @@ namespace UnrealBuildTool
 								else if (ProjectTargetType == TargetType.Editor)
 								{
 									ExecutableToRun = "./UE4Editor";
+									if ((CurConf == UnrealTargetConfiguration.Debug) || 
+										(CurConf == UnrealTargetConfiguration.Shipping) || 
+										(CurConf == UnrealTargetConfiguration.Test))
+									{
+										ExecutableToRun += PlatformConfiguration;
+									}
 								}
 
 							}
@@ -290,12 +291,14 @@ namespace UnrealBuildTool
 								else if (ProjectTargetType == TargetType.Editor)
 								{
 									ExecutableToRun = "./UE4Editor";
-									if (CurConf != UnrealTargetConfiguration.Development)
+									if ((CurConf == UnrealTargetConfiguration.Debug) || 
+										(CurConf == UnrealTargetConfiguration.Shipping) || 
+										(CurConf == UnrealTargetConfiguration.Test))
 									{
 										ExecutableToRun += PlatformConfiguration;
 									}
 									ExecutableToRun += ".app/Contents/MacOS/UE4Editor";
-									if (CurConf != UnrealTargetConfiguration.Development)
+									if ((CurConf != UnrealTargetConfiguration.Development) && (CurConf != UnrealTargetConfiguration.DebugGame))
 									{
 										ExecutableToRun += PlatformConfiguration;
 									}
@@ -318,6 +321,12 @@ namespace UnrealBuildTool
 								else if (ProjectTargetType == TargetType.Editor)
 								{
 									ExecutableToRun = "UE4Editor";
+									if ((CurConf == UnrealTargetConfiguration.Debug) || 
+										(CurConf == UnrealTargetConfiguration.Shipping) || 
+										(CurConf == UnrealTargetConfiguration.Test))
+									{
+										ExecutableToRun += PlatformConfiguration;
+									}
 								}
 
 								ExecutableToRun += ".exe";
@@ -353,6 +362,10 @@ namespace UnrealBuildTool
 							if (ProjectName != "UE4Editor" && GameProjectFile != "")
 							{
 								string commandArguments = "\"" + GameProjectFile + "\"" + " -game";
+								if (CurConf.ToString ().Contains ("Debug")) 
+								{
+									commandArguments += " -debug";
+								}
 								XAttribute CommandArguments = new XAttribute("CommandArguments", commandArguments);
 								CodeLiteConfigurationGeneral.Add(CommandArguments);
 							}
diff --git a/Engine/Source/Programs/UnrealBuildTool/System/MakefileGenerator.cs b/Engine/Source/Programs/UnrealBuildTool/System/MakefileGenerator.cs
--- a/Engine/Source/Programs/UnrealBuildTool/System/MakefileGenerator.cs
+++ b/Engine/Source/Programs/UnrealBuildTool/System/MakefileGenerator.cs
@@ -125,7 +125,7 @@ namespace UnrealBuildTool
 
 			MakefileContent.Append("\n\n" + BuildCommand + ProjectBuildCommand + "\n" +
 				"all: StandardSet\n\n" +
-				"RequiredTools: CrashReportClient-Linux-Shipping ShaderCompileWorker UnrealPak UnrealLightmass\n\n" +
+				"RequiredTools: CrashReportClient-Linux-Shipping ShaderCompileWorker UnrealLightmass\n\n" +
 				"StandardSet: RequiredTools UnrealFrontend UE4Editor\n\n" +
 				"DebugSet: RequiredTools UnrealFrontend-Linux-Debug UE4Editor-Linux-Debug\n\n"
 			);
diff --git a/Engine/Source/Programs/UnrealBuildTool/System/Project.cs b/Engine/Source/Programs/UnrealBuildTool/System/Project.cs
--- a/Engine/Source/Programs/UnrealBuildTool/System/Project.cs
+++ b/Engine/Source/Programs/UnrealBuildTool/System/Project.cs
@@ -394,6 +394,18 @@ namespace UnrealBuildTool
 			throw new BuildException("UnrealBuildTool cannot automatically generate this project type because WriteProjectFile() was not overridden.");
 		}
 
+		/// <summary>
+		/// If found writes a debug project file to disk
+		/// </summary>
+		/// <param name="InPlatforms">The platforms to write the project files for</param>
+		/// <param name="InConfigurations">The configurations to add to the project files</param>
+		/// <returns>List of project files written</returns>
+		public virtual List<Tuple<ProjectFile, string>> WriteDebugProjectFiles(List<UnrealTargetPlatform> InPlatforms, List<UnrealTargetConfiguration> InConfigurations)
+		{
+			return null;
+		}
+
+
 		public virtual void LoadGUIDFromExistingProject()
 		{
 		}
diff --git a/Engine/Source/Programs/UnrealBuildTool/System/ProjectFileGenerator.cs b/Engine/Source/Programs/UnrealBuildTool/System/ProjectFileGenerator.cs
--- a/Engine/Source/Programs/UnrealBuildTool/System/ProjectFileGenerator.cs
+++ b/Engine/Source/Programs/UnrealBuildTool/System/ProjectFileGenerator.cs
@@ -347,6 +347,26 @@ namespace UnrealBuildTool
 			return UProjectInfo.FilterGameProjects(true, bGeneratingGameProjectFiles ? GameProjectName : null);
 		}
 
+		/// <summary>
+		/// Test for if ProjectName passed is the Engine Project
+		/// </summary>
+		/// <param name="ProjectName"></param>
+		/// <returns>true if ProjectName is the Engine project, else false</returns>
+		public static bool IsEngineProject(string ProjectName)
+		{
+			return (ProjectName.Equals(EngineProjectFileNameBase));
+		}
+
+		/// <summary>
+		/// Test if the Project Name is a Game Project
+		/// </summary>
+		/// <param name="ProjectName"></param>
+		/// <returns>true if ProjectName is a Game project, else false</returns>
+		public static bool IsGameProject(string ProjectName)
+		{
+			return UProjectInfo.IsGameProject(ProjectName);			
+		}
+
 		/// <summary>
 		/// Generates a Visual Studio solution file and Visual C++ project files for all known engine and game targets.
 		/// Does not actually build anything.
@@ -451,8 +471,11 @@ namespace UnrealBuildTool
 					DefaultProject = GameProjects.Values.First();
 				}
 
+				//Related Debug Project Files - Tuple here has the related Debug Project, SolutionFolder
+				List<Tuple<ProjectFile, string>> DebugProjectFiles = new List<Tuple<ProjectFile, string>>();
+
 				// Place projects into root level solution folders
-				if( IncludeEngineSource )
+				if ( IncludeEngineSource )
 				{
 					// If we're still missing an engine project because we don't have any targets for it, make one up.
 					if( EngineProject == null )
@@ -516,6 +539,13 @@ namespace UnrealBuildTool
 						{
 							AddEngineDocumentation( EngineProject );
 						}
+
+						List<Tuple<ProjectFile, string>> NewProjectFiles = EngineProject.WriteDebugProjectFiles(InPlatforms: SupportedPlatforms, InConfigurations: SupportedConfigurations);
+
+						if (NewProjectFiles != null)
+						{
+							DebugProjectFiles.AddRange(NewProjectFiles);
+						}
 					}
 
 					foreach( ProjectFile CurGameProject in GameProjects.Values )
@@ -533,7 +563,25 @@ namespace UnrealBuildTool
 						{
 							RootFolder.AddSubFolder( "Games" ).ChildProjects.Add( CurGameProject );
 						}
+
+						List<Tuple<ProjectFile, string>> NewProjectFiles = CurGameProject.WriteDebugProjectFiles(InPlatforms: SupportedPlatforms, InConfigurations: SupportedConfigurations);
+
+						if (NewProjectFiles != null)
+						{
+							DebugProjectFiles.AddRange(NewProjectFiles);
+						}
+						
+					}
+
+					//Related Debug Project Files - Tuple has the related Debug Project, SolutionFolder
+					foreach (Tuple<ProjectFile, string> DebugProjectFile in DebugProjectFiles)
+					{						
+						AddExistingProjectFile(DebugProjectFile.Item1, bForceDevelopmentConfiguration: false);
+
+						//add it to the Android Debug Projects folder in the solution
+						RootFolder.AddSubFolder(DebugProjectFile.Item2).ChildProjects.Add(DebugProjectFile.Item1);						
 					}
+					
 
 					foreach( ProjectFile CurProgramProject in ProgramProjects.Values )
 					{
diff --git a/Engine/Source/Programs/UnrealBuildTool/System/UEPlatformProjectGenerator.cs b/Engine/Source/Programs/UnrealBuildTool/System/UEPlatformProjectGenerator.cs
--- a/Engine/Source/Programs/UnrealBuildTool/System/UEPlatformProjectGenerator.cs
+++ b/Engine/Source/Programs/UnrealBuildTool/System/UEPlatformProjectGenerator.cs
@@ -241,6 +241,16 @@ namespace UnrealBuildTool
 			return PathsLines;
 		}
 
+		/// <summary>
+		/// Return any custom property settings. These will be included in the ImportGroup section
+		/// </summary>
+		/// <param name="InPlatform">  The UnrealTargetPlatform being built</param>
+		/// <returns>string    The custom property import lines for the project file; Empty string if it doesn't require one</returns>
+		public virtual string GetVisualStudioImportGroupProperties(UnrealTargetPlatform InPlatform)
+		{
+			return "";
+		}
+
 		/// <summary>
 		/// Return any custom property settings. These will be included right after Global properties to make values available to all other imports.
 		/// </summary>
@@ -317,5 +327,30 @@ namespace UnrealBuildTool
 		{
 			return "";
 		}
+
+		/// <summary>
+		/// For Additional Project Property files that need to be written out.  This is currently used only on Android. 
+		/// </summary>
+		public virtual void WriteAdditionalPropFile()
+		{
+		}
+
+		/// <summary>
+		/// For additional Project files (ex. *PROJECTNAME*-AndroidRun.androidproj.user) that needs to be written out.  This is currently used only on Android. 
+		/// </summary>
+		/// <param name="ProjectFile">Project file this will be related to</param>
+		public virtual void WriteAdditionalProjUserFile(ProjectFile ProjectFile)
+		{
+		}
+
+		/// <summary>
+		/// For additional Project files (ex. *PROJECTNAME*-AndroidRun.androidproj) that needs to be written out.  This is currently used only on Android. 
+		/// </summary>
+		/// <param name="ProjectFile">Project file this will be related to</param>
+		/// <returns>Project file written out, Solution folder it should be put in</returns>
+		public virtual Tuple<ProjectFile, string> WriteAdditionalProjFile(ProjectFile ProjectFile)
+		{
+			return null;
+		}
 	}
 }
diff --git a/Engine/Source/Programs/UnrealBuildTool/System/UProjectInfo.cs b/Engine/Source/Programs/UnrealBuildTool/System/UProjectInfo.cs
--- a/Engine/Source/Programs/UnrealBuildTool/System/UProjectInfo.cs
+++ b/Engine/Source/Programs/UnrealBuildTool/System/UProjectInfo.cs
@@ -298,6 +298,24 @@ namespace UnrealBuildTool
 			return Projects;
 		}
 
+		/// <summary>
+		/// Check to see if a project name is a Game
+		/// </summary>
+		/// <param name="GameNameFilter"></param>
+		/// <returns>true if it is a game</returns>
+		public static bool IsGameProject(string GameNameFilter)
+		{
+			List<UProjectInfo> Projects = new List<UProjectInfo>();
+			foreach (KeyValuePair<FileReference, UProjectInfo> Entry in ProjectInfoDictionary)
+			{
+				if(Entry.Value.GameName == GameNameFilter)
+				{
+					return true;
+				}
+			}
+			return false;
+		}
+
 		/// <summary>
 		/// Get the project folder for the given target name
 		/// </summary>
diff --git a/Engine/Source/Programs/UnrealBuildTool/System/Unity.cs b/Engine/Source/Programs/UnrealBuildTool/System/Unity.cs
--- a/Engine/Source/Programs/UnrealBuildTool/System/Unity.cs
+++ b/Engine/Source/Programs/UnrealBuildTool/System/Unity.cs
@@ -312,7 +312,9 @@ namespace UnrealBuildTool
 					}
 					else
 					{
-						IncludePath = CPPFile.Reference.MakeRelativeTo(CompileEnvironment.OutputDirectory).Replace('\\', '/');
+						// @todo: MakeRelativeTo does not work with code projects on a different drive than the engine. reverting to old version until we can come
+						// up with a better solution
+						IncludePath = RemoteExports.ConvertPath(CPPFile.AbsolutePath).Replace('\\', '/');
 					}
 					OutputUnityCPPWriter.WriteLine("#include \"{0}\"", IncludePath);
 				}
diff --git a/Engine/Source/Programs/UnrealBuildTool/System/UnrealBuildTool.cs b/Engine/Source/Programs/UnrealBuildTool/System/UnrealBuildTool.cs
--- a/Engine/Source/Programs/UnrealBuildTool/System/UnrealBuildTool.cs
+++ b/Engine/Source/Programs/UnrealBuildTool/System/UnrealBuildTool.cs
@@ -687,6 +687,10 @@ namespace UnrealBuildTool
 							// It's the platform set...
 							//PlatformName = Arg;
 						}
+						else if (LowercaseArg == "-vsdebugandroid")
+						{
+							AndroidProjectGenerator.VSDebugCommandLineOptionPresent = true;
+						}
 						else
 						{
 							// This arg may be a game name. Check for the existence of a game folder with this name.
diff --git a/Engine/Source/Programs/UnrealBuildTool/System/UnrealPluginLanguage.cs b/Engine/Source/Programs/UnrealBuildTool/System/UnrealPluginLanguage.cs
--- a/Engine/Source/Programs/UnrealBuildTool/System/UnrealPluginLanguage.cs
+++ b/Engine/Source/Programs/UnrealBuildTool/System/UnrealPluginLanguage.cs
@@ -131,9 +131,9 @@ namespace UnrealBuildTool
 	 * The boolean operator nodes may be combined to create a final state for more complex
 	 * conditions:
 	 * 
-	 *	<setBoolNot result="notDistribution" value="$B(Distribution)/>
-	 *	<setBoolEquals result="isX86" arg1="$S(Architecture)" arg2="x86"/>
-	 *	<setBoolEquals result="isX86_64" arg2="$S(Architecture)" arg2="x86_64/">
+	 *	<setBoolNot result="notDistribution" source="$B(Distribution)/>
+	 *	<setBoolIsEqual result="isX86" arg1="$S(Architecture)" arg2="x86"/>
+	 *	<setBoolIsEqual result="isX86_64" arg2="$S(Architecture)" arg2="x86_64/">
 	 *	<setBoolOr result="isIntel" arg1="$B(isX86)" arg2="$B(isX86_64)"/>
 	 *	<setBoolAnd result="intelAndNotDistribution" arg1="$B(isIntel)" arg2="$B(notDistribution)"/>
 	 *	<if condition="intelAndNotDistribution">
@@ -350,6 +350,9 @@ namespace UnrealBuildTool
 	 * 	<!-- optional additions to GameActivity onResume in GameActivity.java -->
 	 *	<gameActivityOnResumeAdditions>	</gameActivityOnResumeAdditions>
 	 *	
+	 * 	<!-- optional additions to GameActivity onNewIntent in GameActivity.java -->
+	 *	<gameActivityOnNewIntentAdditions> </gameActivityOnNewIntentAdditions>
+	 *	
 	 * 	<!-- optional additions to GameActivity onActivityResult in GameActivity.java -->
 	 * 	<gameActivityOnActivityResultAdditions>	</gameActivityOnActivityResultAdditions>
 	 * 	
diff --git a/Engine/Source/Programs/UnrealBuildTool/System/VCProject.cs b/Engine/Source/Programs/UnrealBuildTool/System/VCProject.cs
--- a/Engine/Source/Programs/UnrealBuildTool/System/VCProject.cs
+++ b/Engine/Source/Programs/UnrealBuildTool/System/VCProject.cs
@@ -373,6 +373,117 @@
 			}
 		}
 
+		List<ProjectConfigAndTargetCombination> ProjectConfigAndTargetCombinations = new List<ProjectConfigAndTargetCombination>();
+		
+		private void BuildProjectConfigAndTargetCombinations(List<UnrealTargetPlatform> InPlatforms, List<UnrealTargetConfiguration> InConfigurations)
+		{
+			//no need to do this more than once
+			if(ProjectConfigAndTargetCombinations.Count > 0) 
+			{
+				return;
+			}
+
+			// Build up a list of platforms and configurations this project will support.  In this list, Unknown simply
+			// means that we should use the default "stub" project platform and configuration name.
+
+			// If this is a "stub" project, then only add a single configuration to the project
+			if (IsStubProject)
+			{
+				ProjectConfigAndTargetCombination StubCombination = new ProjectConfigAndTargetCombination(UnrealTargetPlatform.Unknown, UnrealTargetConfiguration.Unknown, StubProjectPlatformName, StubProjectConfigurationName, null);
+				ProjectConfigAndTargetCombinations.Add(StubCombination);
+			}
+			else
+			{
+				// Figure out all the desired configurations
+				foreach (UnrealTargetConfiguration Configuration in InConfigurations)
+				{
+					[REDACTED_EMAIL]: Put this in a commonly accessible place?
+					if (UnrealBuildTool.IsValidConfiguration(Configuration) == false)
+					{
+						continue;
+					}
+					foreach (UnrealTargetPlatform Platform in InPlatforms)
+					{
+						if (UnrealBuildTool.IsValidPlatform(Platform) == false)
+						{
+							continue;
+						}
+						UEBuildPlatform BuildPlatform = UEBuildPlatform.GetBuildPlatform(Platform, true);
+						if ((BuildPlatform != null) && (BuildPlatform.HasRequiredSDKsInstalled() == SDKStatus.Valid))
+						{
+							// Now go through all of the target types for this project
+							if (ProjectTargets.Count == 0)
+							{
+								throw new BuildException("Expecting at least one ProjectTarget to be associated with project '{0}' in the TargetProjects list ", ProjectFilePath);
+							}
+
+							foreach (ProjectTarget ProjectTarget in ProjectTargets)
+							{
+								if (IsValidProjectPlatformAndConfiguration(ProjectTarget, Platform, Configuration))
+								{
+									string ProjectPlatformName, ProjectConfigurationName;
+									MakeProjectPlatformAndConfigurationNames(Platform, Configuration, ProjectTarget.TargetRules.Type.ToString(), out ProjectPlatformName, out ProjectConfigurationName);
+
+									ProjectConfigAndTargetCombination Combination = new ProjectConfigAndTargetCombination(Platform, Configuration, ProjectPlatformName, ProjectConfigurationName, ProjectTarget);
+									ProjectConfigAndTargetCombinations.Add(Combination);
+								}
+							}
+						}
+					}
+				}
+			}
+		}
+
+
+		/// <summary>
+		/// If found writes a debug project file to disk
+		/// </summary>
+		/// <returns>True on success</returns>
+		public override List<Tuple<ProjectFile, string>> WriteDebugProjectFiles(List<UnrealTargetPlatform> InPlatforms, List<UnrealTargetConfiguration> InConfigurations)
+		{
+			string ProjectName = ProjectFilePath.GetFileNameWithoutExtension();
+			
+			List<UnrealTargetPlatform> ProjectPlatforms = new List<UnrealTargetPlatform>();
+			List<Tuple<ProjectFile, string>> ProjectFiles = new List<Tuple<ProjectFile, string>>();
+
+			BuildProjectConfigAndTargetCombinations(InPlatforms, InConfigurations);
+
+
+			foreach (ProjectConfigAndTargetCombination Combination in ProjectConfigAndTargetCombinations)
+			{
+				if (!ProjectPlatforms.Contains(Combination.Platform))
+				{
+					ProjectPlatforms.Add(Combination.Platform);
+				}
+			}
+
+			//write out any additional project files
+			if (!IsStubProject && (ProjectFileGenerator.IsEngineProject(ProjectName) || ProjectFileGenerator.IsGameProject(ProjectName)))
+			{
+				foreach (UnrealTargetPlatform Platform in ProjectPlatforms)
+				{
+					UEPlatformProjectGenerator ProjGenerator = UEPlatformProjectGenerator.GetPlatformProjectGenerator(Platform, true);
+					if (ProjGenerator != null)
+					{
+						//write out additional prop file
+						ProjGenerator.WriteAdditionalPropFile();
+
+						//write out additional project user files
+						ProjGenerator.WriteAdditionalProjUserFile(this);
+
+						//write out additional project files
+						Tuple<ProjectFile, string> DebugProjectInfo = ProjGenerator.WriteAdditionalProjFile(this);
+						if(DebugProjectInfo != null) 
+						{
+							ProjectFiles.Add(DebugProjectInfo);
+						}
+					}
+				}
+			}
+
+			return ProjectFiles;
+		}
+
 		/// Implements Project interface
 		public override bool WriteProjectFile(List<UnrealTargetPlatform> InPlatforms, List<UnrealTargetConfiguration> InConfigurations)
 		{
@@ -436,56 +547,7 @@
 					);
 			}
 
-			// Build up a list of platforms and configurations this project will support.  In this list, Unknown simply
-			// means that we should use the default "stub" project platform and configuration name.
-			List<ProjectConfigAndTargetCombination> ProjectConfigAndTargetCombinations = new List<ProjectConfigAndTargetCombination>();
+			BuildProjectConfigAndTargetCombinations(InPlatforms, InConfigurations);
 
-			// If this is a "stub" project, then only add a single configuration to the project
-			if (IsStubProject)
-			{
-				ProjectConfigAndTargetCombination StubCombination = new ProjectConfigAndTargetCombination(UnrealTargetPlatform.Unknown, UnrealTargetConfiguration.Unknown, StubProjectPlatformName, StubProjectConfigurationName, null);
-				ProjectConfigAndTargetCombinations.Add(StubCombination);
-			}
-			else
-			{
-				// Figure out all the desired configurations
-				foreach (UnrealTargetConfiguration Configuration in InConfigurations)
-				{
-					[REDACTED_EMAIL]: Put this in a commonly accessible place?
-					if (UnrealBuildTool.IsValidConfiguration(Configuration) == false)
-					{
-						continue;
-					}
-					foreach (UnrealTargetPlatform Platform in InPlatforms)
-					{
-						if (UnrealBuildTool.IsValidPlatform(Platform) == false)
-						{
-							continue;
-						}
-						UEBuildPlatform BuildPlatform = UEBuildPlatform.GetBuildPlatform(Platform, true);
-						if ((BuildPlatform != null) && (BuildPlatform.HasRequiredSDKsInstalled() == SDKStatus.Valid))
-						{
-							// Now go through all of the target types for this project
-							if (ProjectTargets.Count == 0)
-							{
-								throw new BuildException("Expecting at least one ProjectTarget to be associated with project '{0}' in the TargetProjects list ", ProjectFilePath);
-							}
-
-							foreach (ProjectTarget ProjectTarget in ProjectTargets)
-							{
-								if (IsValidProjectPlatformAndConfiguration(ProjectTarget, Platform, Configuration))
-								{
-									string ProjectPlatformName, ProjectConfigurationName;
-									MakeProjectPlatformAndConfigurationNames(Platform, Configuration, ProjectTarget.TargetRules.Type.ToString(), out ProjectPlatformName, out ProjectConfigurationName);
-
-									ProjectConfigAndTargetCombination Combination = new ProjectConfigAndTargetCombination(Platform, Configuration, ProjectPlatformName, ProjectConfigurationName, ProjectTarget);
-									ProjectConfigAndTargetCombinations.Add(Combination);
-								}
-							}
-						}
-					}
-				}
-			}
-
 			VCProjectFileContent.Append(
 				"	<ItemGroup Label=\"ProjectConfigurations\">" + ProjectFileGenerator.NewLine);
@@ -1051,6 +1113,7 @@
 				VCProjectFileContent.Append(
 					"	<ImportGroup " + ConditionString + " Label=\"PropertySheets\">" + ProjectFileGenerator.NewLine +
 					"		<Import Project=\"$(UserRootDir)\\Microsoft.Cpp.$(Platform).user.props\" Condition=\"exists('$(UserRootDir)\\Microsoft.Cpp.$(Platform).user.props')\" Label=\"LocalAppDataPlatform\" />" + ProjectFileGenerator.NewLine +
+							ProjGenerator.GetVisualStudioImportGroupProperties(Platform) +
 					"	</ImportGroup>" + ProjectFileGenerator.NewLine);
 
 				DirectoryReference ProjectDirectory = ProjectFilePath.Directory;
@@ -1377,4 +1440,5 @@
 		/// System assemblies this project is dependent on
 		protected readonly List<string> DotNetAssemblyReferences = new List<string>() { "System", "System.Core", "System.Data", "System.Xml" };
 	}
+
 }
diff --git a/Engine/Source/Programs/UnrealBuildTool/ToolChain/RemoteExports.cs b/Engine/Source/Programs/UnrealBuildTool/ToolChain/RemoteExports.cs
new file mode 100644
--- /dev/null
+++ b/Engine/Source/Programs/UnrealBuildTool/ToolChain/RemoteExports.cs
@@ -0,0 +1,33 @@
+﻿using System;
+using System.Collections.Generic;
+using System.Linq;
+using System.Text;
+using System.Threading.Tasks;
+
+namespace UnrealBuildTool
+{
+	/// <summary>
+	/// Public Remote toolchain functions exposed to UAT
+	/// </summary>
+	public class RemoteExports
+	{
+		/// <summary>
+		/// 
+		/// </summary>
+		/// <returns></returns>
+		public static string ConvertPath(string Path)
+		{
+			return RemoteToolChain.ConvertPath(Path);
+		}
+
+		/// <summary>
+		/// 
+		/// </summary>
+		/// <param name="RemotePath"></param>
+		/// <returns></returns>
+		public static string UnconvertPath(string RemotePath)
+		{
+			return RemoteToolChain.UnconvertPath(RemotePath);
+		}
+	}
+}
diff --git a/Engine/Source/Programs/UnrealBuildTool/ToolChain/RemoteToolChain.cs b/Engine/Source/Programs/UnrealBuildTool/ToolChain/RemoteToolChain.cs
--- a/Engine/Source/Programs/UnrealBuildTool/ToolChain/RemoteToolChain.cs
+++ b/Engine/Source/Programs/UnrealBuildTool/ToolChain/RemoteToolChain.cs
@@ -586,6 +586,39 @@ namespace UnrealBuildTool
 				return OriginalPath.Replace("\\", "/");
 			}
 		}
+		public static string UnconvertPath(string RemotePath)
+		{
+			if (BuildHostPlatform.Current.Platform != UnrealTargetPlatform.Mac)
+			{
+				string StrippedPath = RemotePath.Replace("\\", "/");
+
+				// skip over the UserDevRootMac and MachineName
+				string StuffToSkip = string.Format("{0}{1}", UserDevRootMac, Environment.MachineName);
+				if (!StrippedPath.StartsWith(StuffToSkip))
+				{
+					return RemotePath;
+				}
+				StrippedPath = StrippedPath.Substring(StuffToSkip.Length);
+
+				// now we make sure the path is /{DriverLetter}/
+				if (StrippedPath[0] != '/' || StrippedPath[2] != '/')
+				{
+					return RemotePath;
+				}
+
+				// convert /{DriveLetter}/ to {DriveLetter}:
+				char DriveLetter = StrippedPath[1];
+				// and skip over it
+				StrippedPath = StrippedPath.Substring(3);
+
+				// put back into PC parlance
+				return string.Format("{0}:/{1}", DriveLetter, StrippedPath).Replace("/", "\\");
+			}
+			else
+			{
+				return RemotePath;
+			}
+		}
 
 		protected string GetMacDevSrcRoot()
 		{
diff --git a/Engine/Source/Programs/UnrealBuildTool/UnrealBuildTool.csproj b/Engine/Source/Programs/UnrealBuildTool/UnrealBuildTool.csproj
--- a/Engine/Source/Programs/UnrealBuildTool/UnrealBuildTool.csproj
+++ b/Engine/Source/Programs/UnrealBuildTool/UnrealBuildTool.csproj
@@ -191,6 +191,7 @@
     </When>
   </Choose>
   <ItemGroup>
+    <Compile Include="ToolChain\RemoteExports.cs" />
     <Compile Include="Utilities\BinaryReaderExtensions.cs" />
     <Compile Include="Utilities\BinaryWriterExtensions.cs" />
     <Compile Include="Utilities\HashSetExtensions.cs" />
diff --git a/Engine/Source/Runtime/Android/AndroidAudio/Private/AndroidAudioBuffer.cpp b/Engine/Source/Runtime/Android/AndroidAudio/Private/AndroidAudioBuffer.cpp
--- a/Engine/Source/Runtime/Android/AndroidAudio/Private/AndroidAudioBuffer.cpp
+++ b/Engine/Source/Runtime/Android/AndroidAudio/Private/AndroidAudioBuffer.cpp
@@ -100,6 +100,53 @@ FSLESSoundBuffer* FSLESSoundBuffer::CreateQueuedBuffer( FSLESAudioDevice* AudioD
 	return Buffer;
 }
 
+FSLESSoundBuffer* FSLESSoundBuffer::CreateStreamBuffer( FSLESAudioDevice* AudioDevice, USoundWave* InWave )
+{
+	// Always create a new buffer for streaming
+	FSLESSoundBuffer* Buffer = new FSLESSoundBuffer( AudioDevice);
+	
+	FSoundQualityInfo QualityInfo = { 0 };
+	
+	Buffer->DecompressionState = AudioDevice->CreateCompressedAudioInfo(InWave);
+
+	if (Buffer->DecompressionState && Buffer->DecompressionState->StreamCompressedInfo(InWave, &QualityInfo))
+	{	
+		// Clear out any dangling pointers
+		Buffer->AudioData = NULL;
+		Buffer->BufferSize = 0;
+		
+		// Keep track of associated resource name.
+		Buffer->ResourceName	= InWave->GetPathName();		
+		Buffer->NumChannels		= InWave->NumChannels;
+		Buffer->SampleRate		= InWave->SampleRate;
+
+		FPlatformMisc::LowLevelOutputDebugStringf(TEXT("DEBUG: FSLESSoundBuffer::CreateStreamBuffer Buffer->SampleRate = %d"), Buffer->SampleRate);
+
+		//Android can't handle more than 48kHz, so turn on halfrate decoding and adjust parameters
+		if (Buffer->SampleRate > 48000)
+		{
+			UE_LOG(LogAndroidAudio, Log, TEXT( "Converting %s to halfrate from %d" ), *InWave->GetName(), Buffer->SampleRate );
+			Buffer->DecompressionState->EnableHalfRate( true);
+			Buffer->SampleRate = Buffer->SampleRate / 2;
+			InWave->SampleRate = InWave->SampleRate / 2;
+			uint32 SampleCount = QualityInfo.SampleDataSize / (QualityInfo.NumChannels * sizeof(uint16));
+			SampleCount /= 2;
+			InWave->RawPCMDataSize = SampleCount * QualityInfo.NumChannels * sizeof(uint16);;
+		}
+
+		Buffer->Format = SoundFormat_Streaming;	
+	}
+	else
+	{
+		InWave->DecompressionType = DTYPE_Invalid;
+		InWave->NumChannels = 0;
+		
+		InWave->RemoveAudioResource();
+	}
+	
+	return Buffer;
+}
+
 /**
  * Static function used to create an OpenSL buffer and upload decompressed ogg vorbis data to.
  *
@@ -232,6 +279,11 @@ FSLESSoundBuffer* FSLESSoundBuffer::Init(  FSLESAudioDevice* AudioDevice ,USound
 		// Always create a new buffer for streaming ogg vorbis data
 		Buffer = CreateQueuedBuffer( AudioDevice, InWave );
 		break;
+
+	case DTYPE_Streaming:
+		// Always create a new buffer for streaming ogg vorbis data
+		Buffer = CreateStreamBuffer( AudioDevice, InWave );
+		break;
 	
 	case DTYPE_Procedural:
 		// New buffer for procedural data
@@ -259,6 +311,12 @@ FSLESSoundBuffer* FSLESSoundBuffer::Init(  FSLESAudioDevice* AudioDevice ,USound
 bool FSLESSoundBuffer::ReadCompressedData( uint8* Destination, bool bLooping )
 {
 	ensure( DecompressionState);
+
+	if(Format == SoundFormat_Streaming)
+	{
+		return(DecompressionState->StreamCompressedData(Destination, bLooping, DecompressionState->GetStreamBufferSize() * NumChannels));
+	}
+
 	return(DecompressionState->ReadCompressedData(Destination, bLooping, DecompressionState->GetStreamBufferSize() * NumChannels));
 }
 
@@ -270,6 +328,30 @@ void FSLESSoundBuffer::Seek(const float SeekTime)
 	}
 }
 
+int32 FSLESSoundBuffer::GetCurrentChunkIndex() const
+{
+	int32 result = -1;
+
+	if (ensure(DecompressionState))
+	{
+		result = DecompressionState->GetCurrentChunkIndex();
+	}
+
+	return result;
+}
+
+int32 FSLESSoundBuffer::GetCurrentChunkOffset() const
+{
+	int32 result = -1;
+
+	if (ensure(DecompressionState))
+	{
+		result = DecompressionState->GetCurrentChunkOffset();
+	}
+
+	return result;
+}
+
 /**
 * Returns the size for a real time/streaming buffer based on decompressor
 *
diff --git a/Engine/Source/Runtime/Android/AndroidAudio/Private/AndroidAudioDevice.cpp b/Engine/Source/Runtime/Android/AndroidAudio/Private/AndroidAudioDevice.cpp
--- a/Engine/Source/Runtime/Android/AndroidAudio/Private/AndroidAudioDevice.cpp
+++ b/Engine/Source/Runtime/Android/AndroidAudio/Private/AndroidAudioDevice.cpp
@@ -143,12 +143,23 @@ FSoundSource* FSLESAudioDevice::CreateSoundSource()
 bool FSLESAudioDevice::HasCompressedAudioInfoClass(USoundWave* SoundWave)
 {
 #if WITH_OGGVORBIS
+	if(SoundWave->bStreaming)
+	{
+		return true;
+	}
+
 	static FName NAME_OGG(TEXT("OGG"));
 	if (SoundWave->HasCompressedData(NAME_OGG))
 	{
 		return true;
 	}
 #endif
+
+	if(SoundWave->bStreaming)
+	{
+		return true;
+	}
+
 	static FName NAME_ADPCM(TEXT("ADPCM"));
 	if (SoundWave->HasCompressedData(NAME_ADPCM))
 	{
@@ -163,13 +174,13 @@ class ICompressedAudioInfo* FSLESAudioDevice::CreateCompressedAudioInfo(USoundWa
 {
 #if WITH_OGGVORBIS
 	static FName NAME_OGG(TEXT("OGG"));
-	if (SoundWave->HasCompressedData(NAME_OGG))
+	if (SoundWave->bStreaming || SoundWave->HasCompressedData(NAME_OGG))
 	{
 		return new FVorbisAudioInfo();
 	}
 #endif
 	static FName NAME_ADPCM(TEXT("ADPCM"));
-	if (SoundWave->HasCompressedData(NAME_ADPCM))
+	if (SoundWave->bStreaming || SoundWave->HasCompressedData(NAME_ADPCM))
 	{
 		return new FADPCMAudioInfo();
 	}
diff --git a/Engine/Source/Runtime/Android/AndroidAudio/Private/AndroidAudioSource.cpp b/Engine/Source/Runtime/Android/AndroidAudio/Private/AndroidAudioSource.cpp
--- a/Engine/Source/Runtime/Android/AndroidAudio/Private/AndroidAudioSource.cpp
+++ b/Engine/Source/Runtime/Android/AndroidAudio/Private/AndroidAudioSource.cpp
@@ -22,7 +22,7 @@ void FSLESSoundSource::OnRequeueBufferCallback( SLAndroidSimpleBufferQueueItf In
 {
 	if (!bStreamedSound)
 	{
-		SLresult result = (*SL_PlayerBufferQueue)->Enqueue(SL_PlayerBufferQueue, Buffer->AudioData, Buffer->GetSize() );
+		SLresult result = (*SL_PlayerBufferQueue)->Enqueue(SL_PlayerBufferQueue, SLESBuffer->AudioData, SLESBuffer->GetSize() );
 		if(result != SL_RESULT_SUCCESS) 
 		{ 
 			UE_LOG( LogAndroidAudio, Warning, TEXT("FAILED OPENSL BUFFER Enqueue SL_PlayerBufferQueue (Requeing)"));  
@@ -86,9 +86,9 @@ bool FSLESSoundSource::CreatePlayer()
 	SLDataLocator_AndroidSimpleBufferQueue LocationBuffer	= {		SL_DATALOCATOR_ANDROIDSIMPLEBUFFERQUEUE, 1 };
 		
 	// PCM Info
-	SLDataFormat_PCM PCM_Format				= {		SL_DATAFORMAT_PCM, SLuint32(Buffer->NumChannels), SLuint32( Buffer->SampleRate * 1000 ),	
+	SLDataFormat_PCM PCM_Format				= {		SL_DATAFORMAT_PCM, SLuint32(SLESBuffer->NumChannels), SLuint32( SLESBuffer->SampleRate * 1000 ),	
 		SL_PCMSAMPLEFORMAT_FIXED_16, SL_PCMSAMPLEFORMAT_FIXED_16, 
-		Buffer->NumChannels == 2 ? ( SL_SPEAKER_FRONT_LEFT | SL_SPEAKER_FRONT_RIGHT ) : SL_SPEAKER_FRONT_CENTER, 
+		SLESBuffer->NumChannels == 2 ? ( SL_SPEAKER_FRONT_LEFT | SL_SPEAKER_FRONT_RIGHT ) : SL_SPEAKER_FRONT_CENTER, 
 		SL_BYTEORDER_LITTLEENDIAN };
 		
 	SLDataSource SoundDataSource			= {		&LocationBuffer, &PCM_Format};
@@ -107,9 +107,9 @@ bool FSLESSoundSource::CreatePlayer()
 		UE_LOG(LogAndroidAudio, Warning, TEXT("FAILED OPENSL BUFFER CreateAudioPlayer 0x%x"), result);
 		return false;
 	}
-		
+	
 	bool bFailedSetup = false;
-		
+	
 	// realize the player
 	result = (*SL_PlayerObject)->Realize(SL_PlayerObject, SL_BOOLEAN_FALSE);
 	if (result != SL_RESULT_SUCCESS) { UE_LOG(LogAndroidAudio, Warning, TEXT("FAILED OPENSL BUFFER Realize 0x%x"), result); return false; }
@@ -150,9 +150,9 @@ bool FSLESSoundSource::EnqueuePCMBuffer( bool bLoop)
 		if (result != SL_RESULT_SUCCESS) { UE_LOG(LogAndroidAudio, Warning, TEXT("FAILED OPENSL BUFFER QUEUE RegisterCallback 0x%x "), result); return false; }
 	}
 	
-	result = (*SL_PlayerBufferQueue)->Enqueue(SL_PlayerBufferQueue, Buffer->AudioData, Buffer->GetSize() );
+	result = (*SL_PlayerBufferQueue)->Enqueue(SL_PlayerBufferQueue, SLESBuffer->AudioData, SLESBuffer->GetSize() );
 	if (result != SL_RESULT_SUCCESS) {
-		UE_LOG(LogAndroidAudio, Warning, TEXT("FAILED OPENSL BUFFER Enqueue SL_PlayerBufferQueue 0x%x params( %p, %d)"), result, Buffer->AudioData, int32(Buffer->GetSize()));
+		UE_LOG(LogAndroidAudio, Warning, TEXT("FAILED OPENSL BUFFER Enqueue SL_PlayerBufferQueue 0x%x params( %p, %d)"), result, SLESBuffer->AudioData, int32(SLESBuffer->GetSize()));
 		if (bLoop)
 		{
 			result = (*SL_PlayerBufferQueue)->RegisterCallback(SL_PlayerBufferQueue, NULL, NULL);
@@ -192,11 +192,11 @@ bool FSLESSoundSource::ReadMorePCMData(const int32 BufferIndex, EDataReadMode Da
 	{
 		if (DataReadMode == EDataReadMode::Synchronous)
 		{
-			return Buffer->ReadCompressedData(AudioBuffers[BufferIndex].AudioData, WaveInstance->LoopingMode != LOOP_Never);
+			return SLESBuffer->ReadCompressedData(AudioBuffers[BufferIndex].AudioData, WaveInstance->LoopingMode != LOOP_Never);
 		}
 		else
 		{
-			RealtimeAsyncTask = new FAsyncRealtimeAudioTask(Buffer, AudioBuffers[BufferIndex].AudioData, WaveInstance->LoopingMode != LOOP_Never, DataReadMode == EDataReadMode::AsynchronousSkipFirstFrame);
+			RealtimeAsyncTask = new FAsyncRealtimeAudioTask(SLESBuffer, AudioBuffers[BufferIndex].AudioData, WaveInstance->LoopingMode != LOOP_Never, DataReadMode == EDataReadMode::AsynchronousSkipFirstFrame);
 			RealtimeAsyncTask->StartBackgroundTask();
 			return false;
 		}
@@ -213,7 +213,7 @@ bool FSLESSoundSource::EnqueuePCMRTBuffer( bool bLoop )
 	FMemory::Memzero( AudioBuffers, sizeof( SLESAudioBuffer ) * 2 );
 
 	// Set up double buffer area to decompress to
-	BufferSize = Buffer->GetRTBufferSize() * Buffer->NumChannels;
+	BufferSize = SLESBuffer->GetRTBufferSize() * SLESBuffer->NumChannels;
 
 	AudioBuffers[0].AudioData = (uint8*)FMemory::Malloc(BufferSize);
 	AudioBuffers[0].AudioDataSize = BufferSize;
@@ -242,7 +242,7 @@ bool FSLESSoundSource::EnqueuePCMRTBuffer( bool bLoop )
 	if(result == SL_RESULT_SUCCESS) 
 	{
 		result = (*SL_PlayerBufferQueue)->Enqueue(SL_PlayerBufferQueue, AudioBuffers[0].AudioData, AudioBuffers[0].AudioDataSize );
-		if (result != SL_RESULT_SUCCESS) { UE_LOG(LogAndroidAudio, Warning, TEXT("FAILED OPENSL BUFFER Enqueue SL_PlayerBufferQueue 0x%x params( %p, %d)"), result, Buffer->AudioData, int32(Buffer->GetSize())); return false; }
+		if (result != SL_RESULT_SUCCESS) { UE_LOG(LogAndroidAudio, Warning, TEXT("FAILED OPENSL BUFFER Enqueue SL_PlayerBufferQueue 0x%x params( %p, %d)"), result, SLESBuffer->AudioData, int32(SLESBuffer->GetSize())); return false; }
 	}
 	else
 	{
@@ -273,11 +273,12 @@ bool FSLESSoundSource::Init( FWaveInstance* InWaveInstance )
 		return false;
 	}
 
-	if (Buffer && Buffer->ResourceID == 0)
+	if (SLESBuffer && SLESBuffer->ResourceID == 0)
 	{
 		UE_LOG( LogAndroidAudio, Warning, TEXT(" InitSoundSouce with Buffer already allocated"));
-		delete Buffer;
-		Buffer = 0;
+		delete SLESBuffer;
+		SLESBuffer = nullptr;
+		Buffer = nullptr;
 	}
 
 	if (SL_PlayerObject)
@@ -286,9 +287,10 @@ bool FSLESSoundSource::Init( FWaveInstance* InWaveInstance )
 	}
 	
 	// Find matching buffer.
-	Buffer = FSLESSoundBuffer::Init( (FSLESAudioDevice *)AudioDevice, InWaveInstance->WaveData );
+	SLESBuffer = FSLESSoundBuffer::Init( (FSLESAudioDevice *)AudioDevice, InWaveInstance->WaveData );
+	Buffer = SLESBuffer;
 
-	if( Buffer && InWaveInstance->WaveData->NumChannels <= 2 && InWaveInstance->WaveData->SampleRate <= 48000 )
+	if( SLESBuffer && InWaveInstance->WaveData->NumChannels <= 2 && InWaveInstance->WaveData->SampleRate <= 48000 )
 	{
 		SCOPE_CYCLE_COUNTER( STAT_AudioSourceInitTime );
 		
@@ -300,15 +302,16 @@ bool FSLESSoundSource::Init( FWaveInstance* InWaveInstance )
 
 			if (WaveInstance->StartTime > 0.f)
 			{
-				Buffer->Seek(WaveInstance->StartTime);
+				SLESBuffer->Seek(WaveInstance->StartTime);
 			}
 
-			switch( Buffer->Format)
+			switch( SLESBuffer->Format)
 			{
 				case SoundFormat_PCM:
 					bFailedSetup |= !EnqueuePCMBuffer( InWaveInstance->LoopingMode != LOOP_Never );
 					break;
 				case SoundFormat_PCMRT:
+				case SoundFormat_Streaming:
 					bFailedSetup |= !EnqueuePCMRTBuffer( InWaveInstance->LoopingMode != LOOP_Never );
 					break;
 				default:
@@ -341,10 +344,11 @@ bool FSLESSoundSource::Init( FWaveInstance* InWaveInstance )
 		UE_LOG( LogAndroidAudio, Warning, TEXT("  SampleRate %d"), InWaveInstance->WaveData->SampleRate);
 		UE_LOG( LogAndroidAudio, Warning, TEXT("  Channels %d"), InWaveInstance->WaveData->NumChannels);
 
-		if (Buffer && Buffer->ResourceID == 0)
+		if (SLESBuffer && SLESBuffer->ResourceID == 0)
 		{
-			delete Buffer;
-			Buffer = 0;
+			delete SLESBuffer;
+			SLESBuffer = nullptr;
+			Buffer = nullptr;
 		} 
 	}
 	return false;
@@ -354,7 +358,7 @@ bool FSLESSoundSource::Init( FWaveInstance* InWaveInstance )
 FSLESSoundSource::FSLESSoundSource( class FAudioDevice* InAudioDevice )
 	:	FSoundSource( InAudioDevice ),
 		Device((FSLESAudioDevice *)InAudioDevice),
-		Buffer( NULL ),
+		SLESBuffer( NULL ),
 		bStreamedSound(false),
 		bBuffersToFlush(false),
 		BufferSize(0),
@@ -368,6 +372,7 @@ FSLESSoundSource::FSLESSoundSource( class FAudioDevice* InAudioDevice )
 		SL_VolumeInterface(NULL),
 		RealtimeAsyncTask(NULL)
 {
+	Buffer = NULL;
 	FMemory::Memzero( AudioBuffers, sizeof( AudioBuffers ) );
 }
 
@@ -395,11 +400,12 @@ void FSLESSoundSource::ReleaseResources()
 
 	FMemory::Memzero( AudioBuffers, sizeof( AudioBuffers ) );
 
-	if (Buffer && Buffer->ResourceID == 0)
+	if (SLESBuffer && SLESBuffer->ResourceID == 0)
 	{
-		delete Buffer;
+		delete SLESBuffer;
 	}
-	Buffer = NULL;
+	SLESBuffer = nullptr;
+	Buffer = nullptr;
 }
 
 /**
@@ -521,7 +527,8 @@ void FSLESSoundSource::Stop( void )
 		
 		Paused = false;
 		Playing = false;
-		Buffer = NULL;
+		SLESBuffer = nullptr;
+		Buffer = nullptr;
 	}
 	
 	FSoundSource::Stop();
diff --git a/Engine/Source/Runtime/Android/AndroidAudio/Public/AndroidAudioDevice.h b/Engine/Source/Runtime/Android/AndroidAudio/Public/AndroidAudioDevice.h
--- a/Engine/Source/Runtime/Android/AndroidAudio/Public/AndroidAudioDevice.h
+++ b/Engine/Source/Runtime/Android/AndroidAudio/Public/AndroidAudioDevice.h
@@ -22,7 +22,8 @@ enum ESoundFormat
 {
 	SoundFormat_Invalid,
 	SoundFormat_PCM,
-	SoundFormat_PCMRT
+	SoundFormat_PCMRT,
+	SoundFormat_Streaming
 };
 
 struct SLESAudioBuffer
@@ -61,6 +62,15 @@ public:
 	 */
 	static FSLESSoundBuffer* CreateQueuedBuffer( FSLESAudioDevice* AudioDevice, USoundWave* Wave );
 
+	/**
+	 * Static function used to create a buffer and stream data directly off disk as needed
+	 *
+	 * @param InWave		USoundWave to use as template and wave source
+	 * @param AudioDevice	audio device to attach created buffer to
+	 * @return FSLESSoundBuffer pointer if buffer creation succeeded, NULL otherwise
+	 */
+	static FSLESSoundBuffer* CreateStreamBuffer( FSLESAudioDevice* AudioDevice, USoundWave* Wave );
+
 	/**
 	 * Static function used to create an Audio buffer and upload decompressed ogg vorbis data to.
 	 *
@@ -114,6 +124,10 @@ public:
 	{ 
 		return( BufferSize ); 
 	}
+	
+	// These are used by the streaming engine to manage loading/unloading of chunks
+	virtual int32 GetCurrentChunkIndex() const override;
+	virtual int32 GetCurrentChunkOffset() const override;
 
 	/**
 	 * Returns the size for a real time/streaming buffer based on decompressor
@@ -217,7 +231,8 @@ protected:
 	};
 
 	friend class FSLESAudioDevice;
-	FSLESSoundBuffer*		Buffer;
+	// Do not shadow the Buffer member variable defined in the parent class
+	FSLESSoundBuffer*		SLESBuffer;
 	FSLESAudioDevice*		Device;
 	
 	// OpenSL interface objects
diff --git a/Engine/Source/Runtime/Android/AndroidRuntimeSettings/Classes/AndroidRuntimeSettings.h b/Engine/Source/Runtime/Android/AndroidRuntimeSettings/Classes/AndroidRuntimeSettings.h
--- a/Engine/Source/Runtime/Android/AndroidRuntimeSettings/Classes/AndroidRuntimeSettings.h
+++ b/Engine/Source/Runtime/Android/AndroidRuntimeSettings/Classes/AndroidRuntimeSettings.h
@@ -349,6 +349,10 @@ public:
 	UPROPERTY(GlobalConfig, EditAnywhere, Category = GooglePlayServices)
 	FString GooglePlayLicenseKey;
 
+	// The sender id obtained from Firebase Console, leave blank to disable (associate this with your app in Google Player Developer Console).
+	UPROPERTY(GlobalConfig, EditAnywhere, Category = GooglePlayServices, meta = (DisplayName = "Google Cloud Messaging Sender ID"))
+	FString GCMClientSenderID;
+
 	/** Show the launch image as a startup slash screen */
 	UPROPERTY(GlobalConfig, EditAnywhere, Category = LaunchImages, meta = (DisplayName = "Show launch image"))
 	bool bShowLaunchImage;
diff --git a/Engine/Source/Runtime/Core/Private/Android/AndroidFile.cpp b/Engine/Source/Runtime/Core/Private/Android/AndroidFile.cpp
--- a/Engine/Source/Runtime/Core/Private/Android/AndroidFile.cpp
+++ b/Engine/Source/Runtime/Core/Private/Android/AndroidFile.cpp
@@ -1107,7 +1107,7 @@ public:
 				{
 					FileInfo.st_mode |= S_IWUSR;
 				}
-				return chmod(TCHAR_TO_UTF8(*LocalPath), FileInfo.st_mode);
+				return chmod(TCHAR_TO_UTF8(*LocalPath), FileInfo.st_mode) == 0;
 			}
 		}
 		return false;
diff --git a/Engine/Source/Runtime/Core/Private/GenericPlatform/GenericPlatformFile.cpp b/Engine/Source/Runtime/Core/Private/GenericPlatform/GenericPlatformFile.cpp
--- a/Engine/Source/Runtime/Core/Private/GenericPlatform/GenericPlatformFile.cpp
+++ b/Engine/Source/Runtime/Core/Private/GenericPlatform/GenericPlatformFile.cpp
@@ -176,7 +176,12 @@ void FGenericReadRequestWorker::DoWork()
 	ReadRequest.PerformRequest();
 }
 
+// @todo switch et al: this is a temporary measure until we can track down some threaded file handling issues on Switch 
+#if (PLATFORM_SWITCH || PLATFORM_LINUX || PLATFORM_PS4 || PLATFORM_IOS || PLATFORM_MAC)
+#define DISABLE_HANDLE_CACHING (1)
+#else
 #define DISABLE_HANDLE_CACHING (0)
+#endif
 
 class FGenericAsyncReadFileHandle final : public IAsyncReadFileHandle
 {
diff --git a/Engine/Source/Runtime/Core/Private/GenericPlatform/GenericPlatformMisc.cpp b/Engine/Source/Runtime/Core/Private/GenericPlatform/GenericPlatformMisc.cpp
--- a/Engine/Source/Runtime/Core/Private/GenericPlatform/GenericPlatformMisc.cpp
+++ b/Engine/Source/Runtime/Core/Private/GenericPlatform/GenericPlatformMisc.cpp
@@ -1085,3 +1085,28 @@ void FGenericPlatformMisc::RegisterForRemoteNotifications()
 {
 	// not implemented by default
 }
+
+const TArray<FString>& FGenericPlatformMisc::GetConfidentialPlatforms()
+{
+	static bool bHasSearchedForPlatforms = false;
+	static TArray<FString> FoundPlatforms;
+
+	// look on disk for special files
+	if (bHasSearchedForPlatforms == false)
+	{
+		// look for the special files in any congfig subdirectories
+		IFileManager::Get().FindFilesRecursive(FoundPlatforms, *FPaths::EngineConfigDir(), TEXT("ConfidentialPlatform.ini"), true, false);
+
+		// fix up the paths
+		for (int32 PlatformIndex = 0; PlatformIndex < FoundPlatforms.Num(); PlatformIndex++)
+		{
+			// pull the directory name out
+			FoundPlatforms[PlatformIndex] = FPaths::GetCleanFilename(FPaths::GetPath(FoundPlatforms[PlatformIndex]));
+		}
+
+		bHasSearchedForPlatforms = true;
+	}
+
+	// return whatever we have already found
+	return FoundPlatforms;
+}
diff --git a/Engine/Source/Runtime/Core/Private/GenericPlatform/GenericPlatformSplash.cpp b/Engine/Source/Runtime/Core/Private/GenericPlatform/GenericPlatformSplash.cpp
--- a/Engine/Source/Runtime/Core/Private/GenericPlatform/GenericPlatformSplash.cpp
+++ b/Engine/Source/Runtime/Core/Private/GenericPlatform/GenericPlatformSplash.cpp
@@ -49,7 +49,7 @@ bool FGenericPlatformSplash::GetSplashPath(const TCHAR* SplashFilename, FString&
 	FString Filename = FString(TEXT("Splash/")) + SplashFilename;
 
 	// first look in game's splash directory
-	OutPath = GetSplashFilename(FPaths::GameContentDir(), Filename);
+	OutPath = FPaths::ConvertRelativePathToFull(GetSplashFilename(FPaths::GameContentDir(), Filename));
 	OutIsCustom = true;
 
 	// if this was found, then we're done
diff --git a/Engine/Source/Runtime/Core/Private/HAL/ThreadHeartBeat.cpp b/Engine/Source/Runtime/Core/Private/HAL/ThreadHeartBeat.cpp
--- a/Engine/Source/Runtime/Core/Private/HAL/ThreadHeartBeat.cpp
+++ b/Engine/Source/Runtime/Core/Private/HAL/ThreadHeartBeat.cpp
@@ -191,6 +191,12 @@ void FThreadHeartBeat::Start()
 
 void FThreadHeartBeat::HeartBeat()
 {
+	// disable on platforms that don't start the thread
+	if (FPlatformMisc::AllowThreadHeartBeat() == false)
+	{
+		return;
+	}
+
 	uint32 ThreadId = FPlatformTLS::GetCurrentThreadId();
 	FScopeLock HeartBeatLock(&HeartBeatCritical);
 	FHeartBeatInfo& HeartBeatInfo = ThreadHeartBeat.FindOrAdd(ThreadId);
diff --git a/Engine/Source/Runtime/Core/Private/HTML5/HTML5Application.cpp b/Engine/Source/Runtime/Core/Private/HTML5/HTML5Application.cpp
--- a/Engine/Source/Runtime/Core/Private/HTML5/HTML5Application.cpp
+++ b/Engine/Source/Runtime/Core/Private/HTML5/HTML5Application.cpp
@@ -12,6 +12,7 @@ DEFINE_LOG_CATEGORY_STATIC(LogHTML5Application, Log, All);
 #if PLATFORM_HTML5_BROWSER
 #include <emscripten/emscripten.h>
 #include <emscripten/html5.h>
+#include "HTML5JavascriptFX.h"
 
 EM_BOOL mouse_callback(int eventType, const EmscriptenMouseEvent *e, void *userData)
 {
@@ -106,8 +107,10 @@ void FHTML5Application::PollGameDeviceState( const float TimeDelta )
 							FDisplayMetrics::GetDisplayMetrics(DisplayMetrics);
 							if ( DisplayMetrics.PrimaryDisplayWidth != WindowWidth || DisplayMetrics.PrimaryDisplayHeight != WindowHeight )
 							{
-								int32 delta = FMath::Abs<int32>(DisplayMetrics.PrimaryDisplayHeight - WindowHeight);
-								if (delta > 2) // UE-35363: retina bug fix - otherwise, event does infinite loops
+								int32 deltaH = FMath::Abs<int32>(DisplayMetrics.PrimaryDisplayHeight - WindowHeight);
+								int32 deltaW = FMath::Abs<int32>(DisplayMetrics.PrimaryDisplayWidth - WindowWidth);
+								int32 pixelRatio = UE_PixelRatio();
+								if ((deltaH > pixelRatio) || (deltaW > pixelRatio)) // UE-35363: retina bug fix - otherwise, event does infinite loops
 								{
 									MessageHandler->OnSizeChanged(ApplicationWindow,WindowWidth,WindowHeight, false);
 									MessageHandler->OnResizingWindow(ApplicationWindow);
@@ -123,8 +126,10 @@ void FHTML5Application::PollGameDeviceState( const float TimeDelta )
 							FDisplayMetrics::GetDisplayMetrics(DisplayMetrics);
 							if ( DisplayMetrics.PrimaryDisplayWidth != WindowWidth || DisplayMetrics.PrimaryDisplayHeight != WindowHeight )
 							{
-								int32 delta = FMath::Abs<int32>(DisplayMetrics.PrimaryDisplayHeight - WindowHeight);
-								if (delta > 2) // UE-35363: retina bug fix - otherwise, event does infinite loops
+								int32 deltaH = FMath::Abs<int32>(DisplayMetrics.PrimaryDisplayHeight - WindowHeight);
+								int32 deltaW = FMath::Abs<int32>(DisplayMetrics.PrimaryDisplayWidth - WindowWidth);
+								int32 pixelRatio = UE_PixelRatio();
+								if ((deltaH > pixelRatio) || (deltaW > pixelRatio)) // UE-35363: retina bug fix - otherwise, event does infinite loops
 								{
 									MessageHandler->OnResizingWindow(ApplicationWindow);
 									BroadcastDisplayMetricsChanged(DisplayMetrics);
diff --git a/Engine/Source/Runtime/Core/Private/HTML5/HTML5PlatformFile.cpp b/Engine/Source/Runtime/Core/Private/HTML5/HTML5PlatformFile.cpp
--- a/Engine/Source/Runtime/Core/Private/HTML5/HTML5PlatformFile.cpp
+++ b/Engine/Source/Runtime/Core/Private/HTML5/HTML5PlatformFile.cpp
@@ -267,7 +267,7 @@ public:
 			{
 				FileInfo.st_mode |= S_IWUSR;
 			}
-			return chmod(TCHAR_TO_UTF8(*NormalizeFilename(Filename)), FileInfo.st_mode) > 0;
+			return chmod(TCHAR_TO_UTF8(*NormalizeFilename(Filename)), FileInfo.st_mode) == 0;
 		}
 		return false;
 	}
diff --git a/Engine/Source/Runtime/Core/Private/IOS/IOSAppDelegate.cpp b/Engine/Source/Runtime/Core/Private/IOS/IOSAppDelegate.cpp
--- a/Engine/Source/Runtime/Core/Private/IOS/IOSAppDelegate.cpp
+++ b/Engine/Source/Runtime/Core/Private/IOS/IOSAppDelegate.cpp
@@ -831,6 +831,12 @@ void InstallSignalHandlers()
     }, TStatId(), NULL, ENamedThreads::GameThread);
 }
 
+- (NSUInteger)application:(UIApplication*)application supportedInterfaceOrientationsForWindow:(UIWindow*)window
+{
+ 	// UIImagePickerController or GameCenter might have portrait-only variant and will throw exception if portrait is not supported here
+ 	return UIInterfaceOrientationMaskAll;
+}
+
 - (void)application:(UIApplication *)application didReceiveLocalNotification:(UILocalNotification *)notification
 {
 	NSString*	activationEvent = (NSString*)[notification.userInfo objectForKey: @"ActivationEvent"];
diff --git a/Engine/Source/Runtime/Core/Private/IOS/IOSPlatformFile.cpp b/Engine/Source/Runtime/Core/Private/IOS/IOSPlatformFile.cpp
--- a/Engine/Source/Runtime/Core/Private/IOS/IOSPlatformFile.cpp
+++ b/Engine/Source/Runtime/Core/Private/IOS/IOSPlatformFile.cpp
@@ -478,7 +478,7 @@ bool FIOSPlatformFile::SetReadOnly(const TCHAR* Filename, bool bNewReadOnlyValue
 		{
 			FileInfo.st_mode |= S_IWUSR;
 		}
-		return chmod(TCHAR_TO_UTF8(*IOSFilename), FileInfo.st_mode);
+		return chmod(TCHAR_TO_UTF8(*IOSFilename), FileInfo.st_mode) == 0;
 	}
 	return false;
 }
diff --git a/Engine/Source/Runtime/Core/Private/IOS/IOSPlatformMisc.cpp b/Engine/Source/Runtime/Core/Private/IOS/IOSPlatformMisc.cpp
--- a/Engine/Source/Runtime/Core/Private/IOS/IOSPlatformMisc.cpp
+++ b/Engine/Source/Runtime/Core/Private/IOS/IOSPlatformMisc.cpp
@@ -689,9 +689,7 @@ FString FIOSPlatformMisc::GetUniqueDeviceId()
 		NSUUID* Id = [[UIDevice currentDevice] identifierForVendor];
 		if (Id != nil)
 		{
-			NSString* IdfvString = [Id UUIDString];
-			FString IDFV(IdfvString);
-			return IDFV;
+			return FString([[Id UUIDString] autorelease]);
 		}
 	}
 
diff --git a/Engine/Source/Runtime/Core/Private/IOS/IOSView.cpp b/Engine/Source/Runtime/Core/Private/IOS/IOSView.cpp
--- a/Engine/Source/Runtime/Core/Private/IOS/IOSView.cpp
+++ b/Engine/Source/Runtime/Core/Private/IOS/IOSView.cpp
@@ -337,7 +337,7 @@ id<MTLDevice> GMetalDevice = nil;
 #if HAS_METAL
 - (id<CAMetalDrawable>)MakeDrawable
 {
-	return [(CAMetalLayer*)self.layer nextDrawable];
+    return ![IOSAppDelegate GetDelegate].bIsSuspended ? [(CAMetalLayer*)self.layer nextDrawable] : nil;
 }
 #endif
 
diff --git a/Engine/Source/Runtime/Core/Private/Linux/LinuxPlatformFile.cpp b/Engine/Source/Runtime/Core/Private/Linux/LinuxPlatformFile.cpp
--- a/Engine/Source/Runtime/Core/Private/Linux/LinuxPlatformFile.cpp
+++ b/Engine/Source/Runtime/Core/Private/Linux/LinuxPlatformFile.cpp
@@ -679,7 +679,17 @@ bool FLinuxPlatformFile::MoveFile(const TCHAR* To, const TCHAR* From)
 		return false;
 	}
 
-	return rename(TCHAR_TO_UTF8(*CaseSensitiveFilename), TCHAR_TO_UTF8(*NormalizeFilename(To))) != -1;
+	int32 Result = rename(TCHAR_TO_UTF8(*CaseSensitiveFilename), TCHAR_TO_UTF8(*NormalizeFilename(To)));
+	if (Result == -1 && errno == EXDEV)
+	{
+		// Copy the file if rename failed because To and From are on different file systems
+		if (CopyFile(To, *CaseSensitiveFilename))
+		{
+			DeleteFile(*CaseSensitiveFilename);
+			Result = 0;
+		}
+	}
+	return Result != -1;
 }
 
 bool FLinuxPlatformFile::SetReadOnly(const TCHAR* Filename, bool bNewReadOnlyValue)
@@ -702,7 +712,7 @@ bool FLinuxPlatformFile::SetReadOnly(const TCHAR* Filename, bool bNewReadOnlyVal
 		{
 			FileInfo.st_mode |= S_IWUSR;
 		}
-		return chmod(TCHAR_TO_UTF8(*CaseSensitiveFilename), FileInfo.st_mode);
+		return chmod(TCHAR_TO_UTF8(*CaseSensitiveFilename), FileInfo.st_mode) == 0;
 	}
 	return false;
 }
diff --git a/Engine/Source/Runtime/Core/Private/Linux/LinuxPlatformMisc.cpp b/Engine/Source/Runtime/Core/Private/Linux/LinuxPlatformMisc.cpp
--- a/Engine/Source/Runtime/Core/Private/Linux/LinuxPlatformMisc.cpp
+++ b/Engine/Source/Runtime/Core/Private/Linux/LinuxPlatformMisc.cpp
@@ -86,12 +86,6 @@ bool FLinuxPlatformMisc::ControlScreensaver(EScreenSaverAction Action)
 	return true;
 }
 
-const TCHAR* FLinuxPlatformMisc::RootDir()
-{
-	const TCHAR* TrueRootDir = FGenericPlatformMisc::RootDir();
-	return TrueRootDir;
-}
-
 void FLinuxPlatformMisc::NormalizePath(FString& InPath)
 {
 	// only expand if path starts with ~, e.g. ~/ should be expanded, /~ should not
diff --git a/Engine/Source/Runtime/Core/Private/Linux/LinuxPlatformSplash.cpp b/Engine/Source/Runtime/Core/Private/Linux/LinuxPlatformSplash.cpp
--- a/Engine/Source/Runtime/Core/Private/Linux/LinuxPlatformSplash.cpp
+++ b/Engine/Source/Runtime/Core/Private/Linux/LinuxPlatformSplash.cpp
@@ -925,31 +925,8 @@ void FLinuxPlatformSplash::Show( )
 
 	// first look for the splash, do not init anything if not found
 	{
-		const TCHAR* EditorSplashes[] =
-		{
-			TEXT("EdSplash"),
-			TEXT("EdSplashDefault"),
-			nullptr
-		};
-
-		const TCHAR* GameSplashes[] =
-		{
-			TEXT("Splash"),
-			TEXT("SplashDefault"),
-			nullptr
-		};
-
-		const TCHAR** SplashesToTry = GIsEditor? EditorSplashes : GameSplashes;
-		bool bSplashFound = false;
-		for (const TCHAR* SplashImage = *SplashesToTry; SplashImage != nullptr && !bSplashFound; SplashImage = *(++SplashesToTry))
-		{
-			if (GetSplashPath(SplashImage, GSplashPath, IsCustom))
-			{
-				bSplashFound = true;
-			}
-		}
-
-		if (!bSplashFound)
+		const TCHAR* SplashImage = GIsEditor ? (GameName.IsEmpty() ? TEXT("EdSplashDefault") : TEXT("EdSplash")) : (GameName.IsEmpty() ? TEXT("SplashDefault") : TEXT("Splash"));
+		if (!GetSplashPath(SplashImage, GSplashPath, IsCustom))
 		{
 			UE_LOG(LogHAL, Warning, TEXT("Splash screen image not found."));
 			return;	// early out
diff --git a/Engine/Source/Runtime/Core/Private/Misc/ConfigCacheIni.cpp b/Engine/Source/Runtime/Core/Private/Misc/ConfigCacheIni.cpp
--- a/Engine/Source/Runtime/Core/Private/Misc/ConfigCacheIni.cpp
+++ b/Engine/Source/Runtime/Core/Private/Misc/ConfigCacheIni.cpp
@@ -137,6 +137,36 @@ FString FConfigValue::CollapseValue(const FString& InExpandedValue)
 	return CollapsedValue;
 }
 
+#if !UE_BUILD_SHIPPING
+/**
+* Checks if the section name is the expected name format (long package name or simple name)
+*/
+static void CheckLongSectionNames(const TCHAR* Section, const FConfigFile* File)
+{
+	if (!FPlatformProperties::RequiresCookedData())
+	{
+		// Guard against short names in ini files.
+		if (FCString::Strnicmp(Section, TEXT("/Script/"), 8) == 0)
+		{
+			// Section is a long name
+			if (File->Find(Section + 8))
+			{
+				UE_LOG(LogConfig, Fatal, TEXT("Short config section found while looking for %s"), Section);
+			}
+		}
+		else
+		{
+			// Section is a short name
+			FString LongName = FString(TEXT("/Script/")) + Section;
+			if (File->Find(*LongName))
+			{
+				UE_LOG(LogConfig, Fatal, TEXT("Short config section used instead of long %s"), Section);
+			}
+		}
+	}
+}
+#endif // UE_BUILD_SHIPPING
+
 /*-----------------------------------------------------------------------------
 	FConfigSection
 -----------------------------------------------------------------------------*/
@@ -1341,6 +1371,31 @@ bool FConfigFile::GetBool(const TCHAR* Section, const TCHAR* Key, bool& Value )
 	return 0;
 }
 
+int32 FConfigFile::GetArray(const TCHAR* Section, const TCHAR* Key, TArray<FString>& Value) const
+{
+	const FConfigSection* Sec = Find(Section);
+	if (Sec != nullptr)
+	{
+		TArray<FConfigValue> RemapArray;
+		Sec->MultiFind(Key, RemapArray);
+
+		// TMultiMap::MultiFind will return the results in reverse order
+		Value.AddZeroed(RemapArray.Num());
+		for (int32 RemapIndex = RemapArray.Num() - 1, Index = 0; RemapIndex >= 0; RemapIndex--, Index++)
+		{
+			Value[Index] = RemapArray[RemapIndex].GetValue();
+		}
+	}
+#if !UE_BUILD_SHIPPING
+	else
+	{
+		CheckLongSectionNames(Section, this);
+	}
+#endif
+
+	return Value.Num();
+}
+
 
 
 void FConfigFile::SetString( const TCHAR* Section, const TCHAR* Key, const TCHAR* Value )
@@ -1780,36 +1835,6 @@ void FConfigCacheIni::Detach(const FString& Filename)
 		File->NoSave = 1;
 }
 
-#if !UE_BUILD_SHIPPING
-/**
- * Checks if the section name is the expected name format (long package name or simple name)
- */
-static void CheckLongSectionNames(const TCHAR* Section, FConfigFile* File)
-{
-	if (!FPlatformProperties::RequiresCookedData())
-	{
-		// Guard against short names in ini files.
-		if (FCString::Strnicmp(Section, TEXT("/Script/"), 8) == 0)
-		{
-			// Section is a long name
-			if (File->Find( Section + 8 ))
-			{
-				UE_LOG(LogConfig, Fatal, TEXT("Short config section found while looking for %s"), Section);
-			}
-		}
-		else
-		{
-			// Section is a short name
-			FString LongName = FString(TEXT("/Script/")) + Section;
-			if (File->Find( *LongName ))
-			{
-				UE_LOG(LogConfig, Fatal, TEXT("Short config section used instead of long %s"), Section);
-			}
-		}
-	}
-}
-#endif // UE_BUILD_SHIPPING
-
 bool FConfigCacheIni::GetString( const TCHAR* Section, const TCHAR* Key, FString& Value, const FString& Filename )
 {
 	FRemoteConfig::Get()->FinishRead(*Filename); // Ensure the remote file has been loaded and processed
@@ -3204,7 +3229,7 @@ void FConfigCacheIni::LoadConsoleVariablesFromINI()
 		IConsoleManager::Get().CallAllConsoleVariableSinks();
 }
 
-void FConfigFile::UpdateSections(const TCHAR* DiskFilename, const TCHAR* IniRootName/*=nullptr*/)
+void FConfigFile::UpdateSections(const TCHAR* DiskFilename, const TCHAR* IniRootName/*=nullptr*/, const TCHAR* OverridePlatform/*=nullptr*/)
 {
 	// since we don't want any modifications to other sections, we manually process the file, not read into sections, etc
 	FString DiskFile;
@@ -3251,7 +3276,7 @@ void FConfigFile::UpdateSections(const TCHAR* DiskFilename, const TCHAR* IniRoot
 	{
 		// get the standard ini files
 		SourceIniHierarchy.Empty();
-		GetSourceIniHierarchyFilenames(IniRootName, nullptr, *FPaths::EngineConfigDir(), *FPaths::SourceConfigDir(), SourceIniHierarchy, false);
+		GetSourceIniHierarchyFilenames(IniRootName, OverridePlatform, *FPaths::EngineConfigDir(), *FPaths::SourceConfigDir(), SourceIniHierarchy, false);
 
 		// now chop off this file and any after it
 		for (auto& HierarchyFileIt : SourceIniHierarchy)
@@ -3273,7 +3298,7 @@ void FConfigFile::UpdateSections(const TCHAR* DiskFilename, const TCHAR* IniRoot
 		SourceConfigFile = new FConfigFile();
 
 		// now when Write it called below, it will diff against SourceIniHierarchy
-		LoadIniFileHierarchy(SourceIniHierarchy, *SourceConfigFile, true );
+		LoadIniFileHierarchy(SourceIniHierarchy, *SourceConfigFile, true);
 
 	}
 
diff --git a/Engine/Source/Runtime/Core/Private/Misc/SecureHash.cpp b/Engine/Source/Runtime/Core/Private/Misc/SecureHash.cpp
--- a/Engine/Source/Runtime/Core/Private/Misc/SecureHash.cpp
+++ b/Engine/Source/Runtime/Core/Private/Misc/SecureHash.cpp
@@ -669,6 +669,13 @@ void FSHA1::Transform(uint32 *state, const uint8 *buffer)
 	state[4] += e;
 }
 
+// do not remove #undefs, or you can get name collision with libc++'s <functional>
+#undef _R0
+#undef _R1
+#undef _R2
+#undef _R3
+#undef _R4
+
 // Use this function to hash in binary data
 void FSHA1::Update(const uint8 *data, uint32 len)
 {
diff --git a/Engine/Source/Runtime/Core/Public/Android/AndroidProperties.h b/Engine/Source/Runtime/Core/Public/Android/AndroidProperties.h
--- a/Engine/Source/Runtime/Core/Public/Android/AndroidProperties.h
+++ b/Engine/Source/Runtime/Core/Public/Android/AndroidProperties.h
@@ -111,6 +111,11 @@ struct FAndroidPlatformProperties
 	{
 		return true;
 	}
+
+	static FORCEINLINE bool SupportsAudioStreaming()
+	{
+		return true;
+	}
 };
 
 struct FAndroid_PVRTCPlatformProperties : public FAndroidPlatformProperties
diff --git a/Engine/Source/Runtime/Core/Public/GenericPlatform/GenericPlatformMisc.h b/Engine/Source/Runtime/Core/Public/GenericPlatform/GenericPlatformMisc.h
--- a/Engine/Source/Runtime/Core/Public/GenericPlatform/GenericPlatformMisc.h
+++ b/Engine/Source/Runtime/Core/Public/GenericPlatform/GenericPlatformMisc.h
@@ -1079,6 +1079,12 @@ public:
 		return false;
 	}
 
+	/**
+	 * Returns a list of platforms that are confidential in nature. To avoid hardcoding the list, this
+	 * looks on disk the first time for special files, so it is non-instant.
+	 */
+	static const TArray<FString>& GetConfidentialPlatforms();
+
 #if !UE_BUILD_SHIPPING
 protected:
 	/** Whether the user should be prompted to allow for a remote debugger to be attached */
diff --git a/Engine/Source/Runtime/Core/Public/HAL/Platform.h b/Engine/Source/Runtime/Core/Public/HAL/Platform.h
--- a/Engine/Source/Runtime/Core/Public/HAL/Platform.h
+++ b/Engine/Source/Runtime/Core/Public/HAL/Platform.h
@@ -336,6 +336,10 @@
 	#define PLATFORM_USES_FACE_BUTTON_RIGHT_FOR_ACCEPT		0
 #endif
 
+#ifndef PLATFORM_SUPPORTS_EARLY_MOVIE_PLAYBACK
+	#define PLATFORM_SUPPORTS_EARLY_MOVIE_PLAYBACK			0
+#endif
+
 // deprecated, do not use
 #define PLATFORM_HAS_THREADSAFE_RHIGetRenderQueryResult	#
 #define PLATFORM_SUPPORTS_RHI_THREAD #
diff --git a/Engine/Source/Runtime/Core/Public/IOS/IOSPlatformMisc.h b/Engine/Source/Runtime/Core/Public/IOS/IOSPlatformMisc.h
--- a/Engine/Source/Runtime/Core/Public/IOS/IOSPlatformMisc.h
+++ b/Engine/Source/Runtime/Core/Public/IOS/IOSPlatformMisc.h
@@ -220,8 +220,8 @@ struct CORE_API FIOSPlatformMisc : public FGenericPlatformMisc
 			L"IPadPro",
 			L"AppleTV",
 			L"IPhoneSE",
-			L"IPadPro_12.9",
-			L"IPadPro_9.7",
+			L"IPadPro129",
+			L"IPadPro97",
 			L"Unknown",
 		};
 		static_assert((sizeof(IOSDeviceNames) / sizeof(IOSDeviceNames[0])) == ((int32)IOS_Unknown + 1), "Mismatched IOSDeviceNames and EIOSDevice.");
diff --git a/Engine/Source/Runtime/Core/Public/Linux/LinuxPlatformMisc.h b/Engine/Source/Runtime/Core/Public/Linux/LinuxPlatformMisc.h
--- a/Engine/Source/Runtime/Core/Public/Linux/LinuxPlatformMisc.h
+++ b/Engine/Source/Runtime/Core/Public/Linux/LinuxPlatformMisc.h
@@ -79,7 +79,6 @@ struct CORE_API FLinuxPlatformMisc : public FGenericPlatformMisc
 	static void ClipboardCopy(const TCHAR* Str);
 	static void ClipboardPaste(class FString& Dest);
 
-	static const TCHAR* RootDir();
 	static void NormalizePath(FString& InPath);
 
 	static const TCHAR* GetPathVarDelimiter()
@@ -146,11 +145,6 @@ struct CORE_API FLinuxPlatformMisc : public FGenericPlatformMisc
 	 */
 	static uint32 GetCPUInfo();
 
-	static const TCHAR* EngineDir()
-	{
-		return TEXT("../../../Engine/");
-	}
-
 	/**
 	 * Initializes video (and not only) subsystem.
 	 */
diff --git a/Engine/Source/Runtime/Core/Public/Misc/ConfigCacheIni.h b/Engine/Source/Runtime/Core/Public/Misc/ConfigCacheIni.h
--- a/Engine/Source/Runtime/Core/Public/Misc/ConfigCacheIni.h
+++ b/Engine/Source/Runtime/Core/Public/Misc/ConfigCacheIni.h
@@ -324,6 +324,7 @@ public:
 	CORE_API bool GetText( const TCHAR* Section, const TCHAR* Key, FText& Value ) const;
 	CORE_API bool GetInt64( const TCHAR* Section, const TCHAR* Key, int64& Value ) const;
 	CORE_API bool GetBool( const TCHAR* Section, const TCHAR* Key, bool& Value ) const;
+	CORE_API int32 GetArray(const TCHAR* Section, const TCHAR* Key, TArray<FString>& Value) const;
 
 	CORE_API void SetString( const TCHAR* Section, const TCHAR* Key, const TCHAR* Value );
 	CORE_API void SetText( const TCHAR* Section, const TCHAR* Key, const FText& Value );
@@ -353,7 +354,7 @@ public:
 	 *
 	 * @param IniRootName the name (like "Engine") to use to load a .ini hierarchy to diff against
 	 */
-	CORE_API void UpdateSections(const TCHAR* DiskFilename, const TCHAR* IniRootName=NULL);
+	CORE_API void UpdateSections(const TCHAR* DiskFilename, const TCHAR* IniRootName=nullptr, const TCHAR* OverridePlatform=nullptr);
 
 	/**
 	 * Update a single property in the config file, for the section that is specified.
diff --git a/Engine/Source/Runtime/Core/Public/Serialization/Archive.h b/Engine/Source/Runtime/Core/Public/Serialization/Archive.h
--- a/Engine/Source/Runtime/Core/Public/Serialization/Archive.h
+++ b/Engine/Source/Runtime/Core/Public/Serialization/Archive.h
@@ -19,7 +19,11 @@ template<class TEnum> class TEnumAsByte;
 
 
 // Temporary while we shake out the EDL at boot
+#if PLATFORM_SWITCH
+#define USE_EVENT_DRIVEN_ASYNC_LOAD_AT_BOOT_TIME (0)
+#else
 #define USE_EVENT_DRIVEN_ASYNC_LOAD_AT_BOOT_TIME (1)
+#endif
 
 #if USE_EVENT_DRIVEN_ASYNC_LOAD_AT_BOOT_TIME
 	#define EVENT_DRIVEN_ASYNC_LOAD_ACTIVE_AT_RUNTIME (1)
diff --git a/Engine/Source/Runtime/CoreUObject/Private/Blueprint/BlueprintSupport.cpp b/Engine/Source/Runtime/CoreUObject/Private/Blueprint/BlueprintSupport.cpp
--- a/Engine/Source/Runtime/CoreUObject/Private/Blueprint/BlueprintSupport.cpp
+++ b/Engine/Source/Runtime/CoreUObject/Private/Blueprint/BlueprintSupport.cpp
@@ -291,7 +291,7 @@ struct FPreloadMembersHelper
 			UObject* CurrentObject = *it;
 			if (!CurrentObject->HasAnyFlags(RF_LoadCompleted))
 			{
-				check(!GEventDrivenLoaderEnabled);
+				check(!GEventDrivenLoaderEnabled || !EVENT_DRIVEN_ASYNC_LOAD_ACTIVE_AT_RUNTIME);
 				CurrentObject->SetFlags(RF_NeedLoad);
 				if (FLinkerLoad* Linker = CurrentObject->GetLinker())
 				{
@@ -306,7 +306,7 @@ struct FPreloadMembersHelper
 	{
 		if (InObject && !InObject->HasAnyFlags(RF_LoadCompleted))
 		{
-			check(!GEventDrivenLoaderEnabled);
+			check(!GEventDrivenLoaderEnabled || !EVENT_DRIVEN_ASYNC_LOAD_ACTIVE_AT_RUNTIME);
 			InObject->SetFlags(RF_NeedLoad);
 			if (FLinkerLoad* Linker = InObject->GetLinker())
 			{
@@ -1332,7 +1332,7 @@ void FLinkerLoad::FinalizeBlueprint(UClass* LoadClass)
 			// to re-run the serialization)
 			if ( (SuperCDO != nullptr) && !SuperCDO->HasAnyFlags(RF_NeedLoad|RF_LoadCompleted) )
 			{
-				check(!GEventDrivenLoaderEnabled);
+				check(!GEventDrivenLoaderEnabled || !EVENT_DRIVEN_ASYNC_LOAD_ACTIVE_AT_RUNTIME);
 				SuperCDO->SetFlags(RF_NeedLoad);
 			}
 			SuperLinker->FinalizeBlueprint(SuperClass);
@@ -2270,7 +2270,7 @@ void FDeferredObjInitializerTracker::ResolveDeferredSubObjects(UObject* CDO)
 		{
 			if (FLinkerLoad* SubObjLinker = SubObjArchetype->GetLinker())
 			{
-				check(!GEventDrivenLoaderEnabled);
+				check(!GEventDrivenLoaderEnabled || !EVENT_DRIVEN_ASYNC_LOAD_ACTIVE_AT_RUNTIME);
 				SubObjArchetype->SetFlags(RF_NeedLoad);
 				SubObjLinker->Preload(SubObjArchetype);
 			}
diff --git a/Engine/Source/Runtime/CoreUObject/Private/Serialization/AsyncLoading.cpp b/Engine/Source/Runtime/CoreUObject/Private/Serialization/AsyncLoading.cpp
--- a/Engine/Source/Runtime/CoreUObject/Private/Serialization/AsyncLoading.cpp
+++ b/Engine/Source/Runtime/CoreUObject/Private/Serialization/AsyncLoading.cpp
@@ -773,7 +773,7 @@ struct FEDLBootNotificationManager
 		{
 			UE_LOG(LogStreaming, Fatal, TEXT("Initial load is complete, but we still have %d waiting packages."), PathToWaitingPackageNodes.Num());
 		}
-		if (GEventDrivenLoaderEnabled && PathsToFire.Num())
+		if (GEventDrivenLoaderEnabled && PathsToFire.Num() && USE_EVENT_DRIVEN_ASYNC_LOAD_AT_BOOT_TIME)
 		{
 			for (FName Path : PathsToFire)
 			{
@@ -4239,8 +4239,6 @@ EAsyncPackageState::Type FAsyncLoadingThread::ProcessAsyncLoading(int32& OutPack
 				}
 				if (LocalLoadingState == EAsyncPackageState::Complete)
 				{
-
-					// We're done, at least on this thread, so we can remove the package now.
 					{
 #if THREADSAFE_UOBJECTS
 						FScopeLock LockAsyncPackages(&AsyncPackagesCritical);
@@ -4250,7 +4248,9 @@ EAsyncPackageState::Type FAsyncLoadingThread::ProcessAsyncLoading(int32& OutPack
 						AsyncPackages.RemoveAt(PackageIndex);
 						AsyncPackagesReadyForTick.RemoveAt(0, 1, false); //@todoio this should maybe be a heap or something to avoid the removal cost
 					}
-				AddToLoadedPackages(Package);
+
+					// We're done, at least on this thread, so we can remove the package now.
+					AddToLoadedPackages(Package);
 				}
 				if (IsTimeLimitExceeded(TickStartTime, bUseTimeLimit, TimeLimit, TEXT("TickAsyncPackage")))
 				{
@@ -7374,7 +7374,7 @@ void FArchiveAsync2::ReadCallback(bool bWasCancelled, IAsyncReadRequest* Request
 	}
 	else if (LoadPhase == ELoadPhase::WaitingForSummary)
 	{
-		check(!GEventDrivenLoaderEnabled);
+		check(!GEventDrivenLoaderEnabled || !EVENT_DRIVEN_ASYNC_LOAD_ACTIVE_AT_RUNTIME);
 		uint8* Mem = Request->GetReadResults();
 		if (!Mem)
 		{
@@ -7780,6 +7780,29 @@ void FArchiveAsync2::FirstExportStarting()
 	ExportReadTime = FPlatformTime::Seconds();
 	LogItem(TEXT("Exports"));
 	LoadPhase = ELoadPhase::ProcessingExports;
+
+	if (GEventDrivenLoaderEnabled && !EVENT_DRIVEN_ASYNC_LOAD_ACTIVE_AT_RUNTIME)
+	{
+		FlushCache();
+		if (Handle)
+		{
+			delete Handle;
+			Handle = nullptr;
+		}
+
+		HeaderSizeWhenReadingExportsFromSplitFile = HeaderSize;
+		FileName = FPaths::GetBaseFilename(FileName, false) + TEXT(".uexp");
+
+		Handle = FPlatformFileManager::Get().GetPlatformFile().OpenAsyncRead(*FileName);
+		check(Handle); // this generally cannot fail because it is async
+
+		check(!SizeRequestPtr);
+		SizeRequestPtr = Handle->SizeRequest();
+		if (SizeRequestPtr->PollCompletion())
+		{
+			TotalSize(); // complete the request
+		}
+	}
 }
 
 IAsyncReadRequest* FArchiveAsync2::MakeEventDrivenPrecacheRequest(int64 Offset, int64 BytesToRead, FAsyncFileCallBack* CompleteCallback)
diff --git a/Engine/Source/Runtime/CoreUObject/Private/UObject/LinkerLoad.cpp b/Engine/Source/Runtime/CoreUObject/Private/UObject/LinkerLoad.cpp
--- a/Engine/Source/Runtime/CoreUObject/Private/UObject/LinkerLoad.cpp
+++ b/Engine/Source/Runtime/CoreUObject/Private/UObject/LinkerLoad.cpp
@@ -1411,7 +1411,7 @@ FLinkerLoad::ELinkerStatus FLinkerLoad::SerializeNameMap()
 			if (GNewAsyncIO && bLoaderIsFArchiveAsync2)
 			{
 				bFinishedPrecaching = GetFArchiveAsync2Loader()->ReadyToStartReadingHeader(bUseTimeLimit, bUseFullTimeLimit, TickStartTime, TimeLimit);
-				check(!GEventDrivenLoaderEnabled || bFinishedPrecaching);
+				check(!GEventDrivenLoaderEnabled || bFinishedPrecaching || !EVENT_DRIVEN_ASYNC_LOAD_ACTIVE_AT_RUNTIME);
 			}
 			else
 			{
@@ -2312,7 +2312,7 @@ FLinkerLoad::EVerifyResult FLinkerLoad::VerifyImport(int32 ImportIndex)
 				// serialize in the properties of the redirector (to get the object the redirector point to)
 				// Always load redirectors in case there was a circular dependency. This will allow inner redirector
 				// references to always serialize fully here before accessing the DestinationObject
-				check(!GEventDrivenLoaderEnabled);
+				check(!GEventDrivenLoaderEnabled || !EVENT_DRIVEN_ASYNC_LOAD_ACTIVE_AT_RUNTIME);
 				Redir->SetFlags(RF_NeedLoad);
 				Preload(Redir);
 
@@ -3197,7 +3197,7 @@ void FLinkerLoad::Preload( UObject* Object )
 	{
 		if (Object->GetLinker() == this)
 		{
-			check(!GEventDrivenLoaderEnabled || !bLockoutLegacyOperations);
+			check(!GEventDrivenLoaderEnabled || !bLockoutLegacyOperations || !EVENT_DRIVEN_ASYNC_LOAD_ACTIVE_AT_RUNTIME);
 #if USE_CIRCULAR_DEPENDENCY_LOAD_DEFERRING
 			bool const bIsNonNativeObject = !Object->GetOutermost()->HasAnyPackageFlags(PKG_CompiledIn);
 			// we can determine that this is a blueprint class/struct by checking if it 
@@ -3344,7 +3344,7 @@ void FLinkerLoad::Preload( UObject* Object )
 
 							// reset the flag and return (don't worry, we make
 							// sure to force load this later)
-							check(!GEventDrivenLoaderEnabled);
+							check(!GEventDrivenLoaderEnabled || !EVENT_DRIVEN_ASYNC_LOAD_ACTIVE_AT_RUNTIME);
 							Object->SetFlags(RF_NeedLoad);
 							return;
 						}
@@ -3570,7 +3570,7 @@ UObject* FLinkerLoad::CreateExport( int32 Index )
 	// Check whether we already loaded the object and if not whether the context flags allow loading it.
 	if( !Export.Object && !FilterExport(Export) ) // for some acceptable position, it was not "not for" 
 	{
-		check(!GEventDrivenLoaderEnabled || !bLockoutLegacyOperations);
+		check(!GEventDrivenLoaderEnabled || !bLockoutLegacyOperations || !EVENT_DRIVEN_ASYNC_LOAD_ACTIVE_AT_RUNTIME);
 		check(Export.ObjectName!=NAME_None || !(Export.ObjectFlags&RF_Public));
 		check(IsLoading());
 
@@ -3689,7 +3689,7 @@ UObject* FLinkerLoad::CreateExport( int32 Index )
 					const UClass* AsClass = dynamic_cast<UClass*>(SuperStruct);
 					if (AsClass && !AsClass->ClassDefaultObject)
 					{
-						check(!GEventDrivenLoaderEnabled);
+						check(!GEventDrivenLoaderEnabled || !EVENT_DRIVEN_ASYNC_LOAD_ACTIVE_AT_RUNTIME);
 						SuperStruct->SetFlags(RF_NeedLoad);
 						Preload(SuperStruct);
 					}
@@ -3888,7 +3888,7 @@ UObject* FLinkerLoad::CreateExport( int32 Index )
 					if (!Export.Object->HasAnyFlags(RF_LoadCompleted) &&
 						(Export.Object->HasAnyFlags(RF_DefaultSubObject) || (ThisParent && ThisParent->IsTemplate(RF_ClassDefaultObject))))
 					{
-						check(!GEventDrivenLoaderEnabled);
+						check(!GEventDrivenLoaderEnabled || !EVENT_DRIVEN_ASYNC_LOAD_ACTIVE_AT_RUNTIME);
 						Export.Object->SetFlags(RF_NeedLoad | RF_NeedPostLoad | RF_NeedPostLoadSubobjects | RF_WasLoaded);
 					}
 				}
@@ -4131,7 +4131,7 @@ bool FLinkerLoad::IsImportNative(const int32 Index) const
 // Return the loaded object corresponding to an import index; any errors are fatal.
 UObject* FLinkerLoad::CreateImport( int32 Index )
 {
-	check(!GEventDrivenLoaderEnabled || !bLockoutLegacyOperations);
+	check(!GEventDrivenLoaderEnabled || !bLockoutLegacyOperations || !EVENT_DRIVEN_ASYNC_LOAD_ACTIVE_AT_RUNTIME);
 
 	FScopedCreateImportCounter ScopedCounter( this, Index );
 	FObjectImport& Import = ImportMap[ Index ];
diff --git a/Engine/Source/Runtime/CoreUObject/Private/UObject/LinkerPlaceholderBase.inl b/Engine/Source/Runtime/CoreUObject/Private/UObject/LinkerPlaceholderBase.inl
--- a/Engine/Source/Runtime/CoreUObject/Private/UObject/LinkerPlaceholderBase.inl
+++ b/Engine/Source/Runtime/CoreUObject/Private/UObject/LinkerPlaceholderBase.inl
@@ -35,7 +35,7 @@ int32 TLinkerImportPlaceholder<PlaceholderType>::ResolveAllPlaceholderReferences
 				// here, needing to use the ClassB redirector, but we haven't 
 				// returned up the stack for it to be set yet... here we force 
 				// it to finish preloading (like we do in VerifyImport): 
-				check(!GEventDrivenLoaderEnabled);
+				check(!GEventDrivenLoaderEnabled || !EVENT_DRIVEN_ASYNC_LOAD_ACTIVE_AT_RUNTIME);
 				ReplacementRedirector->SetFlags(RF_NeedLoad);
 				ReplacementLinker->Preload(ReplacementRedirector);
 			}
diff --git a/Engine/Source/Runtime/CoreUObject/Private/UObject/Obj.cpp b/Engine/Source/Runtime/CoreUObject/Private/UObject/Obj.cpp
--- a/Engine/Source/Runtime/CoreUObject/Private/UObject/Obj.cpp
+++ b/Engine/Source/Runtime/CoreUObject/Private/UObject/Obj.cpp
@@ -1812,6 +1812,18 @@ void UObject::LoadConfig( UClass* ConfigClass/*=NULL*/, const TCHAR* InFilename/
 
 	const bool bPerObject = UsesPerObjectConfig(this);
 
+	// does the class want to override the platform hierarchy (ignored if we passd in a specific ini file),
+	// and if the name isn't the current running platform (no need to load extra files if already in GConfig)
+	bool bUseConfigOverride = InFilename == nullptr && GetConfigOverridePlatform() != nullptr &&
+		FCString::Stricmp(GetConfigOverridePlatform(), ANSI_TO_TCHAR(FPlatformProperties::IniPlatformName())) != 0;
+	FConfigFile OverrideConfig;
+	if (bUseConfigOverride)
+	{
+		// load into a local ini file
+		FConfigCacheIni::LoadLocalIniFile(OverrideConfig, *GetClass()->ClassConfigName.ToString(), true, GetConfigOverridePlatform());
+	}
+
+
 	FString ClassSection;
 	FName LongCommitName;
 	if ( bPerObject == true )
@@ -1903,7 +1915,16 @@ void UObject::LoadConfig( UClass* ConfigClass/*=NULL*/, const TCHAR* InFilename/
 				}
 
 				FString Value;
-				const bool bFoundValue = GConfig->GetString( *ClassSection, *Key, Value, *PropFileName );
+				bool bFoundValue;
+				if (bUseConfigOverride)
+				{
+					bFoundValue = OverrideConfig.GetString(*ClassSection, *Key, Value);
+				}
+				else
+				{
+					bFoundValue = GConfig->GetString(*ClassSection, *Key, Value, *PropFileName);
+				}
+
 				if (bFoundValue)
 				{
 					if (Property->ImportText(*Value, Property->ContainerPtrToValuePtr<uint8>(this, i), PortFlags, this) == NULL)
@@ -1923,7 +1944,16 @@ void UObject::LoadConfig( UClass* ConfigClass/*=NULL*/, const TCHAR* InFilename/
 		}
 		else
 		{
-			FConfigSection* Sec = GConfig->GetSectionPrivate( *ClassSection, false, true, *PropFileName );
+			FConfigSection* Sec;
+			if (bUseConfigOverride)
+			{
+				Sec = OverrideConfig.Find(*ClassSection);
+			}
+			else
+			{
+				Sec = GConfig->GetSectionPrivate(*ClassSection, false, true, *PropFileName);
+			}
+
 			FConfigSection* AltSec = NULL;
 			//@Package name transition
 			if( Sec )
@@ -2010,9 +2040,7 @@ void UObject::SaveConfig( uint64 Flags, const TCHAR* InFilename, FConfigCacheIni
 		: GetConfigFilename(this);
 
 	// Determine whether the file we are writing is a default file config.
-	const bool bIsADefaultIniWrite = !Filename.Contains(FPaths::GameSavedDir())
-		&& !Filename.Contains(FPaths::EngineSavedDir())
-		&& (FPaths::GetBaseFilename(Filename).StartsWith(TEXT("Default")) || FPaths::GetBaseFilename(Filename).StartsWith(TEXT("User")));
+	const bool bIsADefaultIniWrite = Filename == GetDefaultConfigFilename() || Filename == GetGlobalUserConfigFilename();
 
 	const bool bPerObject = UsesPerObjectConfig(this);
 	FString Section;
@@ -2164,8 +2192,23 @@ void UObject::SaveConfig( uint64 Flags, const TCHAR* InFilename, FConfigCacheIni
 	}
 }
 
+static FString GetFinalOverridePlatform(const UObject* Obj)
+{
+	FString Platform;
+	if (Obj->GetConfigOverridePlatform() != nullptr && FCString::Stricmp(Obj->GetConfigOverridePlatform(), ANSI_TO_TCHAR(FPlatformProperties::IniPlatformName())) != 0)
+	{
+		Platform = Obj->GetConfigOverridePlatform();
+	}
+	return Platform;
+}
+
 FString UObject::GetDefaultConfigFilename() const
 {
+	FString OverridePlatform = GetFinalOverridePlatform(this);
+	if (OverridePlatform.Len())
+	{
+		return FString::Printf(TEXT("%s%s/%s%s.ini"), *FPaths::SourceConfigDir(), *OverridePlatform, *OverridePlatform, *GetClass()->ClassConfigName.ToString());
+	}
 	return FString::Printf(TEXT("%sDefault%s.ini"), *FPaths::SourceConfigDir(), *GetClass()->ClassConfigName.ToString());
 }
 
@@ -2188,12 +2231,19 @@ void UObject::UpdateSingleSectionOfConfigFile(const FString& ConfigIniName)
 
 	ensureMsgf(Config.Num() == 1, TEXT("UObject::UpdateDefaultConfig() caused more files than expected in the Sandbox config cache!"));
 
+	// do we need to use a special platform hierarchy?
+	FString OverridePlatform = GetFinalOverridePlatform(this);
+
 	// make sure SaveConfig wrote only to the file we expected
-	NewFile.UpdateSections(*ConfigIniName, *GetClass()->ClassConfigName.ToString());
+	NewFile.UpdateSections(*ConfigIniName, *GetClass()->ClassConfigName.ToString(), OverridePlatform.Len() ? *OverridePlatform : nullptr);
 
-	// reload the file, so that it refresh the cache internally.
-	FString FinalIniFileName;
-	GConfig->LoadGlobalIniFile(FinalIniFileName, *GetClass()->ClassConfigName.ToString(), NULL, true);
+	// reload the file, so that it refresh the cache internally, unless a non-standard platform was used,
+	// then we don't want to touch GConfig
+	if (OverridePlatform.Len() == 0)
+	{
+		FString FinalIniFileName;
+		GConfig->LoadGlobalIniFile(FinalIniFileName, *GetClass()->ClassConfigName.ToString(), NULL, true);
+	}
 }
 
 void UObject::UpdateDefaultConfigFile(const FString& SpecificFileLocation)
@@ -2238,11 +2288,18 @@ void UObject::UpdateSinglePropertyInConfigFile(const UProperty* InProperty, cons
 		}
 #endif // #if WITH_EDITOR
 
+		// do we need to use a special platform hierarchy?
+		FString OverridePlatform = GetFinalOverridePlatform(this);
+
 		NewFile.UpdateSinglePropertyInSection(*InConfigIniName, *PropertyKey, *SectionName);
 
-		// reload the file, so that it refresh the cache internally.
-		FString FinalIniFileName;
-		GConfig->LoadGlobalIniFile(FinalIniFileName, *GetClass()->ClassConfigName.ToString(), NULL, true);
+		// reload the file, so that it refresh the cache internally, unless a non-standard platform was used,
+		// then we don't want to touch GConfig
+		if (OverridePlatform.Len() == 0)
+		{
+			FString FinalIniFileName;
+			GConfig->LoadGlobalIniFile(FinalIniFileName, *GetClass()->ClassConfigName.ToString(), NULL, true);
+		}
 	}
 	else
 	{
diff --git a/Engine/Source/Runtime/CoreUObject/Public/UObject/Object.h b/Engine/Source/Runtime/CoreUObject/Public/UObject/Object.h
--- a/Engine/Source/Runtime/CoreUObject/Public/UObject/Object.h
+++ b/Engine/Source/Runtime/CoreUObject/Public/UObject/Object.h
@@ -899,6 +899,9 @@ public:
 	 */
 	FString GetGlobalUserConfigFilename() const;
 
+	/** Returns the override config hierarchy platform (if NDAd platforms need defaults to not be in Base*.ini but still want editor to load them) */
+	virtual const TCHAR* GetConfigOverridePlatform() const { return nullptr; }
+
 	/**
 	 * Imports property values from an .ini file.
 	 *
diff --git a/Engine/Source/Runtime/Engine/Classes/DeviceProfiles/DeviceProfile.h b/Engine/Source/Runtime/Engine/Classes/DeviceProfiles/DeviceProfile.h
--- a/Engine/Source/Runtime/Engine/Classes/DeviceProfiles/DeviceProfile.h
+++ b/Engine/Source/Runtime/Engine/Classes/DeviceProfiles/DeviceProfile.h
@@ -36,6 +36,8 @@ class ENGINE_API UDeviceProfile : public UTextureLODSettings
 	/** Flag used in the editor to determine whether the profile is visible in the property matrix */
 	bool bVisible;
 
+	/** This is not a property, it shouldn't be set by the editor */
+	FString ConfigPlatform;
 
 public:
 
@@ -43,6 +45,12 @@ public:
 	UPROPERTY(EditAnywhere, config, Category=ConsoleVariables)
 	TArray<FString> CVars;
 
+	/** Prefer to load the DP from its platform's hierarchy */
+	virtual const TCHAR* GetConfigOverridePlatform() const override
+	{
+		return ConfigPlatform.Len() ? *ConfigPlatform : Super::GetConfigOverridePlatform();
+	}
+
 	/** 
 	 * Get the collection of Console Variables that this profile inherits from its' parents
 	 *
diff --git a/Engine/Source/Runtime/Engine/Classes/DeviceProfiles/DeviceProfileManager.h b/Engine/Source/Runtime/Engine/Classes/DeviceProfiles/DeviceProfileManager.h
--- a/Engine/Source/Runtime/Engine/Classes/DeviceProfiles/DeviceProfileManager.h
+++ b/Engine/Source/Runtime/Engine/Classes/DeviceProfiles/DeviceProfileManager.h
@@ -40,7 +40,7 @@ public:
 	 *
 	 * @return the created profile.
 	 */
-	UDeviceProfile* CreateProfile( const FString& ProfileName, const FString& ProfileType, const FString& ParentName = TEXT("") );
+	UDeviceProfile* CreateProfile(const FString& ProfileName, const FString& ProfileType, const FString& ParentName=TEXT(""), const TCHAR* ConfigPlatform=nullptr);
 
 	/**
 	 * Delete a profile.
diff --git a/Engine/Source/Runtime/Engine/Private/AudioDerivedData.cpp b/Engine/Source/Runtime/Engine/Private/AudioDerivedData.cpp
--- a/Engine/Source/Runtime/Engine/Private/AudioDerivedData.cpp
+++ b/Engine/Source/Runtime/Engine/Private/AudioDerivedData.cpp
@@ -37,7 +37,7 @@ Derived data key generation.
 
 // If you want to bump this version, generate a new guid using
 // VS->Tools->Create GUID and paste it here. https://www.guidgen.com works too.
-#define STREAMEDAUDIO_DERIVEDDATA_VER		TEXT("8886fd5b8a934260aff24cf2642acada")
+#define STREAMEDAUDIO_DERIVEDDATA_VER		TEXT("8486fd5b8a934260a6f44cf2642acada")
 
 /**
  * Computes the derived data key suffix for a SoundWave's Streamed Audio.
diff --git a/Engine/Source/Runtime/Engine/Private/DeviceProfiles/DeviceProfileManager.cpp b/Engine/Source/Runtime/Engine/Private/DeviceProfiles/DeviceProfileManager.cpp
--- a/Engine/Source/Runtime/Engine/Private/DeviceProfiles/DeviceProfileManager.cpp
+++ b/Engine/Source/Runtime/Engine/Private/DeviceProfiles/DeviceProfileManager.cpp
@@ -187,17 +187,24 @@ void UDeviceProfileManager::InitializeCVarsForActiveDeviceProfile(bool bPushSett
 }
 
 
-UDeviceProfile* UDeviceProfileManager::CreateProfile( const FString& ProfileName, const FString& ProfileType, const FString& InSpecifyParentName )
+UDeviceProfile* UDeviceProfileManager::CreateProfile( const FString& ProfileName, const FString& ProfileType, const FString& InSpecifyParentName, const TCHAR* ConfigPlatform )
 {
 	UDeviceProfile* DeviceProfile = FindObject<UDeviceProfile>( GetTransientPackage(), *ProfileName );
-	if( DeviceProfile == NULL )
+	if (DeviceProfile == NULL)
 	{
+		// use ConfigPlatform ini hierarchy to look in for the parent profile
+		// @todo config: we could likely cache local ini files to speed this up,
+		// along with the ones we load in LoadConfig
+		// NOTE: This happens at runtime, so maybe only do this if !RequiresCookedData()?
+		FConfigFile PlatformConfigFile;
+		FConfigCacheIni::LoadLocalIniFile(PlatformConfigFile, TEXT("DeviceProfiles"), true, ConfigPlatform);
+
 		// Build Parent objects first. Important for setup
 		FString ParentName = InSpecifyParentName;
 		if (ParentName.Len() == 0)
 		{
 			const FString SectionName = FString::Printf(TEXT("%s %s"), *ProfileName, *UDeviceProfile::StaticClass()->GetName());
-			GConfig->GetString(*SectionName, TEXT("BaseProfileName"), ParentName, GetDeviceProfileIniName());
+			PlatformConfigFile.GetString(*SectionName, TEXT("BaseProfileName"), ParentName);
 		}
 
 		UObject* ParentObject = nullptr;
@@ -207,13 +214,26 @@ UDeviceProfile* UDeviceProfileManager::CreateProfile( const FString& ProfileName
 			ParentObject = FindObject<UDeviceProfile>(GetTransientPackage(), *ParentName);
 			if (ParentObject == nullptr)
 			{
-				ParentObject = CreateProfile(ParentName, ProfileType);
+				ParentObject = CreateProfile(ParentName, ProfileType, TEXT(""), ConfigPlatform);
 			}
 		}
 
 		// Create the profile after it's parents have been created.
 		DeviceProfile = NewObject<UDeviceProfile>(GetTransientPackage(), *ProfileName);
-		DeviceProfile->DeviceType = DeviceProfile->DeviceType.Len() > 0 ? DeviceProfile->DeviceType : ProfileType;
+		if (ConfigPlatform != nullptr)
+		{
+			// if the config needs to come from a platform, set it now, then reload the config
+			DeviceProfile->ConfigPlatform = ConfigPlatform;
+			DeviceProfile->LoadConfig();
+		}
+
+		// if the config didn't specify a DeviceType, use the passed in one
+		if (DeviceProfile->DeviceType.IsEmpty())
+		{
+			DeviceProfile->DeviceType = ProfileType;
+		}
+
+		// final fixups
 		DeviceProfile->BaseProfileName = DeviceProfile->BaseProfileName.Len() > 0 ? DeviceProfile->BaseProfileName : ParentName;
 		DeviceProfile->Parent = ParentObject;
 		// the DP manager can be marked as Disregard for GC, so what it points to needs to be in the Root set
@@ -269,18 +289,56 @@ void UDeviceProfileManager::LoadProfiles()
 {
 	if( !HasAnyFlags( RF_ClassDefaultObject ) )
 	{
-		TArray< FString > DeviceProfileMapArray;
-		GConfig->GetArray(TEXT("DeviceProfiles"), TEXT("DeviceProfileNameAndTypes"), DeviceProfileMapArray, DeviceProfileFileName);
+		TMap<FString, FString> DeviceProfileToPlatformConfigMap;
+		TArray<FString> ConfidentialPlatforms = FPlatformMisc::GetConfidentialPlatforms();
+		
+		checkf(ConfidentialPlatforms.Contains(FString(FPlatformProperties::IniPlatformName())) == false,
+			TEXT("UDeviceProfileManager::LoadProfiles is called from a confidential platform (%s). Confidential platforms are not expected to be editor/non-cooked builds."), 
+			ANSI_TO_TCHAR(FPlatformProperties::IniPlatformName()));
+
+		// go over all the platforms we find, starting with the current platform
+		for (int32 PlatformIndex = 0; PlatformIndex <= ConfidentialPlatforms.Num(); PlatformIndex++)
+		{
+			// which platform's set of ini files should we load from?
+			FString ConfigLoadPlatform = PlatformIndex == 0 ? FString(FPlatformProperties::IniPlatformName()) : ConfidentialPlatforms[PlatformIndex - 1];
+
+			// load the DP.ini files (from current platform and then by the extra confidential platforms)
+			FConfigFile PlatformConfigFile;
+			FConfigCacheIni::LoadLocalIniFile(PlatformConfigFile, TEXT("DeviceProfiles"), true, *ConfigLoadPlatform);
+
+			// load all of the DeviceProfiles
+			TArray<FString> ProfileDescriptions;
+			PlatformConfigFile.GetArray(TEXT("DeviceProfiles"), TEXT("DeviceProfileNameAndTypes"), ProfileDescriptions);
 
-		for (int32 DeviceProfileIndex = 0; DeviceProfileIndex < DeviceProfileMapArray.Num(); ++DeviceProfileIndex)
+
+			// add them to our collection of profiles by platform
+			for (const FString& Desc : ProfileDescriptions)
+			{
+				if (!DeviceProfileToPlatformConfigMap.Contains(Desc))
+				{
+					DeviceProfileToPlatformConfigMap.Add(Desc, ConfigLoadPlatform);
+				}
+			}
+		}
+
+		// now that we have gathered all the unique DPs, load them from the proper platform hierarchy
+		for (auto It = DeviceProfileToPlatformConfigMap.CreateIterator(); It; ++It)
 		{
-			FString NewDeviceProfileSelectorPlatformName;
-			FString NewDeviceProfileSelectorPlatformType;
-			DeviceProfileMapArray[DeviceProfileIndex].Split(TEXT(","), &NewDeviceProfileSelectorPlatformName, &NewDeviceProfileSelectorPlatformType);
+			// the value of the map is in the format Name,DeviceType (DeviceType is usually platform)
+			FString Name, DeviceType;
+			It.Key().Split(TEXT(","), &Name, &DeviceType);
 
-			if (FindObject<UDeviceProfile>(GetTransientPackage(), *NewDeviceProfileSelectorPlatformName) == NULL)
+			if (FindObject<UDeviceProfile>(GetTransientPackage(), *Name) == NULL)
 			{
-				CreateProfile(NewDeviceProfileSelectorPlatformName, NewDeviceProfileSelectorPlatformType);
+				// set the config platform if it's not the current platform
+				if (It.Value() != FPlatformProperties::IniPlatformName())
+				{
+					CreateProfile(Name, DeviceType, TEXT(""), *It.Value());
+				}
+				else
+				{
+					CreateProfile(Name, DeviceType);
+				}
 			}
 		}
 
diff --git a/Engine/Source/Runtime/Engine/Private/Materials/Material.cpp b/Engine/Source/Runtime/Engine/Private/Materials/Material.cpp
--- a/Engine/Source/Runtime/Engine/Private/Materials/Material.cpp
+++ b/Engine/Source/Runtime/Engine/Private/Materials/Material.cpp
@@ -412,9 +412,12 @@ void UMaterialInterface::InitDefaultMaterials()
 
 void UMaterialInterface::PostCDOContruct()
 {
-	UMaterial::StaticClass()->GetDefaultObject();
-	InitializeSharedSamplerStates();
-	UMaterialInterface::InitDefaultMaterials();
+	if (GEventDrivenLoaderEnabled && EVENT_DRIVEN_ASYNC_LOAD_ACTIVE_AT_RUNTIME)
+	{
+		UMaterial::StaticClass()->GetDefaultObject();
+		InitializeSharedSamplerStates();
+		UMaterialInterface::InitDefaultMaterials();
+	}
 }
 
 
diff --git a/Engine/Source/Runtime/Engine/Private/Texture.cpp b/Engine/Source/Runtime/Engine/Private/Texture.cpp
--- a/Engine/Source/Runtime/Engine/Private/Texture.cpp
+++ b/Engine/Source/Runtime/Engine/Private/Texture.cpp
@@ -279,6 +279,10 @@ void UTexture::Serialize(FArchive& Ar)
 #if WITH_EDITORONLY_DATA
 	if (!StripFlags.IsEditorDataStripped())
 	{
+		if ( Ar.IsSaving() )
+		{
+			check( Source.HasHadBulkDataCleared() == false);
+		}
 		Source.BulkData.Serialize(Ar, this);
 	}
 
diff --git a/Engine/Source/Runtime/Engine/Private/VorbisAudioInfo.cpp b/Engine/Source/Runtime/Engine/Private/VorbisAudioInfo.cpp
--- a/Engine/Source/Runtime/Engine/Private/VorbisAudioInfo.cpp
+++ b/Engine/Source/Runtime/Engine/Private/VorbisAudioInfo.cpp
@@ -4,6 +4,7 @@
 #include "VorbisAudioInfo.h"
 #include "Misc/Paths.h"
 #include "Interfaces/IAudioFormat.h"
+#include "ContentStreaming.h"
 #if PLATFORM_WINDOWS
 #include "WindowsHWrapper.h"
 #endif
@@ -86,6 +87,8 @@ FVorbisAudioInfo::FVorbisAudioInfo( void )
 	, SrcBufferDataSize(0)
 	, BufferOffset(0)
 	, bPerformingOperation(false)
+	, StreamingSoundWave(NULL)
+	, StreamingChunksSize(0)
 { 
 	// Make sure we have properly allocated a VFWrapper
 	check(VFWrapper != NULL);
@@ -103,7 +106,7 @@ FVorbisAudioInfo::~FVorbisAudioInfo( void )
 }
 
 /** Emulate read from memory functionality */
-size_t FVorbisAudioInfo::Read( void *Ptr, uint32 Size )
+size_t FVorbisAudioInfo::ReadMemory( void *Ptr, uint32 Size )
 {
 	check(Ptr);
 	size_t BytesToRead = FMath::Min(Size, SrcBufferDataSize - BufferOffset);
@@ -112,15 +115,15 @@ size_t FVorbisAudioInfo::Read( void *Ptr, uint32 Size )
 	return( BytesToRead );
 }
 
-static size_t OggRead( void *ptr, size_t size, size_t nmemb, void *datasource )
+static size_t OggReadMemory( void *ptr, size_t size, size_t nmemb, void *datasource )
 {
 	check(ptr);
 	check(datasource);
 	FVorbisAudioInfo* OggInfo = (FVorbisAudioInfo*)datasource;
-	return( OggInfo->Read( ptr, size * nmemb ) );
+	return( OggInfo->ReadMemory( ptr, size * nmemb ) );
 }
 
-int FVorbisAudioInfo::Seek( uint32 offset, int whence )
+int FVorbisAudioInfo::SeekMemory( uint32 offset, int whence )
 {
 	switch( whence )
 	{
@@ -144,70 +147,90 @@ int FVorbisAudioInfo::Seek( uint32 offset, int whence )
 	return( BufferOffset );
 }
 
-static int OggSeek( void *datasource, ogg_int64_t offset, int whence )
+static int OggSeekMemory( void *datasource, ogg_int64_t offset, int whence )
 {
 	FVorbisAudioInfo* OggInfo = ( FVorbisAudioInfo* )datasource;
-	return( OggInfo->Seek( offset, whence ) );
+	return( OggInfo->SeekMemory( offset, whence ) );
 }
 
-int FVorbisAudioInfo::Close( void )
+int FVorbisAudioInfo::CloseMemory( void )
 {
 	return( 0 );
 }
 
-static int OggClose( void *datasource )
+static int OggCloseMemory( void *datasource )
 {
 	FVorbisAudioInfo* OggInfo = ( FVorbisAudioInfo* )datasource;
-	return( OggInfo->Close() );
+	return( OggInfo->CloseMemory() );
 }
 
-long FVorbisAudioInfo::Tell( void )
+long FVorbisAudioInfo::TellMemory( void )
 {
 	return( BufferOffset );
 }
 
-static long OggTell( void *datasource )
+static long OggTellMemory( void *datasource )
 {
 	FVorbisAudioInfo *OggInfo = ( FVorbisAudioInfo* )datasource;
-	return( OggInfo->Tell() );
+	return( OggInfo->TellMemory() );
 }
 
-/** 
- * Reads the header information of an ogg vorbis file
- * 
- * @param	Resource		Info about vorbis data
- */
-bool FVorbisAudioInfo::ReadCompressedInfo( const uint8* InSrcBufferData, uint32 InSrcBufferDataSize, FSoundQualityInfo* QualityInfo )
+/** Emulate read from memory functionality */
+size_t FVorbisAudioInfo::ReadStreaming( void *Ptr, uint32 Size )
 {
-	bPerformingOperation = true;
+	size_t	BytesCopied = 0;
 
-	SCOPE_CYCLE_COUNTER( STAT_VorbisPrepareDecompressionTime );
+	while(Size > 0)
+	{
+		uint32	CurChunkSize = 0;
 
-	FScopeLock ScopeLock(&VorbisCriticalSection);
+		uint8 const* ChunkData = IStreamingManager::Get().GetAudioStreamingManager().GetLoadedChunk(StreamingSoundWave, BufferOffset / StreamingChunksSize, &CurChunkSize);
+		
+		check(CurChunkSize >= (BufferOffset % StreamingChunksSize));
+		size_t	BytesToCopy = FMath::Min<uint32>(CurChunkSize - (BufferOffset % StreamingChunksSize), Size);
+		check((BufferOffset % StreamingChunksSize) + BytesToCopy <= CurChunkSize);
 
-	if (!VFWrapper)
-	{
-		bPerformingOperation = false;
-		return false;
+		if(ChunkData == NULL || BytesToCopy == 0)
+		{
+			return BytesCopied;
+		}
+
+		FMemory::Memcpy( Ptr, ChunkData + (BufferOffset % StreamingChunksSize), BytesToCopy );
+		BufferOffset += BytesToCopy;
+		BytesCopied += BytesToCopy;
+		Size -= BytesToCopy;
+		Ptr = (void*)((uint8*)Ptr + BytesToCopy);
 	}
 
-	ov_callbacks Callbacks;
+	return BytesCopied;
+}
 
-	SrcBufferData = InSrcBufferData;
-	SrcBufferDataSize = InSrcBufferDataSize;
-	BufferOffset = 0;
+static size_t OggReadStreaming( void *ptr, size_t size, size_t nmemb, void *datasource )
+{
+	check(ptr);
+	check(datasource);
+	FVorbisAudioInfo* OggInfo = (FVorbisAudioInfo*)datasource;
+	return( OggInfo->ReadStreaming( ptr, size * nmemb ) );
+}
+
+int FVorbisAudioInfo::CloseStreaming( void )
+{
+	return( 0 );
+}
 
-	Callbacks.read_func = OggRead;
-	Callbacks.seek_func = OggSeek;
-	Callbacks.close_func = OggClose;
-	Callbacks.tell_func = OggTell;
+static int OggCloseStreaming( void *datasource )
+{
+	FVorbisAudioInfo* OggInfo = ( FVorbisAudioInfo* )datasource;
+	return( OggInfo->CloseStreaming() );
+}
 
+bool FVorbisAudioInfo::GetCompressedInfoCommon(void* Callbacks, FSoundQualityInfo* QualityInfo)
+{
 	// Set up the read from memory variables
-	int Result = ov_open_callbacks(this, &VFWrapper->vf, NULL, 0, Callbacks);
+	int Result = ov_open_callbacks(this, &VFWrapper->vf, NULL, 0, (*(ov_callbacks*)Callbacks));
 	if (Result < 0)
 	{
 		UE_LOG(LogAudio, Error, TEXT("FVorbisAudioInfo::ReadCompressedInfo, ov_open_callbacks error code: %d"), Result);
-		bPerformingOperation = false;
 		return false;
 	}
 
@@ -231,9 +254,44 @@ bool FVorbisAudioInfo::ReadCompressedInfo( const uint8* InSrcBufferData, uint32
 		}
 	}
 
+	return true;
+}
+
+/** 
+ * Reads the header information of an ogg vorbis file
+ * 
+ * @param	Resource		Info about vorbis data
+ */
+bool FVorbisAudioInfo::ReadCompressedInfo( const uint8* InSrcBufferData, uint32 InSrcBufferDataSize, FSoundQualityInfo* QualityInfo )
+{
+	bPerformingOperation = true;
+
+	SCOPE_CYCLE_COUNTER( STAT_VorbisPrepareDecompressionTime );
+
+	FScopeLock ScopeLock(&VorbisCriticalSection);
+
+	if (!VFWrapper)
+	{
+		bPerformingOperation = false;
+		return false;
+	}
+
+	ov_callbacks Callbacks;
+
+	SrcBufferData = InSrcBufferData;
+	SrcBufferDataSize = InSrcBufferDataSize;
+	BufferOffset = 0;
+
+	Callbacks.read_func = OggReadMemory;
+	Callbacks.seek_func = OggSeekMemory;
+	Callbacks.close_func = OggCloseMemory;
+	Callbacks.tell_func = OggTellMemory;
+
+	bool result = GetCompressedInfoCommon(&Callbacks, QualityInfo);
+
 	bPerformingOperation = false;
 
-	return( true );
+	return( result );
 }
 
 
@@ -372,6 +430,101 @@ void FVorbisAudioInfo::EnableHalfRate( bool HalfRate )
 	bPerformingOperation = false;
 }
 
+bool FVorbisAudioInfo::StreamCompressedInfo(USoundWave* Wave, struct FSoundQualityInfo* QualityInfo)
+{
+	bPerformingOperation = true;
+
+	SCOPE_CYCLE_COUNTER( STAT_VorbisPrepareDecompressionTime );
+
+	FScopeLock ScopeLock(&VorbisCriticalSection);
+
+	if (!VFWrapper)
+	{
+		bPerformingOperation = false;
+		return false;
+	}
+
+	ov_callbacks Callbacks;
+
+	SrcBufferData = NULL;
+	SrcBufferDataSize = 0;
+	BufferOffset = 0;
+	StreamingSoundWave = Wave;
+
+	Callbacks.read_func = OggReadStreaming;
+	Callbacks.close_func = OggCloseStreaming;
+	Callbacks.seek_func = NULL;	// Force streaming
+	Callbacks.tell_func = NULL;	// Force streaming
+
+	// We need to start with a valid StreamingChunksSize so just use this
+	StreamingChunksSize = MONO_PCM_BUFFER_SIZE * 2 * 2;
+
+	bool result = GetCompressedInfoCommon(&Callbacks, QualityInfo);
+
+	// Now we can set the real StreamingChunksSize
+	StreamingChunksSize = MONO_PCM_BUFFER_SIZE * 2 * QualityInfo->NumChannels;
+
+	bPerformingOperation = false;
+
+	return( result );
+}
+
+bool FVorbisAudioInfo::StreamCompressedData(uint8* InDestination, bool bLooping, uint32 BufferSize)
+{
+	check( VFWrapper != NULL );
+	bPerformingOperation = true;
+
+	SCOPE_CYCLE_COUNTER( STAT_VorbisDecompressTime );
+
+	FScopeLock ScopeLock(&VorbisCriticalSection);
+
+	bool	bLooped = false;
+
+	while( BufferSize > 0 )
+	{
+		long BytesActuallyRead = ov_read( &VFWrapper->vf, (char*)InDestination, (int)BufferSize, 0, 2, 1, NULL );
+
+		if( BytesActuallyRead <= 0 )
+		{
+			// We've reached the end
+			bLooped = true;
+
+			BufferOffset = 0;
+
+			// Since we can't tell a streaming file to go back to the start of the stream (there is no seek) we have to close and reopen it which is a bummer
+			ov_clear(&VFWrapper->vf);
+			FMemory::Memzero( &VFWrapper->vf, sizeof( OggVorbis_File ) );
+			ov_callbacks Callbacks;
+			Callbacks.read_func = OggReadStreaming;
+			Callbacks.close_func = OggCloseStreaming;
+			Callbacks.seek_func = NULL;	// Force streaming
+			Callbacks.tell_func = NULL;	// Force streaming
+			int Result = ov_open_callbacks(this, &VFWrapper->vf, NULL, 0, Callbacks);
+			if (Result < 0)
+			{
+				UE_LOG(LogAudio, Error, TEXT("FVorbisAudioInfo::StreamCompressedData, ov_open_callbacks error code: %d"), Result);
+				break;
+			}
+
+			if( !bLooping )
+			{
+				// Need to clear out the remainder of the buffer
+				FMemory::Memzero(InDestination, BufferSize);
+				break;
+			}
+			// else we start over to get the samples from the start of the compressed audio data
+			continue;
+		}
+
+		InDestination += BytesActuallyRead;
+		BufferSize -= BytesActuallyRead;
+	}
+
+	bPerformingOperation = false;
+
+	return( bLooped );
+}
+
 void LoadVorbisLibraries()
 {
 	static bool bIsInitialized = false;
diff --git a/Engine/Source/Runtime/Engine/Public/ShaderCompiler.h b/Engine/Source/Runtime/Engine/Public/ShaderCompiler.h
--- a/Engine/Source/Runtime/Engine/Public/ShaderCompiler.h
+++ b/Engine/Source/Runtime/Engine/Public/ShaderCompiler.h
@@ -25,7 +25,7 @@ DECLARE_LOG_CATEGORY_EXTERN(LogShaderCompilers, Log, All);
 class FShaderCompileJob;
 class FShaderPipelineCompileJob;
 
-#define DEBUG_INFINITESHADERCOMPILE 1
+#define DEBUG_INFINITESHADERCOMPILE 0
 
 
 /** Stores all of the common information used to compile a shader or pipeline. */
diff --git a/Engine/Source/Runtime/Engine/Public/VorbisAudioInfo.h b/Engine/Source/Runtime/Engine/Public/VorbisAudioInfo.h
--- a/Engine/Source/Runtime/Engine/Public/VorbisAudioInfo.h
+++ b/Engine/Source/Runtime/Engine/Public/VorbisAudioInfo.h
@@ -32,6 +32,7 @@ namespace VorbisChannelInfo
 ENGINE_API void LoadVorbisLibraries();
 
 #if WITH_OGGVORBIS
+
 /** 
  * Helper class to parse ogg vorbis data
  */
@@ -42,10 +43,18 @@ public:
 	ENGINE_API virtual ~FVorbisAudioInfo( void );
 
 	/** Emulate read from memory functionality */
-	size_t			Read( void *ptr, uint32 size );
-	int				Seek( uint32 offset, int whence );
-	int				Close( void );
-	long			Tell( void );
+	size_t			ReadMemory( void *ptr, uint32 size );
+	int				SeekMemory( uint32 offset, int whence );
+	int				CloseMemory( void );
+	long			TellMemory( void );
+
+	/** Emulate read from streaming functionality */
+	size_t			ReadStreaming( void *ptr, uint32 size );
+	int				CloseStreaming( void );
+
+	/** Common info and data functions between ReadCompressedInfo/ReadCompressedData and StreamCompressedInfo/StreamCompressedData */
+	/** Can't use ov_callbacks here so we have to pass in a void* for the callbacks instead */
+	bool GetCompressedInfoCommon(void*	Callbacks, FSoundQualityInfo* QualityInfo);
 
 	/** 
 	 * Reads the header information of an ogg vorbis file
@@ -63,7 +72,7 @@ public:
 	 *
 	 * @return	bool		true if the end of the data was reached (for both single shot and looping sounds)
 	 */
-	ENGINE_API virtual bool ReadCompressedData( uint8* Destination, bool bLooping, uint32 BufferSize ) override;
+	ENGINE_API virtual bool ReadCompressedData( uint8* InDestination, bool bLooping, uint32 BufferSize ) override;
 
 	ENGINE_API virtual void SeekToTime( const float SeekTime ) override;
 
@@ -84,6 +93,14 @@ public:
 	virtual bool UsesVorbisChannelOrdering() const override { return true; }
 	virtual int GetStreamBufferSize() const override { return MONO_PCM_BUFFER_SIZE; }
 
+	// Additional overrides for streaming
+	virtual bool SupportsStreaming() const override {return true;}
+	virtual bool StreamCompressedInfo(USoundWave* Wave, struct FSoundQualityInfo* QualityInfo) override;
+	virtual bool StreamCompressedData(uint8* InDestination, bool bLooping, uint32 BufferSize) override;
+	virtual int32 GetCurrentChunkIndex() const override {return StreamingChunksSize == 0 ? 0 : BufferOffset / StreamingChunksSize;}
+	virtual int32 GetCurrentChunkOffset() const override {return StreamingChunksSize == 0 ? 0 : BufferOffset % StreamingChunksSize;}
+	// End of ICompressedAudioInfo Interface
+
 	struct FVorbisFileWrapper* VFWrapper;
 	const uint8*		SrcBufferData;
 	uint32			SrcBufferDataSize;
@@ -93,5 +110,8 @@ public:
 
 	/** Critical section used to prevent multiple threads accessing same ogg-vorbis file handles at the same time */
 	FCriticalSection VorbisCriticalSection;
+
+	USoundWave*		StreamingSoundWave;				// The current sound wave being streamed, this is used to fetch new chunks
+	uint32			StreamingChunksSize;
 };
 #endif
diff --git a/Engine/Source/Runtime/HTML5/HTML5JS/Private/HTML5JavaScriptFx.js b/Engine/Source/Runtime/HTML5/HTML5JS/Private/HTML5JavaScriptFx.js
--- a/Engine/Source/Runtime/HTML5/HTML5JS/Private/HTML5JavaScriptFx.js
+++ b/Engine/Source/Runtime/HTML5/HTML5JS/Private/HTML5JavaScriptFx.js
@@ -367,6 +367,10 @@ var UE_JavaScriptLibary = {
       return Runtime.dynCall('i', resY, []);
     };
   },
+
+  UE_PixelRatio: function() {
+  	  return window.devicePixelRatio;
+  },
  
 };
 
diff --git a/Engine/Source/Runtime/HTML5/HTML5JS/Public/HTML5JavaScriptFx.h b/Engine/Source/Runtime/HTML5/HTML5JS/Public/HTML5JavaScriptFx.h
--- a/Engine/Source/Runtime/HTML5/HTML5JS/Public/HTML5JavaScriptFx.h
+++ b/Engine/Source/Runtime/HTML5/HTML5JS/Public/HTML5JavaScriptFx.h
@@ -34,5 +34,6 @@
 
 		// GSystemResolution
 		void UE_GSystemResolution( int(*resX)(), int(*rexY)() );
+		int UE_PixelRatio();
 	}
 
diff --git a/Engine/Source/Runtime/Launch/Private/Android/AndroidEventManager.cpp b/Engine/Source/Runtime/Launch/Private/Android/AndroidEventManager.cpp
--- a/Engine/Source/Runtime/Launch/Private/Android/AndroidEventManager.cpp
+++ b/Engine/Source/Runtime/Launch/Private/Android/AndroidEventManager.cpp
@@ -484,6 +484,13 @@ bool FAppEventManager::IsGamePaused()
 	return !bRunning;
 }
 
+
+bool FAppEventManager::IsGameInFocus()
+{
+	return (bWindowInFocus && bHaveWindow);
+}
+
+
 bool FAppEventManager::WaitForEventInQueue(EAppEventState InState, double TimeoutSeconds)
 {
 	bool FoundEvent = false;
diff --git a/Engine/Source/Runtime/Launch/Private/Android/AndroidEventManager.h b/Engine/Source/Runtime/Launch/Private/Android/AndroidEventManager.h
deleted file mode 100644
diff --git a/Engine/Source/Runtime/Launch/Public/Android/AndroidEventManager.h b/Engine/Source/Runtime/Launch/Public/Android/AndroidEventManager.h
new file mode 100644
--- /dev/null
+++ b/Engine/Source/Runtime/Launch/Public/Android/AndroidEventManager.h
@@ -0,0 +1,95 @@
+// Copyright 1998-2017 Epic Games, Inc. All Rights Reserved.
+
+#pragma once
+
+#include "CoreMinimal.h"
+#include "Containers/Queue.h"
+
+struct ANativeWindow;
+
+DECLARE_LOG_CATEGORY_EXTERN(LogAndroidEvents, Log, All);
+
+enum EAppEventState
+{
+	APP_EVENT_STATE_WINDOW_CREATED,
+	APP_EVENT_STATE_WINDOW_RESIZED,
+	APP_EVENT_STATE_WINDOW_CHANGED,
+	APP_EVENT_STATE_WINDOW_DESTROYED,
+	APP_EVENT_STATE_WINDOW_REDRAW_NEEDED,
+	APP_EVENT_STATE_ON_DESTROY,
+	APP_EVENT_STATE_ON_PAUSE,
+	APP_EVENT_STATE_ON_RESUME,
+	APP_EVENT_STATE_ON_STOP,
+	APP_EVENT_STATE_ON_START,
+	APP_EVENT_STATE_WINDOW_LOST_FOCUS,
+	APP_EVENT_STATE_WINDOW_GAINED_FOCUS,
+	APP_EVENT_STATE_SAVE_STATE,
+	APP_EVENT_STATE_INVALID = -1,
+};
+
+struct FAppEventData
+{
+	EAppEventState State;
+	void* Data;
+
+	FAppEventData():
+		State(APP_EVENT_STATE_INVALID)
+		,Data(NULL)
+	{	}
+};
+
+class FAppEventManager
+{
+public:
+	static FAppEventManager* GetInstance();
+	
+	void Tick();
+	void EnqueueAppEvent(EAppEventState InState, void* InData = NULL);
+	void SetEventHandlerEvent(FEvent* InEventHandlerEvent);
+	void HandleWindowCreated(void* InWindow);
+	void HandleWindowClosed();
+	bool IsGamePaused();
+	bool IsGameInFocus();
+	bool WaitForEventInQueue(EAppEventState InState, double TimeoutSeconds);
+
+	void SetEmptyQueueHandlerEvent(FEvent* InEventHandlerEvent);
+	void WaitForEmptyQueue();
+	void TriggerEmptyQueue();
+
+protected:
+	FAppEventManager();
+
+private:
+	FAppEventData DequeueAppEvent();
+	void PauseRendering();
+	void ResumeRendering();
+	void PauseAudio();
+	void ResumeAudio();
+
+	void ExecWindowCreated();
+	void ExecWindowResized();
+	void ExecDestroyWindow();
+	
+	static FAppEventManager* sInstance;
+
+	pthread_mutex_t QueueMutex;			//@todo android: can probably get rid of this now that we're using an mpsc queue
+	TQueue<FAppEventData, EQueueMode::Mpsc> Queue;
+	
+	pthread_mutex_t MainMutex;
+
+	FEvent* EventHandlerEvent;			// triggered every time the event handler thread fires off an event
+	FEvent* EmptyQueueHandlerEvent;		// triggered every time the queue is emptied
+
+	//states
+	bool FirstInitialized;
+	bool bCreateWindow;
+	ANativeWindow* PendingWindow;
+
+	bool bWindowInFocus;
+	bool bSaveState;
+	bool bAudioPaused;
+
+	bool bHaveWindow;
+	bool bHaveGame;
+	bool bRunning;
+};
diff --git a/Engine/Source/Runtime/MoviePlayer/Public/MoviePlayer.h b/Engine/Source/Runtime/MoviePlayer/Public/MoviePlayer.h
--- a/Engine/Source/Runtime/MoviePlayer/Public/MoviePlayer.h
+++ b/Engine/Source/Runtime/MoviePlayer/Public/MoviePlayer.h
@@ -11,7 +11,6 @@
 
 // In order for a platform to support early movie playback, the platform must support the rendering thread 
 // starting very early and support rendering as soon as it is started and the module containing the movie streamer for the platform must already be loaded
-#define PLATFORM_SUPPORTS_EARLY_MOVIE_PLAYBACK 0
 
 UENUM()
 enum EMoviePlaybackType
diff --git a/Engine/Source/Runtime/OpenGLDrv/Private/HTML5/HTML5OpenGL.cpp b/Engine/Source/Runtime/OpenGLDrv/Private/HTML5/HTML5OpenGL.cpp
--- a/Engine/Source/Runtime/OpenGLDrv/Private/HTML5/HTML5OpenGL.cpp
+++ b/Engine/Source/Runtime/OpenGLDrv/Private/HTML5/HTML5OpenGL.cpp
@@ -196,9 +196,6 @@ struct FPlatformOpenGLDevice
 		WindowHandle = SDL_CreateWindow("HTML5", SDL_WINDOWPOS_CENTERED, SDL_WINDOWPOS_CENTERED,
 			800, 600, SDL_WINDOW_OPENGL | SDL_WINDOW_SHOWN| SDL_WINDOW_RESIZABLE);
 #if PLATFORM_HTML5_BROWSER
-//		EM_ASM(
-//			console.log("SDL_CreateWindow() 800x600");
-//		);
 		UE_GSystemResolution( GSystemResolution_ResX, GSystemResolution_ResY );
 #endif
 		PlatformCreateOpenGLContext(this,WindowHandle);
@@ -374,6 +371,9 @@ void PlatformReleaseRenderQuery( GLuint Query, uint64 QueryContext )
 
 void PlatformRestoreDesktopDisplayMode()
 {
+#if PLATFORM_HTML5_BROWSER
+	EM_ASM( Module['canvas'].UE_canvas.bIsFullScreen = 0; );
+#endif
 }
 
 #if PLATFORM_HTML5_BROWSER
@@ -392,6 +392,7 @@ extern "C"
 #else
 		emscripten_request_fullscreen("canvas", true);
 #endif
+		EM_ASM( Module['canvas'].UE_canvas.bIsFullScreen = 1; );
 		return 0;
 	}
 }
diff --git a/Engine/Source/Runtime/OpenGLDrv/Private/OpenGLTexture.cpp b/Engine/Source/Runtime/OpenGLDrv/Private/OpenGLTexture.cpp
--- a/Engine/Source/Runtime/OpenGLDrv/Private/OpenGLTexture.cpp
+++ b/Engine/Source/Runtime/OpenGLDrv/Private/OpenGLTexture.cpp
@@ -1924,7 +1924,7 @@ FShaderResourceViewRHIRef FOpenGLDynamicRHI::RHICreateShaderResourceView(FTextur
 		const FOpenGLTextureFormat& GLFormat = GOpenGLTextureFormats[TextureCube->GetFormat()];
 		const bool bSRGB = (TextureCube->GetFlags()&TexCreate_SRGB) != 0;
 		
-		FOpenGL::TextureView( Resource, TextureCube->Target, TextureCube->Resource, GLFormat.InternalFormat[bSRGB], MipLevel, 1, 0, 1);
+		FOpenGL::TextureView( Resource, TextureCube->Target, TextureCube->Resource, GLFormat.InternalFormat[bSRGB], MipLevel, 1, 0, 6);
 		
 		View = new FOpenGLShaderResourceView(this, Resource, TextureCube->Target, MipLevel, true);
 	}
diff --git a/Engine/Source/Runtime/PakFile/Private/SignedArchiveReader.cpp b/Engine/Source/Runtime/PakFile/Private/SignedArchiveReader.cpp
--- a/Engine/Source/Runtime/PakFile/Private/SignedArchiveReader.cpp
+++ b/Engine/Source/Runtime/PakFile/Private/SignedArchiveReader.cpp
@@ -24,7 +24,7 @@ FChunkCacheWorker::FChunkCacheWorker(FArchive* InReader, const TCHAR* Filename)
 
 	if (SigFileReader == nullptr)
 	{
-		UE_LOG(LogPakFile, Fatal, TEXT("Couldn't find pak signature file '%s'"), *Filename);
+		UE_LOG(LogPakFile, Fatal, TEXT("Couldn't find pak signature file '%s'"), *SigFileFilename);
 	}
 
 	FEncryptedSignature MasterSignature;
diff --git a/Engine/Source/Runtime/Projects/Private/ProjectDescriptor.cpp b/Engine/Source/Runtime/Projects/Private/ProjectDescriptor.cpp
--- a/Engine/Source/Runtime/Projects/Private/ProjectDescriptor.cpp
+++ b/Engine/Source/Runtime/Projects/Private/ProjectDescriptor.cpp
@@ -137,6 +137,17 @@ bool FProjectDescriptor::Read(const FJsonObject& Object, const FString& PathToPr
 				}
 			}
 		}
+		// If this is a packaged build and there are additional directories, they need to be remapped to the packaged location
+		if (FPlatformProperties::RequiresCookedData() && AdditionalPluginDirectories.Num() > 0)
+		{
+			AdditionalPluginDirectories.Empty();
+			FString RemappedDir = IFileManager::Get().ConvertToAbsolutePathForExternalAppForWrite(*(PathToProject / TEXT("../RemappedPlugins/")));
+			if (!IsRootedPath(RemappedDir))
+			{
+				RemappedDir = FPaths::ConvertRelativePathToFull(RemappedDir);
+			}
+			AddPluginDirectory(RemappedDir);
+		}
 	}
 
 	// Read the target platforms
