Redoing a change from WEX in main to fix rem,ote ios compiling issue

diff --git a/Engine/Source/Programs/UnrealBuildTool/System/DependencyCache.cs b/Engine/Source/Programs/UnrealBuildTool/System/DependencyCache.cs
--- a/Engine/Source/Programs/UnrealBuildTool/System/DependencyCache.cs
+++ b/Engine/Source/Programs/UnrealBuildTool/System/DependencyCache.cs
@@ -33,6 +33,14 @@ namespace UnrealBuildTool
 		/// <param name="InIncludeName"></param>
 		public DependencyInclude(string InIncludeName)
 		{
+            if (BuildHostPlatform.Current.Platform != UnrealTargetPlatform.Mac)
+            {
+                if (InIncludeName.StartsWith("/Users") || InIncludeName.StartsWith("\\Users"))
+                {
+                    // convert the path back to Windows
+                    InIncludeName = RemoteExports.UnconvertPath(InIncludeName);
+                }
+            }
 			IncludeName = InIncludeName;
 		}
 	}
