Fixed crash that occurs when ProjectGenerator is null (special case for stub projects in installed builds)

diff --git a/Engine/Source/Programs/UnrealBuildTool/System/VCProject.cs b/Engine/Source/Programs/UnrealBuildTool/System/VCProject.cs
--- a/Engine/Source/Programs/UnrealBuildTool/System/VCProject.cs
+++ b/Engine/Source/Programs/UnrealBuildTool/System/VCProject.cs
@@ -1110,10 +1110,11 @@ namespace UnrealBuildTool
 			string ConditionString = "Condition=\"'$(Configuration)|$(Platform)'=='" + Combination.ProjectConfigurationAndPlatformName + "'\"";
 
 			{
+				string ImportGroupProperties = (ProjGenerator != null) ? ProjGenerator.GetVisualStudioImportGroupProperties(Platform) : "";
 				VCProjectFileContent.Append(
 					"	<ImportGroup " + ConditionString + " Label=\"PropertySheets\">" + ProjectFileGenerator.NewLine +
 					"		<Import Project=\"$(UserRootDir)\\Microsoft.Cpp.$(Platform).user.props\" Condition=\"exists('$(UserRootDir)\\Microsoft.Cpp.$(Platform).user.props')\" Label=\"LocalAppDataPlatform\" />" + ProjectFileGenerator.NewLine +
-							ProjGenerator.GetVisualStudioImportGroupProperties(Platform) +
+							ImportGroupProperties +
 					"	</ImportGroup>" + ProjectFileGenerator.NewLine);
 
 				DirectoryReference ProjectDirectory = ProjectFilePath.Directory;
