UBT: Fix output directory for targets built with an installed engine build. Regression from Dev-Core integration.

diff --git a/Engine/Source/Programs/UnrealBuildTool/Configuration/UEBuildTarget.cs b/Engine/Source/Programs/UnrealBuildTool/Configuration/UEBuildTarget.cs
--- a/Engine/Source/Programs/UnrealBuildTool/Configuration/UEBuildTarget.cs
+++ b/Engine/Source/Programs/UnrealBuildTool/Configuration/UEBuildTarget.cs
@@ -1418,7 +1418,7 @@ namespace UnrealBuildTool
 
 			// Construct the output paths for this target's executable
 			DirectoryReference OutputDirectory;
-			if (!bUseSharedBuildEnvironment && !Rules.bOutputToEngineBinaries)
+			if ((bCompileMonolithic || TargetType == TargetType.Program || !bUseSharedBuildEnvironment) && !Rules.bOutputToEngineBinaries)
 			{
 				OutputDirectory = ProjectDirectory;
 			}
