Missing file from change to fixup ListThirdPartySoftware commandlet.

diff --git a/Engine/Source/Programs/UnrealBuildTool/Configuration/UEBuildModule.cs b/Engine/Source/Programs/UnrealBuildTool/Configuration/UEBuildModule.cs
--- a/Engine/Source/Programs/UnrealBuildTool/Configuration/UEBuildModule.cs
+++ b/Engine/Source/Programs/UnrealBuildTool/Configuration/UEBuildModule.cs
@@ -778,6 +778,16 @@ namespace UnrealBuildTool
 				Writer.WriteValue(ModuleName);
 			}
 			Writer.WriteArrayEnd();
+
+			Writer.WriteArrayStart("RuntimeDependencies");
+			foreach(RuntimeDependency RuntimeDependency in Rules.RuntimeDependencies)
+			{
+				Writer.WriteObjectStart();
+				Writer.WriteValue("Path", RuntimeDependency.Path);
+				Writer.WriteValue("Type", RuntimeDependency.Type.ToString());
+				Writer.WriteObjectEnd();
+			}
+			Writer.WriteArrayEnd();
 		}
 
 		/// <summary>
