Circular reference changes

diff --git a/Engine/Source/Editor/UnrealEd/UnrealEd.Build.cs b/Engine/Source/Editor/UnrealEd/UnrealEd.Build.cs
--- a/Engine/Source/Editor/UnrealEd/UnrealEd.Build.cs
+++ b/Engine/Source/Editor/UnrealEd/UnrealEd.Build.cs
@@ -151,6 +151,7 @@ public class UnrealEd : ModuleRules
 				"MovieScene",
 				"MovieSceneTracks",
                 "ViewportInteraction",
+                "VREditor",
                 "ClothingSystemEditor",
                 "ClothingSystemRuntime"
             }
@@ -238,7 +239,8 @@ public class UnrealEd : ModuleRules
                 "GraphEditor",
 				"Kismet",
 				"AudioEditor",
-                "ViewportInteraction"
+                "ViewportInteraction",
+                "VREditor"
             }
 		); 
 
diff --git a/Engine/Source/Programs/UnrealBuildTool/Configuration/UEBuildModuleCPP.cs b/Engine/Source/Programs/UnrealBuildTool/Configuration/UEBuildModuleCPP.cs
--- a/Engine/Source/Programs/UnrealBuildTool/Configuration/UEBuildModuleCPP.cs
+++ b/Engine/Source/Programs/UnrealBuildTool/Configuration/UEBuildModuleCPP.cs
@@ -262,7 +262,16 @@ namespace UnrealBuildTool
 			new KeyValuePair<string, string>("HierarchicalLODOutliner", "UnrealEd"),
 			new KeyValuePair<string, string>("PixelInspectorModule", "UnrealEd"),
 			new KeyValuePair<string, string>("GameplayAbilitiesEditor", "BlueprintGraph"),
-		};
+            new KeyValuePair<string, string>("UnrealEd", "ViewportInteraction"),
+            new KeyValuePair<string, string>("UnrealEd", "VREditor"),
+            new KeyValuePair<string, string>("LandscapeEditor", "ViewportInteraction"),
+            new KeyValuePair<string, string>("LandscapeEditor", "VREditor"),
+            new KeyValuePair<string, string>("FoliageEdit", "ViewportInteraction"),
+            new KeyValuePair<string, string>("FoliageEdit", "VREditor"),
+            new KeyValuePair<string, string>("MeshPaint", "ViewportInteraction"),
+            new KeyValuePair<string, string>("MeshPaint", "VREditor"),
+            new KeyValuePair<string, string>("Sequencer", "ViewportInteraction"),
+        };
 
 
 		public UEBuildModuleCPP(
