Fix Java 1.5 obsolete warnings
#ue4
#android

diff --git a/Engine/Source/Programs/UnrealBuildTool/Android/UEDeployAndroid.cs b/Engine/Source/Programs/UnrealBuildTool/Android/UEDeployAndroid.cs
--- a/Engine/Source/Programs/UnrealBuildTool/Android/UEDeployAndroid.cs
+++ b/Engine/Source/Programs/UnrealBuildTool/Android/UEDeployAndroid.cs
@@ -2359,6 +2359,7 @@ namespace UnrealBuildTool
 				}
 
 				// Use ant to build the .apk file
+				string AntOptions = AntBuildType + " -Djava.source=1.7 -Djava.target=1.7";
 				string ShellExecutable = Utils.IsRunningOnMono ? "/bin/sh" : "cmd.exe";
 				string ShellParametersBegin = Utils.IsRunningOnMono ? "-c '" : "/c ";
 				string ShellParametersEnd = Utils.IsRunningOnMono ? "'" : "";
@@ -2366,18 +2367,18 @@ namespace UnrealBuildTool
 				{
 					default:
 					case "quiet":
-						if (RunCommandLineProgramAndReturnResult(UE4BuildPath, ShellExecutable, ShellParametersBegin + "\"" + GetAntPath() + "\" -quiet " + AntBuildType + ShellParametersEnd, "Making .apk with Ant... (note: it's safe to ignore javac obsolete warnings)") != 0)
+						if (RunCommandLineProgramAndReturnResult(UE4BuildPath, ShellExecutable, ShellParametersBegin + "\"" + GetAntPath() + "\" -quiet " + AntOptions + ShellParametersEnd, "Making .apk with Ant... (note: it's safe to ignore javac obsolete warnings)") != 0)
 						{
-							RunCommandLineProgramAndReturnResult(UE4BuildPath, ShellExecutable, ShellParametersBegin + "\"" + GetAntPath() + "\" " + AntBuildType + ShellParametersEnd, "Making .apk with Ant again to show errors");
+							RunCommandLineProgramAndReturnResult(UE4BuildPath, ShellExecutable, ShellParametersBegin + "\"" + GetAntPath() + "\" " + AntOptions + ShellParametersEnd, "Making .apk with Ant again to show errors");
 						}
 						break;
 
 					case "normal":
-						RunCommandLineProgramAndReturnResult(UE4BuildPath, ShellExecutable, ShellParametersBegin + "\"" + GetAntPath() + "\" " + AntBuildType + ShellParametersEnd, "Making .apk with Ant again to show errors");
+						RunCommandLineProgramAndReturnResult(UE4BuildPath, ShellExecutable, ShellParametersBegin + "\"" + GetAntPath() + "\" " + AntOptions + ShellParametersEnd, "Making .apk with Ant again to show errors");
 						break;
 
 					case "verbose":
-						RunCommandLineProgramAndReturnResult(UE4BuildPath, ShellExecutable, ShellParametersBegin + "\"" + GetAntPath() + "\" -verbose " + AntBuildType + ShellParametersEnd, "Making .apk with Ant again to show errors");
+						RunCommandLineProgramAndReturnResult(UE4BuildPath, ShellExecutable, ShellParametersBegin + "\"" + GetAntPath() + "\" -verbose " + AntOptions + ShellParametersEnd, "Making .apk with Ant again to show errors");
 						break;
 				}
 
