Copying //UE4/Dev-Core to //UE4/Dev-Main (Source: //UE4/Dev-Core @ 3345728)

diff --git a/Engine/Source/Programs/CrashReporter/CrashReportWebSite/Properties/Settings.settings b/Engine/Source/Programs/CrashReporter/CrashReportWebSite/Properties/Settings.settings
--- a/Engine/Source/Programs/CrashReporter/CrashReportWebSite/Properties/Settings.settings
+++ b/Engine/Source/Programs/CrashReporter/CrashReportWebSite/Properties/Settings.settings
@@ -29,5 +29,179 @@
     <Setting Name="SlackEmoji" Type="System.String" Scope="Application">
       <Value Profile="(Default)">:rngesus:</Value>
     </Setting>
+    <Setting Name="AWSS3ProfileName" Type="System.String" Scope="Application">
+      <Value Profile="(Default)">crashreporter-s3</Value>
+    </Setting>
+    <Setting Name="AWSCredentialsFilepath" Type="System.String" Scope="Application">
+      <Value Profile="(Default)">C:\ProgramData\awscredentials.ini</Value>
+    </Setting>
+    <Setting Name="AWSS3URL" Type="System.String" Scope="Application">
+      <Value Profile="(Default)">https://s3.amazonaws.com</Value>
+    </Setting>
+    <Setting Name="AWSS3KeyPrefix" Type="System.String" Scope="Application">
+      <Value Profile="(Default)">crashreporterfiles-prod</Value>
+    </Setting>
+    <Setting Name="DownloadFromS3" Type="System.Boolean" Scope="Application">
+      <Value Profile="(Default)">True</Value>
+    </Setting>
+    <Setting Name="TestXML" Type="System.String" Scope="Application">
+      <Value Profile="(Default)">&lt;CrashDescription&gt;
+	&lt;CrashType /&gt;
+	&lt;BranchName&gt;//UE4/Release-4.12&lt;/BranchName&gt;
+	&lt;GameName&gt;UE4-Lenar_3&lt;/GameName&gt;
+	&lt;Platform&gt;Win64 [Windows 7 Service Pack 1 64b]&lt;/Platform&gt;
+	&lt;EngineMode&gt;Editor&lt;/EngineMode&gt;
+	&lt;BuildVersion&gt;4.12.5&lt;/BuildVersion&gt;
+	&lt;BuiltFromCL&gt;3039270&lt;/BuiltFromCL&gt;
+	&lt;CommandLine /&gt;
+	&lt;BaseDir&gt;C:/Program Files (x86)/Epic Games/4.12/Engine/Binaries/Win64/&lt;/BaseDir&gt;
+	&lt;Language&gt;zh_CN&lt;/Language&gt;
+	&lt;SystemLanguage /&gt;
+	&lt;UserName&gt;ohmimi&lt;/UserName&gt;
+	&lt;MachineGuid&gt;A24BAFFA41A33062DF1BCB9A31DDE162&lt;/MachineGuid&gt;
+	&lt;EpicAccountId&gt;cc732bdf5c4443f4a966a15bd4869c1e&lt;/EpicAccountId&gt;
+	&lt;CallStack&gt;UE4Editor_Core!FDebug::OptionallyLogFormattedEnsureMessageReturningFalse() [d:\build\++ue4+release-4.12+compile\sync\engine\source\runtime\core\private\misc\outputdevice.cpp:451]&lt;/CallStack&gt;
+	&lt;CallStack&gt;UE4Editor_Engine!UBlueprint::SetObjectBeingDebugged() [d:\build\++ue4+release-4.12+compile\sync\engine\source\runtime\engine\private\blueprint.cpp:685]&lt;/CallStack&gt;
+	&lt;CallStack&gt;UE4Editor_Engine!UWorld::TransferBlueprintDebugReferences() [d:\build\++ue4+release-4.12+compile\sync\engine\source\runtime\engine\private\world.cpp:1618]&lt;/CallStack&gt;
+	&lt;CallStack&gt;UE4Editor_UnrealEd!UEditorEngine::TeardownPlaySession() [d:\build\++ue4+release-4.12+compile\sync\engine\source\editor\unrealed\private\playlevel.cpp:642]&lt;/CallStack&gt;
+	&lt;CallStack&gt;UE4Editor_UnrealEd!UEditorEngine::EndPlayMap() [d:\build\++ue4+release-4.12+compile\sync\engine\source\editor\unrealed\private\playlevel.cpp:255]&lt;/CallStack&gt;
+	&lt;CallStack&gt;UE4Editor_LevelEditor!FLevelViewportLayout::~FLevelViewportLayout() [d:\build\++ue4+release-4.12+compile\sync\engine\source\editor\leveleditor\private\levelviewportlayout.cpp:145]&lt;/CallStack&gt;
+	&lt;CallStack&gt;UE4Editor_LevelEditor!FLevelViewportLayout2x2::`scalar deleting destructor'()&lt;/CallStack&gt;
+	&lt;CallStack&gt;UE4Editor_LevelEditor!FLevelViewportTabContent::~FLevelViewportTabContent()&lt;/CallStack&gt;
+	&lt;CallStack&gt;UE4Editor_LevelEditor!SharedPointerInternals::TReferenceControllerWithDeleter&amp;lt;FLevelViewportTabContent,SharedPointerInternals::DefaultDeleter&amp;lt;FLevelViewportTabContent&amp;gt; &amp;gt;::DestroyObject() [d:\build\++ue4+release-4.12+compile\sync\engine\source\runtime\core\public\templates\sharedpointerinternals.h:104]&lt;/CallStack&gt;
+	&lt;CallStack&gt;UE4Editor_LevelEditor!SViewportsOverlay::~SViewportsOverlay()&lt;/CallStack&gt;
+	&lt;CallStack&gt;UE4Editor_LevelEditor!SViewportsOverlay::`scalar deleting destructor'()&lt;/CallStack&gt;
+	&lt;CallStack&gt;UE4Editor_SlateCore!FSimpleSlot::~FSimpleSlot()&lt;/CallStack&gt;
+	&lt;CallStack&gt;UE4Editor_SlateCore!SCompoundWidget::~SCompoundWidget()&lt;/CallStack&gt;
+	&lt;CallStack&gt;UE4Editor_Slate!SBorder::`vector deleting destructor'()&lt;/CallStack&gt;
+	&lt;CallStack&gt;UE4Editor_SlateCore!TIndirectArray&amp;lt;SOverlay::FOverlaySlot,FDefaultAllocator&amp;gt;::DestructAndFreeItems() [d:\build\++ue4+release-4.12+compile\sync\engine\source\runtime\core\public\containers\array.h:3410]&lt;/CallStack&gt;
+	&lt;CallStack&gt;UE4Editor_SlateCore!SOverlay::~SOverlay()&lt;/CallStack&gt;
+	&lt;CallStack&gt;UE4Editor_Slate!SOverlay::`scalar deleting destructor'()&lt;/CallStack&gt;
+	&lt;CallStack&gt;UE4Editor_SlateCore!TIndirectArray&amp;lt;SBoxPanel::FSlot,FDefaultAllocator&amp;gt;::DestructAndFreeItems() [d:\build\++ue4+release-4.12+compile\sync\engine\source\runtime\core\public\containers\array.h:3410]&lt;/CallStack&gt;
+	&lt;CallStack&gt;UE4Editor_SlateCore!SBoxPanel::~SBoxPanel()&lt;/CallStack&gt;
+	&lt;CallStack&gt;UE4Editor_Slate!SVerticalBox::`scalar deleting destructor'()&lt;/CallStack&gt;
+	&lt;CallStack&gt;UE4Editor_SlateCore!FSimpleSlot::~FSimpleSlot()&lt;/CallStack&gt;
+	&lt;CallStack&gt;UE4Editor_SlateCore!SCompoundWidget::~SCompoundWidget()&lt;/CallStack&gt;
+	&lt;CallStack&gt;UE4Editor_Slate!SDockingTabStack::`vector deleting destructor'()&lt;/CallStack&gt;
+	&lt;CallStack&gt;UE4Editor_SlateCore!FSlotBase::~FSlotBase() [d:\build\++ue4+release-4.12+compile\sync\engine\source\runtime\slatecore\private\slotbase.cpp:34]&lt;/CallStack&gt;
+	&lt;CallStack&gt;UE4Editor_Slate!TIndirectArray&amp;lt;SSplitter::FSlot,FDefaultAllocator&amp;gt;::DestructAndFreeItems() [d:\build\++ue4+release-4.12+compile\sync\engine\source\runtime\core\public\containers\array.h:3410]&lt;/CallStack&gt;
+	&lt;CallStack&gt;UE4Editor_Slate!SSplitter::~SSplitter()&lt;/CallStack&gt;
+	&lt;CallStack&gt;UE4Editor_Slate!SSplitter::`vector deleting destructor'()&lt;/CallStack&gt;
+	&lt;CallStack&gt;UE4Editor_SlateCore!FSimpleSlot::~FSimpleSlot()&lt;/CallStack&gt;
+	&lt;CallStack&gt;UE4Editor_SlateCore!SCompoundWidget::~SCompoundWidget()&lt;/CallStack&gt;
+	&lt;CallStack&gt;UE4Editor_Slate!SDockingSplitter::`vector deleting destructor'()&lt;/CallStack&gt;
+	&lt;CallStack&gt;UE4Editor_SlateCore!FSlotBase::~FSlotBase() [d:\build\++ue4+release-4.12+compile\sync\engine\source\runtime\slatecore\private\slotbase.cpp:34]&lt;/CallStack&gt;
+	&lt;CallStack&gt;UE4Editor_Slate!TIndirectArray&amp;lt;SSplitter::FSlot,FDefaultAllocator&amp;gt;::DestructAndFreeItems() [d:\build\++ue4+release-4.12+compile\sync\engine\source\runtime\core\public\containers\array.h:3410]&lt;/CallStack&gt;
+	&lt;CallStack&gt;UE4Editor_Slate!SSplitter::~SSplitter()&lt;/CallStack&gt;
+	&lt;CallStack&gt;UE4Editor_Slate!SSplitter::`vector deleting destructor'()&lt;/CallStack&gt;
+	&lt;CallStack&gt;UE4Editor_SlateCore!FSimpleSlot::~FSimpleSlot()&lt;/CallStack&gt;
+	&lt;CallStack&gt;UE4Editor_SlateCore!SCompoundWidget::~SCompoundWidget()&lt;/CallStack&gt;
+	&lt;CallStack&gt;UE4Editor_Slate!SDockingSplitter::`vector deleting destructor'()&lt;/CallStack&gt;
+	&lt;CallStack&gt;UE4Editor_SlateCore!FSlotBase::~FSlotBase() [d:\build\++ue4+release-4.12+compile\sync\engine\source\runtime\slatecore\private\slotbase.cpp:34]&lt;/CallStack&gt;
+	&lt;CallStack&gt;UE4Editor_Slate!TIndirectArray&amp;lt;SSplitter::FSlot,FDefaultAllocator&amp;gt;::DestructAndFreeItems() [d:\build\++ue4+release-4.12+compile\sync\engine\source\runtime\core\public\containers\array.h:3410]&lt;/CallStack&gt;
+	&lt;CallStack&gt;UE4Editor_Slate!SSplitter::~SSplitter()&lt;/CallStack&gt;
+	&lt;CallStack&gt;UE4Editor_Slate!SSplitter::`vector deleting destructor'()&lt;/CallStack&gt;
+	&lt;CallStack&gt;UE4Editor_SlateCore!TIndirectArray&amp;lt;SOverlay::FOverlaySlot,FDefaultAllocator&amp;gt;::DestructAndFreeItems() [d:\build\++ue4+release-4.12+compile\sync\engine\source\runtime\core\public\containers\array.h:3410]&lt;/CallStack&gt;
+	&lt;CallStack&gt;UE4Editor_SlateCore!SOverlay::~SOverlay()&lt;/CallStack&gt;
+	&lt;CallStack&gt;UE4Editor_Slate!SOverlay::`scalar deleting destructor'()&lt;/CallStack&gt;
+	&lt;CallStack&gt;UE4Editor_SlateCore!FSimpleSlot::~FSimpleSlot()&lt;/CallStack&gt;
+	&lt;CallStack&gt;UE4Editor_SlateCore!SCompoundWidget::~SCompoundWidget()&lt;/CallStack&gt;
+	&lt;CallStack&gt;UE4Editor_Slate!SDockingArea::`vector deleting destructor'()&lt;/CallStack&gt;
+	&lt;CallStack&gt;UE4Editor_SlateCore!TIndirectArray&amp;lt;SBoxPanel::FSlot,FDefaultAllocator&amp;gt;::DestructAndFreeItems() [d:\build\++ue4+release-4.12+compile\sync\engine\source\runtime\core\public\containers\array.h:3410]&lt;/CallStack&gt;
+	&lt;CallStack&gt;UE4Editor_SlateCore!SBoxPanel::~SBoxPanel()&lt;/CallStack&gt;
+	&lt;CallStack&gt;UE4Editor_LevelEditor!SVerticalBox::`scalar deleting destructor'()&lt;/CallStack&gt;
+	&lt;CallStack&gt;UE4Editor_SlateCore!FSimpleSlot::~FSimpleSlot()&lt;/CallStack&gt;
+	&lt;CallStack&gt;UE4Editor_SlateCore!SCompoundWidget::~SCompoundWidget()&lt;/CallStack&gt;
+	&lt;CallStack&gt;UE4Editor_LevelEditor!SLevelEditor::`scalar deleting destructor'()&lt;/CallStack&gt;
+	&lt;CallStack&gt;UE4Editor_SlateCore!FSimpleSlot::~FSimpleSlot()&lt;/CallStack&gt;
+	&lt;CallStack&gt;UE4Editor_SlateCore!SCompoundWidget::~SCompoundWidget()&lt;/CallStack&gt;
+	&lt;CallStack&gt;UE4Editor_Slate!SBorder::`vector deleting destructor'()&lt;/CallStack&gt;
+	&lt;CallStack&gt;UE4Editor_SlateCore!TIndirectArray&amp;lt;SOverlay::FOverlaySlot,FDefaultAllocator&amp;gt;::DestructAndFreeItems() [d:\build\++ue4+release-4.12+compile\sync\engine\source\runtime\core\public\containers\array.h:3410]&lt;/CallStack&gt;
+	&lt;CallStack&gt;UE4Editor_SlateCore!SOverlay::~SOverlay()&lt;/CallStack&gt;
+	&lt;CallStack&gt;UE4Editor_Slate!SOverlay::`scalar deleting destructor'()&lt;/CallStack&gt;
+	&lt;CallStack&gt;UE4Editor_SlateCore!TIndirectArray&amp;lt;SBoxPanel::FSlot,FDefaultAllocator&amp;gt;::DestructAndFreeItems() [d:\build\++ue4+release-4.12+compile\sync\engine\source\runtime\core\public\containers\array.h:3410]&lt;/CallStack&gt;
+	&lt;CallStack&gt;UE4Editor_SlateCore!SBoxPanel::~SBoxPanel()&lt;/CallStack&gt;
+	&lt;CallStack&gt;UE4Editor_Slate!SVerticalBox::`scalar deleting destructor'()&lt;/CallStack&gt;
+	&lt;CallStack&gt;UE4Editor_SlateCore!FSimpleSlot::~FSimpleSlot()&lt;/CallStack&gt;
+	&lt;CallStack&gt;UE4Editor_SlateCore!SCompoundWidget::~SCompoundWidget()&lt;/CallStack&gt;
+	&lt;CallStack&gt;UE4Editor_Slate!SDockingTabStack::`vector deleting destructor'()&lt;/CallStack&gt;
+	&lt;CallStack&gt;UE4Editor_SlateCore!FSlotBase::~FSlotBase() [d:\build\++ue4+release-4.12+compile\sync\engine\source\runtime\slatecore\private\slotbase.cpp:34]&lt;/CallStack&gt;
+	&lt;CallStack&gt;UE4Editor_Slate!TIndirectArray&amp;lt;SSplitter::FSlot,FDefaultAllocator&amp;gt;::DestructAndFreeItems() [d:\build\++ue4+release-4.12+compile\sync\engine\source\runtime\core\public\containers\array.h:3410]&lt;/CallStack&gt;
+	&lt;CallStack&gt;UE4Editor_Slate!SSplitter::~SSplitter()&lt;/CallStack&gt;
+	&lt;CallStack&gt;UE4Editor_Slate!SSplitter::`vector deleting destructor'()&lt;/CallStack&gt;
+	&lt;CallStack&gt;UE4Editor_SlateCore!TIndirectArray&amp;lt;SOverlay::FOverlaySlot,FDefaultAllocator&amp;gt;::DestructAndFreeItems() [d:\build\++ue4+release-4.12+compile\sync\engine\source\runtime\core\public\containers\array.h:3410]&lt;/CallStack&gt;
+	&lt;CallStack&gt;UE4Editor_SlateCore!SOverlay::~SOverlay()&lt;/CallStack&gt;
+	&lt;CallStack&gt;UE4Editor_Slate!SOverlay::`scalar deleting destructor'()&lt;/CallStack&gt;
+	&lt;CallStack&gt;UE4Editor_SlateCore!TArray&amp;lt;FArrangedWidget,TInlineAllocator&amp;lt;16,FDefaultAllocator&amp;gt; &amp;gt;::~TArray&amp;lt;FArrangedWidget,TInlineAllocator&amp;lt;16,FDefaultAllocator&amp;gt; &amp;gt;() [d:\build\++ue4+release-4.12+compile\sync\engine\source\runtime\core\public\containers\array.h:723]&lt;/CallStack&gt;
+	&lt;CallStack&gt;UE4Editor_Slate!FSlateApplication::RoutePointerUpEvent() [d:\build\++ue4+release-4.12+compile\sync\engine\source\runtime\slate\private\framework\application\slateapplication.cpp:4667]&lt;/CallStack&gt;
+	&lt;CallStack&gt;UE4Editor_Slate!FSlateApplication::ProcessMouseButtonUpEvent() [d:\build\++ue4+release-4.12+compile\sync\engine\source\runtime\slate\private\framework\application\slateapplication.cpp:5041]&lt;/CallStack&gt;
+	&lt;CallStack&gt;UE4Editor_Slate!FSlateApplication::OnMouseUp() [d:\build\++ue4+release-4.12+compile\sync\engine\source\runtime\slate\private\framework\application\slateapplication.cpp:5021]&lt;/CallStack&gt;
+	&lt;CallStack&gt;UE4Editor_Core!FWindowsApplication::ProcessDeferredMessage() [d:\build\++ue4+release-4.12+compile\sync\engine\source\runtime\core\private\windows\windowsapplication.cpp:1510]&lt;/CallStack&gt;
+	&lt;CallStack&gt;UE4Editor_Core!FWindowsApplication::DeferMessage() [d:\build\++ue4+release-4.12+compile\sync\engine\source\runtime\core\private\windows\windowsapplication.cpp:1850]&lt;/CallStack&gt;
+	&lt;CallStack&gt;UE4Editor_Core!FWindowsApplication::ProcessMessage() [d:\build\++ue4+release-4.12+compile\sync\engine\source\runtime\core\private\windows\windowsapplication.cpp:745]&lt;/CallStack&gt;
+	&lt;CallStack&gt;UE4Editor_Core!FWindowsApplication::AppWndProc() [d:\build\++ue4+release-4.12+compile\sync\engine\source\runtime\core\private\windows\windowsapplication.cpp:667]&lt;/CallStack&gt;
+	&lt;CallStack&gt;user32!UserCallWinProcCheckWow()&lt;/CallStack&gt;
+	&lt;CallStack&gt;user32!DispatchMessageWorker()&lt;/CallStack&gt;
+	&lt;CallStack&gt;UE4Editor_Core!FWindowsPlatformMisc::PumpMessages() [d:\build\++ue4+release-4.12+compile\sync\engine\source\runtime\core\private\windows\windowsplatformmisc.cpp:903]&lt;/CallStack&gt;
+	&lt;CallStack&gt;UE4Editor!FEngineLoop::Tick() [d:\build\++ue4+release-4.12+compile\sync\engine\source\runtime\launch\private\launchengineloop.cpp:2729]&lt;/CallStack&gt;
+	&lt;CallStack&gt;UE4Editor!GuardedMain() [d:\build\++ue4+release-4.12+compile\sync\engine\source\runtime\launch\private\launch.cpp:148]&lt;/CallStack&gt;
+	&lt;CallStack&gt;UE4Editor!GuardedMainWrapper() [d:\build\++ue4+release-4.12+compile\sync\engine\source\runtime\launch\private\windows\launchwindows.cpp:126]&lt;/CallStack&gt;
+	&lt;CallStack&gt;UE4Editor!WinMain() [d:\build\++ue4+release-4.12+compile\sync\engine\source\runtime\launch\private\windows\launchwindows.cpp:200]&lt;/CallStack&gt;
+	&lt;CallStack&gt;UE4Editor!__scrt_common_main_seh() [f:\dd\vctools\crt\vcstartup\src\startup\exe_common.inl:264]&lt;/CallStack&gt;
+	&lt;CallStack&gt;kernel32!BaseThreadInitThunk()&lt;/CallStack&gt;
+	&lt;CallStack&gt;ntdll!RtlUserThreadStart()&lt;/CallStack&gt;
+	&lt;SourceContext&gt;  670       	{&lt;/SourceContext&gt;
+	&lt;SourceContext&gt;  671       		if (OldObject == NewObject)&lt;/SourceContext&gt;
+	&lt;SourceContext&gt;  672       		{&lt;/SourceContext&gt;
+	&lt;SourceContext&gt;  673       			// Nothing changed&lt;/SourceContext&gt;
+	&lt;SourceContext&gt;  674       			return;&lt;/SourceContext&gt;
+	&lt;SourceContext&gt;  675       		}&lt;/SourceContext&gt;
+	&lt;SourceContext&gt;  676       &lt;/SourceContext&gt;
+	&lt;SourceContext&gt;  677       		DebuggingWorldRegistrationHelper(OldObject, NULL);&lt;/SourceContext&gt;
+	&lt;SourceContext&gt;  678       	}&lt;/SourceContext&gt;
+	&lt;SourceContext&gt;  679       &lt;/SourceContext&gt;
+	&lt;SourceContext&gt;  680       	// Note that we allow macro Blueprints to bypass this check&lt;/SourceContext&gt;
+	&lt;SourceContext&gt;  681       	if ((NewObject != NULL) &amp;amp;&amp;amp; !GCompilingBlueprint &amp;amp;&amp;amp; BlueprintType != BPTYPE_MacroLibrary)&lt;/SourceContext&gt;
+	&lt;SourceContext&gt;  682       	{&lt;/SourceContext&gt;
+	&lt;SourceContext&gt;  683       		// You can only debug instances of this!&lt;/SourceContext&gt;
+	&lt;SourceContext&gt;  684 ***** 		if (!ensureMsgf(&lt;/SourceContext&gt;
+	&lt;SourceContext&gt;  685       				NewObject-&amp;gt;IsA(this-&amp;gt;GeneratedClass), &lt;/SourceContext&gt;
+	&lt;SourceContext&gt;  686       				TEXT("Type mismatch: Expected %s, Found %s"), &lt;/SourceContext&gt;
+	&lt;SourceContext&gt;  687       				this-&amp;gt;GeneratedClass ? *(this-&amp;gt;GeneratedClass-&amp;gt;GetName()) : TEXT("NULL"), &lt;/SourceContext&gt;
+	&lt;SourceContext&gt;  688       				NewObject-&amp;gt;GetClass() ? *(this-&amp;gt;GetClass()-&amp;gt;GetName()) : TEXT("NULL")))&lt;/SourceContext&gt;
+	&lt;SourceContext&gt;  689       		{&lt;/SourceContext&gt;
+	&lt;SourceContext&gt;  690       			NewObject = NULL;&lt;/SourceContext&gt;
+	&lt;SourceContext&gt;  691       		}&lt;/SourceContext&gt;
+	&lt;SourceContext&gt;  692       	}&lt;/SourceContext&gt;
+	&lt;SourceContext&gt;  693       &lt;/SourceContext&gt;
+	&lt;SourceContext&gt;  694       	// Update the current object being debugged&lt;/SourceContext&gt;
+	&lt;SourceContext&gt;  695       	CurrentObjectBeingDebugged = NewObject;&lt;/SourceContext&gt;
+	&lt;SourceContext&gt;  696       &lt;/SourceContext&gt;
+	&lt;SourceContext&gt;  697       	// Register the new object&lt;/SourceContext&gt;
+	&lt;SourceContext&gt;  698       	if (NewObject != NULL)&lt;/SourceContext&gt;
+	&lt;SourceContext&gt;  699       	{&lt;/SourceContext&gt;
+	&lt;UserDescription&gt;Sent in the unattended mode&lt;/UserDescription&gt;
+	&lt;UserActivityHint /&gt;
+	&lt;ErrorMessage&gt;Ensure condition failed: NewObject-&amp;gt;IsA(this-&amp;gt;GeneratedClass) [File:D:\Build\++UE4+Release-4.12+Compile\Sync\Engine\Source\Runtime\Engine\Private\Blueprint.cpp] [Line: 689] &lt;/ErrorMessage&gt;
+	&lt;ErrorMessage&gt;Type mismatch: Expected Untitled_C, Found LevelScriptBlueprint&lt;/ErrorMessage&gt;
+	&lt;CrashGUID&gt;UE4CC-Windows-7188BC094DE45EE1A7D39EBF6FB434F5_0000&lt;/CrashGUID&gt;
+	&lt;TimeofCrash&gt;2016-07-13T14:58:43.402875Z&lt;/TimeofCrash&gt;
+	&lt;bHasMiniDump&gt;true&lt;/bHasMiniDump&gt;
+	&lt;bHasLog&gt;true&lt;/bHasLog&gt;
+	&lt;bHasDiags&gt;false&lt;/bHasDiags&gt;
+	&lt;bHasWERData&gt;false&lt;/bHasWERData&gt;
+	&lt;bHasVideo&gt;false&lt;/bHasVideo&gt;
+	&lt;bAllowToBeContacted&gt;true&lt;/bAllowToBeContacted&gt;
+&lt;/CrashDescription&gt;</Value>
+    </Setting>
+    <Setting Name="AWSS3CompressedSuffix" Type="System.String" Scope="Application">
+      <Value Profile="(Default)">.gz</Value>
+    </Setting>
+    <Setting Name="CrashesCacheInMinutes" Type="System.Int32" Scope="Application">
+      <Value Profile="(Default)">15</Value>
+    </Setting>
+    <Setting Name="SearchFilterCacheInMinutes" Type="System.Int32" Scope="Application">
+      <Value Profile="(Default)">60</Value>
+    </Setting>
   </Settings>
 </SettingsFile>
\ No newline at end of file
diff --git a/Engine/Source/Programs/CrashReporter/CrashReportWebSite/ViewModels/BuggViewModel.cs b/Engine/Source/Programs/CrashReporter/CrashReportWebSite/ViewModels/BuggViewModel.cs
--- a/Engine/Source/Programs/CrashReporter/CrashReportWebSite/ViewModels/BuggViewModel.cs
+++ b/Engine/Source/Programs/CrashReporter/CrashReportWebSite/ViewModels/BuggViewModel.cs
@@ -10,10 +10,13 @@ namespace Tools.CrashReporter.CrashReportWebSite.ViewModels
     public class CrashDataModel
     {
         public int Id { get; set; }
+        public int CrashType { get; set; }
         public string ChangelistVersion { get; set; }
         public string BuildVersion { get; set; }
         public string UserName { get; set; }
         public string GameName { get; set; }
+        public string Summary { get; set; }
+        public string SourceContext { get; set; }
 
         public string EngineMode { get; set; }
         public string PlatformName { get; set; }
@@ -21,6 +24,12 @@ namespace Tools.CrashReporter.CrashReportWebSite.ViewModels
 
         public string Description { get; set; }
 
+        public string RawCallStack { get; set; }
+
+        public string ComputerName { get; set; }
+
+        public string Branch { get; set; }
+
         /// <summary>
         /// Return a display friendly version of the time of Crash.
         /// </summary>
diff --git a/Engine/Source/Programs/CrashReporter/CrashReportWebSite/Web.Debug.config b/Engine/Source/Programs/CrashReporter/CrashReportWebSite/Web.Debug.config
--- a/Engine/Source/Programs/CrashReporter/CrashReportWebSite/Web.Debug.config
+++ b/Engine/Source/Programs/CrashReporter/CrashReportWebSite/Web.Debug.config
@@ -16,9 +16,10 @@
   -->
   <connectionStrings  xdt:Transform="Replace">
     <add name="CrashReportConnectionString"
-      connectionString="Data Source=NCLRNDEV-02.ncl.epicgames.net;Database=CrashReporterDev;User ID=rntdashboard;Password=Yellow"
-      providerName="System.Data.SqlClient"/>
-    <add name="Hangfire" connectionString="Data Source=db-09;Initial Catalog=CrashReporterSchduledTask;Server=db-09;Database=CrashReport;Integrated Security=true" providerName="System.Data.SqlClient" />
+          connectionString="Data Source=db-09.ol.epicgames.net;Initial Catalog=CrashReport;Server=db-09.ol.epicgames.net;Database=CrashReport;User Id=crashreporter;Password=Epic123!"
+          providerName="System.Data.SqlClient"/>
+    <add name="CrashReportEntities" connectionString="metadata=res://*/DataModels.CrashReports.csdl|res://*/DataModels.CrashReports.ssdl|res://*/DataModels.CrashReports.msl;provider=System.Data.SqlClient;provider connection string=&quot;data source=db-09.ol.epicgames.net;initial catalog=CrashReport;User Id=crashreporter;Password=Epic123!;MultipleActiveResultSets=True;App=EntityFramework&quot;" providerName="System.Data.EntityClient" />
+    <add name="Hangfire" connectionString="Data Source=db-09.ol.epicgames.net;Initial Catalog=CrashReporterSchduledTask;Server=db-09.ol.epicgames.net;Database=CrashReport;User Id=crashreporter;Password=Epic123!" providerName="System.Data.SqlClient" />
   </connectionStrings>
   <system.web>
     <!--
diff --git a/Engine/Source/Programs/CrashReporter/CrashReportWebSite/Web.Development.config b/Engine/Source/Programs/CrashReporter/CrashReportWebSite/Web.Development.config
--- a/Engine/Source/Programs/CrashReporter/CrashReportWebSite/Web.Development.config
+++ b/Engine/Source/Programs/CrashReporter/CrashReportWebSite/Web.Development.config
@@ -16,9 +16,10 @@
   -->
   <connectionStrings  xdt:Transform="Replace">
     <add name="CrashReportConnectionString"
-      connectionString="Data Source=db-09;Initial Catalog=CrashReport;Server=db-09;Database=CrashReport;Integrated Security=true"
-      providerName="System.Data.SqlClient"/>
-    <add name="Hangfire" connectionString="Data Source=db-09;Initial Catalog=CrashReporterSchduledTask;Server=db-09;Database=CrashReport;Integrated Security=true" providerName="System.Data.SqlClient" />
+          connectionString="Data Source=db-09.ol.epicgames.net;Initial Catalog=CrashReport;Server=db-09.ol.epicgames.net;Database=CrashReport;User Id=crashreporter;Password=Epic123!"
+          providerName="System.Data.SqlClient"/>
+    <add name="CrashReportEntities" connectionString="metadata=res://*/DataModels.CrashReports.csdl|res://*/DataModels.CrashReports.ssdl|res://*/DataModels.CrashReports.msl;provider=System.Data.SqlClient;provider connection string=&quot;data source=db-09.ol.epicgames.net;initial catalog=CrashReport;User Id=crashreporter;Password=Epic123!;MultipleActiveResultSets=True;App=EntityFramework&quot;" providerName="System.Data.EntityClient" />
+    <add name="Hangfire" connectionString="Data Source=db-09.ol.epicgames.net;Initial Catalog=CrashReporterSchduledTask;Server=db-09.ol.epicgames.net;Database=CrashReport;User Id=crashreporter;Password=Epic123!" providerName="System.Data.SqlClient" />
   </connectionStrings>
   <system.web>
     <compilation xdt:Transform="RemoveAttributes(debug)" />
diff --git a/Engine/Source/Programs/CrashReporter/CrashReportWebSite/Web.Release.config b/Engine/Source/Programs/CrashReporter/CrashReportWebSite/Web.Release.config
--- a/Engine/Source/Programs/CrashReporter/CrashReportWebSite/Web.Release.config
+++ b/Engine/Source/Programs/CrashReporter/CrashReportWebSite/Web.Release.config
@@ -14,6 +14,13 @@
         xdt:Transform="SetAttributes" xdt:Locator="Match(name)"/>
     </connectionStrings>
   -->
+  <!--<connectionStrings  xdt:Transform="Replace">
+      <add name="CrashReportConnectionString"
+            connectionString="Data Source=db-09.ol.epicgames.net;Initial Catalog=CrashReport;Server=db-09.ol.epicgames.net;Database=CrashReport;User Id=crashreporter;Password=Epic123!"
+            providerName="System.Data.SqlClient"/>
+      <add name="CrashReportEntities" connectionString="metadata=res://*/DataModels.CrashReports.csdl|res://*/DataModels.CrashReports.ssdl|res://*/DataModels.CrashReports.msl;provider=System.Data.SqlClient;provider connection string=&quot;data source=db-09.ol.epicgames.net;initial catalog=CrashReport;User Id=crashreporter;Password=Epic123!;MultipleActiveResultSets=True;App=EntityFramework&quot;" providerName="System.Data.EntityClient" />
+      <add name="Hangfire" connectionString="Data Source=db-09.ol.epicgames.net;Initial Catalog=CrashReporterSchduledTask;Server=db-09.ol.epicgames.net;Database=CrashReport;User Id=crashreporter;Password=Epic123!" providerName="System.Data.SqlClient" />
+  </connectionStrings>-->
   <system.web>
     <compilation xdt:Transform="RemoveAttributes(debug)" />
     <!--
diff --git a/Engine/Source/Programs/CrashReporter/CrashReportWebSite/Web.config b/Engine/Source/Programs/CrashReporter/CrashReportWebSite/Web.config
--- a/Engine/Source/Programs/CrashReporter/CrashReportWebSite/Web.config
+++ b/Engine/Source/Programs/CrashReporter/CrashReportWebSite/Web.config
@@ -16,8 +16,11 @@
   </configSections>
   <appSettings/>
   <connectionStrings>
-    <add name="Hangfire" connectionString="Data Source=db-09;Initial Catalog=CrashReporterSchduledTask;Server=db-09;Database=CrashReport;Integrated Security=true" providerName="System.Data.SqlClient"/>
-    <add name="CrashReportEntities" connectionString="metadata=res://*/DataModels.CrashReports.csdl|res://*/DataModels.CrashReports.ssdl|res://*/DataModels.CrashReports.msl;provider=System.Data.SqlClient;provider connection string=&quot;data source=db-09.epicgames.net;Initial Catalog=CrashReport;Integrated Security=True;Max Pool Size=200;MultipleActiveResultSets=True;Connect Timeout=30;App=EntityFramework&quot;" providerName="System.Data.EntityClient"/>
+    <add name="CrashReportConnectionString"
+          connectionString="Data Source=db-09.ol.epicgames.net;Initial Catalog=CrashReport;Server=db-09.ol.epicgames.net;Database=CrashReport;User Id=crashreporter;Password=Epic123!"
+          providerName="System.Data.SqlClient"/>
+    <add name="CrashReportEntities" connectionString="metadata=res://*/DataModels.CrashReports.csdl|res://*/DataModels.CrashReports.ssdl|res://*/DataModels.CrashReports.msl;provider=System.Data.SqlClient;provider connection string=&quot;data source=db-09.ol.epicgames.net;initial catalog=CrashReport;User Id=crashreporter;Password=Epic123!;MultipleActiveResultSets=True;App=EntityFramework&quot;" providerName="System.Data.EntityClient" />
+    <add name="Hangfire" connectionString="Data Source=db-09.ol.epicgames.net;Initial Catalog=CrashReporterSchduledTask;Server=db-09.ol.epicgames.net;Database=CrashReport;User Id=crashreporter;Password=Epic123!" providerName="System.Data.SqlClient" />
   </connectionStrings>
   <system.web>
     <compilation debug="true" targetFramework="4.5"/>
@@ -81,6 +84,180 @@
       <setting name="SlackEmoji" serializeAs="String">
         <value>:rngesus:</value>
       </setting>
+      <setting name="AWSS3ProfileName" serializeAs="String">
+        <value>crashreporter-s3</value>
+      </setting>
+      <setting name="AWSCredentialsFilepath" serializeAs="String">
+        <value>C:\ProgramData\awscredentials.ini</value>
+      </setting>
+      <setting name="AWSS3URL" serializeAs="String">
+        <value>https://s3.amazonaws.com</value>
+      </setting>
+      <setting name="AWSS3KeyPrefix" serializeAs="String">
+        <value>crashreporterfiles-prod</value>
+      </setting>
+      <setting name="DownloadFromS3" serializeAs="String">
+        <value>True</value>
+      </setting>
+      <setting name="TestXML" serializeAs="String">
+        <value>&lt;CrashDescription&gt;
+	&lt;CrashType /&gt;
+	&lt;BranchName&gt;//UE4/Release-4.12&lt;/BranchName&gt;
+	&lt;GameName&gt;UE4-Lenar_3&lt;/GameName&gt;
+	&lt;Platform&gt;Win64 [Windows 7 Service Pack 1 64b]&lt;/Platform&gt;
+	&lt;EngineMode&gt;Editor&lt;/EngineMode&gt;
+	&lt;BuildVersion&gt;4.12.5&lt;/BuildVersion&gt;
+	&lt;BuiltFromCL&gt;3039270&lt;/BuiltFromCL&gt;
+	&lt;CommandLine /&gt;
+	&lt;BaseDir&gt;C:/Program Files (x86)/Epic Games/4.12/Engine/Binaries/Win64/&lt;/BaseDir&gt;
+	&lt;Language&gt;zh_CN&lt;/Language&gt;
+	&lt;SystemLanguage /&gt;
+	&lt;UserName&gt;ohmimi&lt;/UserName&gt;
+	&lt;MachineGuid&gt;A24BAFFA41A33062DF1BCB9A31DDE162&lt;/MachineGuid&gt;
+	&lt;EpicAccountId&gt;cc732bdf5c4443f4a966a15bd4869c1e&lt;/EpicAccountId&gt;
+	&lt;CallStack&gt;UE4Editor_Core!FDebug::OptionallyLogFormattedEnsureMessageReturningFalse() [d:\build\++ue4+release-4.12+compile\sync\engine\source\runtime\core\private\misc\outputdevice.cpp:451]&lt;/CallStack&gt;
+	&lt;CallStack&gt;UE4Editor_Engine!UBlueprint::SetObjectBeingDebugged() [d:\build\++ue4+release-4.12+compile\sync\engine\source\runtime\engine\private\blueprint.cpp:685]&lt;/CallStack&gt;
+	&lt;CallStack&gt;UE4Editor_Engine!UWorld::TransferBlueprintDebugReferences() [d:\build\++ue4+release-4.12+compile\sync\engine\source\runtime\engine\private\world.cpp:1618]&lt;/CallStack&gt;
+	&lt;CallStack&gt;UE4Editor_UnrealEd!UEditorEngine::TeardownPlaySession() [d:\build\++ue4+release-4.12+compile\sync\engine\source\editor\unrealed\private\playlevel.cpp:642]&lt;/CallStack&gt;
+	&lt;CallStack&gt;UE4Editor_UnrealEd!UEditorEngine::EndPlayMap() [d:\build\++ue4+release-4.12+compile\sync\engine\source\editor\unrealed\private\playlevel.cpp:255]&lt;/CallStack&gt;
+	&lt;CallStack&gt;UE4Editor_LevelEditor!FLevelViewportLayout::~FLevelViewportLayout() [d:\build\++ue4+release-4.12+compile\sync\engine\source\editor\leveleditor\private\levelviewportlayout.cpp:145]&lt;/CallStack&gt;
+	&lt;CallStack&gt;UE4Editor_LevelEditor!FLevelViewportLayout2x2::`scalar deleting destructor'()&lt;/CallStack&gt;
+	&lt;CallStack&gt;UE4Editor_LevelEditor!FLevelViewportTabContent::~FLevelViewportTabContent()&lt;/CallStack&gt;
+	&lt;CallStack&gt;UE4Editor_LevelEditor!SharedPointerInternals::TReferenceControllerWithDeleter&amp;lt;FLevelViewportTabContent,SharedPointerInternals::DefaultDeleter&amp;lt;FLevelViewportTabContent&amp;gt; &amp;gt;::DestroyObject() [d:\build\++ue4+release-4.12+compile\sync\engine\source\runtime\core\public\templates\sharedpointerinternals.h:104]&lt;/CallStack&gt;
+	&lt;CallStack&gt;UE4Editor_LevelEditor!SViewportsOverlay::~SViewportsOverlay()&lt;/CallStack&gt;
+	&lt;CallStack&gt;UE4Editor_LevelEditor!SViewportsOverlay::`scalar deleting destructor'()&lt;/CallStack&gt;
+	&lt;CallStack&gt;UE4Editor_SlateCore!FSimpleSlot::~FSimpleSlot()&lt;/CallStack&gt;
+	&lt;CallStack&gt;UE4Editor_SlateCore!SCompoundWidget::~SCompoundWidget()&lt;/CallStack&gt;
+	&lt;CallStack&gt;UE4Editor_Slate!SBorder::`vector deleting destructor'()&lt;/CallStack&gt;
+	&lt;CallStack&gt;UE4Editor_SlateCore!TIndirectArray&amp;lt;SOverlay::FOverlaySlot,FDefaultAllocator&amp;gt;::DestructAndFreeItems() [d:\build\++ue4+release-4.12+compile\sync\engine\source\runtime\core\public\containers\array.h:3410]&lt;/CallStack&gt;
+	&lt;CallStack&gt;UE4Editor_SlateCore!SOverlay::~SOverlay()&lt;/CallStack&gt;
+	&lt;CallStack&gt;UE4Editor_Slate!SOverlay::`scalar deleting destructor'()&lt;/CallStack&gt;
+	&lt;CallStack&gt;UE4Editor_SlateCore!TIndirectArray&amp;lt;SBoxPanel::FSlot,FDefaultAllocator&amp;gt;::DestructAndFreeItems() [d:\build\++ue4+release-4.12+compile\sync\engine\source\runtime\core\public\containers\array.h:3410]&lt;/CallStack&gt;
+	&lt;CallStack&gt;UE4Editor_SlateCore!SBoxPanel::~SBoxPanel()&lt;/CallStack&gt;
+	&lt;CallStack&gt;UE4Editor_Slate!SVerticalBox::`scalar deleting destructor'()&lt;/CallStack&gt;
+	&lt;CallStack&gt;UE4Editor_SlateCore!FSimpleSlot::~FSimpleSlot()&lt;/CallStack&gt;
+	&lt;CallStack&gt;UE4Editor_SlateCore!SCompoundWidget::~SCompoundWidget()&lt;/CallStack&gt;
+	&lt;CallStack&gt;UE4Editor_Slate!SDockingTabStack::`vector deleting destructor'()&lt;/CallStack&gt;
+	&lt;CallStack&gt;UE4Editor_SlateCore!FSlotBase::~FSlotBase() [d:\build\++ue4+release-4.12+compile\sync\engine\source\runtime\slatecore\private\slotbase.cpp:34]&lt;/CallStack&gt;
+	&lt;CallStack&gt;UE4Editor_Slate!TIndirectArray&amp;lt;SSplitter::FSlot,FDefaultAllocator&amp;gt;::DestructAndFreeItems() [d:\build\++ue4+release-4.12+compile\sync\engine\source\runtime\core\public\containers\array.h:3410]&lt;/CallStack&gt;
+	&lt;CallStack&gt;UE4Editor_Slate!SSplitter::~SSplitter()&lt;/CallStack&gt;
+	&lt;CallStack&gt;UE4Editor_Slate!SSplitter::`vector deleting destructor'()&lt;/CallStack&gt;
+	&lt;CallStack&gt;UE4Editor_SlateCore!FSimpleSlot::~FSimpleSlot()&lt;/CallStack&gt;
+	&lt;CallStack&gt;UE4Editor_SlateCore!SCompoundWidget::~SCompoundWidget()&lt;/CallStack&gt;
+	&lt;CallStack&gt;UE4Editor_Slate!SDockingSplitter::`vector deleting destructor'()&lt;/CallStack&gt;
+	&lt;CallStack&gt;UE4Editor_SlateCore!FSlotBase::~FSlotBase() [d:\build\++ue4+release-4.12+compile\sync\engine\source\runtime\slatecore\private\slotbase.cpp:34]&lt;/CallStack&gt;
+	&lt;CallStack&gt;UE4Editor_Slate!TIndirectArray&amp;lt;SSplitter::FSlot,FDefaultAllocator&amp;gt;::DestructAndFreeItems() [d:\build\++ue4+release-4.12+compile\sync\engine\source\runtime\core\public\containers\array.h:3410]&lt;/CallStack&gt;
+	&lt;CallStack&gt;UE4Editor_Slate!SSplitter::~SSplitter()&lt;/CallStack&gt;
+	&lt;CallStack&gt;UE4Editor_Slate!SSplitter::`vector deleting destructor'()&lt;/CallStack&gt;
+	&lt;CallStack&gt;UE4Editor_SlateCore!FSimpleSlot::~FSimpleSlot()&lt;/CallStack&gt;
+	&lt;CallStack&gt;UE4Editor_SlateCore!SCompoundWidget::~SCompoundWidget()&lt;/CallStack&gt;
+	&lt;CallStack&gt;UE4Editor_Slate!SDockingSplitter::`vector deleting destructor'()&lt;/CallStack&gt;
+	&lt;CallStack&gt;UE4Editor_SlateCore!FSlotBase::~FSlotBase() [d:\build\++ue4+release-4.12+compile\sync\engine\source\runtime\slatecore\private\slotbase.cpp:34]&lt;/CallStack&gt;
+	&lt;CallStack&gt;UE4Editor_Slate!TIndirectArray&amp;lt;SSplitter::FSlot,FDefaultAllocator&amp;gt;::DestructAndFreeItems() [d:\build\++ue4+release-4.12+compile\sync\engine\source\runtime\core\public\containers\array.h:3410]&lt;/CallStack&gt;
+	&lt;CallStack&gt;UE4Editor_Slate!SSplitter::~SSplitter()&lt;/CallStack&gt;
+	&lt;CallStack&gt;UE4Editor_Slate!SSplitter::`vector deleting destructor'()&lt;/CallStack&gt;
+	&lt;CallStack&gt;UE4Editor_SlateCore!TIndirectArray&amp;lt;SOverlay::FOverlaySlot,FDefaultAllocator&amp;gt;::DestructAndFreeItems() [d:\build\++ue4+release-4.12+compile\sync\engine\source\runtime\core\public\containers\array.h:3410]&lt;/CallStack&gt;
+	&lt;CallStack&gt;UE4Editor_SlateCore!SOverlay::~SOverlay()&lt;/CallStack&gt;
+	&lt;CallStack&gt;UE4Editor_Slate!SOverlay::`scalar deleting destructor'()&lt;/CallStack&gt;
+	&lt;CallStack&gt;UE4Editor_SlateCore!FSimpleSlot::~FSimpleSlot()&lt;/CallStack&gt;
+	&lt;CallStack&gt;UE4Editor_SlateCore!SCompoundWidget::~SCompoundWidget()&lt;/CallStack&gt;
+	&lt;CallStack&gt;UE4Editor_Slate!SDockingArea::`vector deleting destructor'()&lt;/CallStack&gt;
+	&lt;CallStack&gt;UE4Editor_SlateCore!TIndirectArray&amp;lt;SBoxPanel::FSlot,FDefaultAllocator&amp;gt;::DestructAndFreeItems() [d:\build\++ue4+release-4.12+compile\sync\engine\source\runtime\core\public\containers\array.h:3410]&lt;/CallStack&gt;
+	&lt;CallStack&gt;UE4Editor_SlateCore!SBoxPanel::~SBoxPanel()&lt;/CallStack&gt;
+	&lt;CallStack&gt;UE4Editor_LevelEditor!SVerticalBox::`scalar deleting destructor'()&lt;/CallStack&gt;
+	&lt;CallStack&gt;UE4Editor_SlateCore!FSimpleSlot::~FSimpleSlot()&lt;/CallStack&gt;
+	&lt;CallStack&gt;UE4Editor_SlateCore!SCompoundWidget::~SCompoundWidget()&lt;/CallStack&gt;
+	&lt;CallStack&gt;UE4Editor_LevelEditor!SLevelEditor::`scalar deleting destructor'()&lt;/CallStack&gt;
+	&lt;CallStack&gt;UE4Editor_SlateCore!FSimpleSlot::~FSimpleSlot()&lt;/CallStack&gt;
+	&lt;CallStack&gt;UE4Editor_SlateCore!SCompoundWidget::~SCompoundWidget()&lt;/CallStack&gt;
+	&lt;CallStack&gt;UE4Editor_Slate!SBorder::`vector deleting destructor'()&lt;/CallStack&gt;
+	&lt;CallStack&gt;UE4Editor_SlateCore!TIndirectArray&amp;lt;SOverlay::FOverlaySlot,FDefaultAllocator&amp;gt;::DestructAndFreeItems() [d:\build\++ue4+release-4.12+compile\sync\engine\source\runtime\core\public\containers\array.h:3410]&lt;/CallStack&gt;
+	&lt;CallStack&gt;UE4Editor_SlateCore!SOverlay::~SOverlay()&lt;/CallStack&gt;
+	&lt;CallStack&gt;UE4Editor_Slate!SOverlay::`scalar deleting destructor'()&lt;/CallStack&gt;
+	&lt;CallStack&gt;UE4Editor_SlateCore!TIndirectArray&amp;lt;SBoxPanel::FSlot,FDefaultAllocator&amp;gt;::DestructAndFreeItems() [d:\build\++ue4+release-4.12+compile\sync\engine\source\runtime\core\public\containers\array.h:3410]&lt;/CallStack&gt;
+	&lt;CallStack&gt;UE4Editor_SlateCore!SBoxPanel::~SBoxPanel()&lt;/CallStack&gt;
+	&lt;CallStack&gt;UE4Editor_Slate!SVerticalBox::`scalar deleting destructor'()&lt;/CallStack&gt;
+	&lt;CallStack&gt;UE4Editor_SlateCore!FSimpleSlot::~FSimpleSlot()&lt;/CallStack&gt;
+	&lt;CallStack&gt;UE4Editor_SlateCore!SCompoundWidget::~SCompoundWidget()&lt;/CallStack&gt;
+	&lt;CallStack&gt;UE4Editor_Slate!SDockingTabStack::`vector deleting destructor'()&lt;/CallStack&gt;
+	&lt;CallStack&gt;UE4Editor_SlateCore!FSlotBase::~FSlotBase() [d:\build\++ue4+release-4.12+compile\sync\engine\source\runtime\slatecore\private\slotbase.cpp:34]&lt;/CallStack&gt;
+	&lt;CallStack&gt;UE4Editor_Slate!TIndirectArray&amp;lt;SSplitter::FSlot,FDefaultAllocator&amp;gt;::DestructAndFreeItems() [d:\build\++ue4+release-4.12+compile\sync\engine\source\runtime\core\public\containers\array.h:3410]&lt;/CallStack&gt;
+	&lt;CallStack&gt;UE4Editor_Slate!SSplitter::~SSplitter()&lt;/CallStack&gt;
+	&lt;CallStack&gt;UE4Editor_Slate!SSplitter::`vector deleting destructor'()&lt;/CallStack&gt;
+	&lt;CallStack&gt;UE4Editor_SlateCore!TIndirectArray&amp;lt;SOverlay::FOverlaySlot,FDefaultAllocator&amp;gt;::DestructAndFreeItems() [d:\build\++ue4+release-4.12+compile\sync\engine\source\runtime\core\public\containers\array.h:3410]&lt;/CallStack&gt;
+	&lt;CallStack&gt;UE4Editor_SlateCore!SOverlay::~SOverlay()&lt;/CallStack&gt;
+	&lt;CallStack&gt;UE4Editor_Slate!SOverlay::`scalar deleting destructor'()&lt;/CallStack&gt;
+	&lt;CallStack&gt;UE4Editor_SlateCore!TArray&amp;lt;FArrangedWidget,TInlineAllocator&amp;lt;16,FDefaultAllocator&amp;gt; &amp;gt;::~TArray&amp;lt;FArrangedWidget,TInlineAllocator&amp;lt;16,FDefaultAllocator&amp;gt; &amp;gt;() [d:\build\++ue4+release-4.12+compile\sync\engine\source\runtime\core\public\containers\array.h:723]&lt;/CallStack&gt;
+	&lt;CallStack&gt;UE4Editor_Slate!FSlateApplication::RoutePointerUpEvent() [d:\build\++ue4+release-4.12+compile\sync\engine\source\runtime\slate\private\framework\application\slateapplication.cpp:4667]&lt;/CallStack&gt;
+	&lt;CallStack&gt;UE4Editor_Slate!FSlateApplication::ProcessMouseButtonUpEvent() [d:\build\++ue4+release-4.12+compile\sync\engine\source\runtime\slate\private\framework\application\slateapplication.cpp:5041]&lt;/CallStack&gt;
+	&lt;CallStack&gt;UE4Editor_Slate!FSlateApplication::OnMouseUp() [d:\build\++ue4+release-4.12+compile\sync\engine\source\runtime\slate\private\framework\application\slateapplication.cpp:5021]&lt;/CallStack&gt;
+	&lt;CallStack&gt;UE4Editor_Core!FWindowsApplication::ProcessDeferredMessage() [d:\build\++ue4+release-4.12+compile\sync\engine\source\runtime\core\private\windows\windowsapplication.cpp:1510]&lt;/CallStack&gt;
+	&lt;CallStack&gt;UE4Editor_Core!FWindowsApplication::DeferMessage() [d:\build\++ue4+release-4.12+compile\sync\engine\source\runtime\core\private\windows\windowsapplication.cpp:1850]&lt;/CallStack&gt;
+	&lt;CallStack&gt;UE4Editor_Core!FWindowsApplication::ProcessMessage() [d:\build\++ue4+release-4.12+compile\sync\engine\source\runtime\core\private\windows\windowsapplication.cpp:745]&lt;/CallStack&gt;
+	&lt;CallStack&gt;UE4Editor_Core!FWindowsApplication::AppWndProc() [d:\build\++ue4+release-4.12+compile\sync\engine\source\runtime\core\private\windows\windowsapplication.cpp:667]&lt;/CallStack&gt;
+	&lt;CallStack&gt;user32!UserCallWinProcCheckWow()&lt;/CallStack&gt;
+	&lt;CallStack&gt;user32!DispatchMessageWorker()&lt;/CallStack&gt;
+	&lt;CallStack&gt;UE4Editor_Core!FWindowsPlatformMisc::PumpMessages() [d:\build\++ue4+release-4.12+compile\sync\engine\source\runtime\core\private\windows\windowsplatformmisc.cpp:903]&lt;/CallStack&gt;
+	&lt;CallStack&gt;UE4Editor!FEngineLoop::Tick() [d:\build\++ue4+release-4.12+compile\sync\engine\source\runtime\launch\private\launchengineloop.cpp:2729]&lt;/CallStack&gt;
+	&lt;CallStack&gt;UE4Editor!GuardedMain() [d:\build\++ue4+release-4.12+compile\sync\engine\source\runtime\launch\private\launch.cpp:148]&lt;/CallStack&gt;
+	&lt;CallStack&gt;UE4Editor!GuardedMainWrapper() [d:\build\++ue4+release-4.12+compile\sync\engine\source\runtime\launch\private\windows\launchwindows.cpp:126]&lt;/CallStack&gt;
+	&lt;CallStack&gt;UE4Editor!WinMain() [d:\build\++ue4+release-4.12+compile\sync\engine\source\runtime\launch\private\windows\launchwindows.cpp:200]&lt;/CallStack&gt;
+	&lt;CallStack&gt;UE4Editor!__scrt_common_main_seh() [f:\dd\vctools\crt\vcstartup\src\startup\exe_common.inl:264]&lt;/CallStack&gt;
+	&lt;CallStack&gt;kernel32!BaseThreadInitThunk()&lt;/CallStack&gt;
+	&lt;CallStack&gt;ntdll!RtlUserThreadStart()&lt;/CallStack&gt;
+	&lt;SourceContext&gt;  670       	{&lt;/SourceContext&gt;
+	&lt;SourceContext&gt;  671       		if (OldObject == NewObject)&lt;/SourceContext&gt;
+	&lt;SourceContext&gt;  672       		{&lt;/SourceContext&gt;
+	&lt;SourceContext&gt;  673       			// Nothing changed&lt;/SourceContext&gt;
+	&lt;SourceContext&gt;  674       			return;&lt;/SourceContext&gt;
+	&lt;SourceContext&gt;  675       		}&lt;/SourceContext&gt;
+	&lt;SourceContext&gt;  676       &lt;/SourceContext&gt;
+	&lt;SourceContext&gt;  677       		DebuggingWorldRegistrationHelper(OldObject, NULL);&lt;/SourceContext&gt;
+	&lt;SourceContext&gt;  678       	}&lt;/SourceContext&gt;
+	&lt;SourceContext&gt;  679       &lt;/SourceContext&gt;
+	&lt;SourceContext&gt;  680       	// Note that we allow macro Blueprints to bypass this check&lt;/SourceContext&gt;
+	&lt;SourceContext&gt;  681       	if ((NewObject != NULL) &amp;amp;&amp;amp; !GCompilingBlueprint &amp;amp;&amp;amp; BlueprintType != BPTYPE_MacroLibrary)&lt;/SourceContext&gt;
+	&lt;SourceContext&gt;  682       	{&lt;/SourceContext&gt;
+	&lt;SourceContext&gt;  683       		// You can only debug instances of this!&lt;/SourceContext&gt;
+	&lt;SourceContext&gt;  684 ***** 		if (!ensureMsgf(&lt;/SourceContext&gt;
+	&lt;SourceContext&gt;  685       				NewObject-&amp;gt;IsA(this-&amp;gt;GeneratedClass), &lt;/SourceContext&gt;
+	&lt;SourceContext&gt;  686       				TEXT("Type mismatch: Expected %s, Found %s"), &lt;/SourceContext&gt;
+	&lt;SourceContext&gt;  687       				this-&amp;gt;GeneratedClass ? *(this-&amp;gt;GeneratedClass-&amp;gt;GetName()) : TEXT("NULL"), &lt;/SourceContext&gt;
+	&lt;SourceContext&gt;  688       				NewObject-&amp;gt;GetClass() ? *(this-&amp;gt;GetClass()-&amp;gt;GetName()) : TEXT("NULL")))&lt;/SourceContext&gt;
+	&lt;SourceContext&gt;  689       		{&lt;/SourceContext&gt;
+	&lt;SourceContext&gt;  690       			NewObject = NULL;&lt;/SourceContext&gt;
+	&lt;SourceContext&gt;  691       		}&lt;/SourceContext&gt;
+	&lt;SourceContext&gt;  692       	}&lt;/SourceContext&gt;
+	&lt;SourceContext&gt;  693       &lt;/SourceContext&gt;
+	&lt;SourceContext&gt;  694       	// Update the current object being debugged&lt;/SourceContext&gt;
+	&lt;SourceContext&gt;  695       	CurrentObjectBeingDebugged = NewObject;&lt;/SourceContext&gt;
+	&lt;SourceContext&gt;  696       &lt;/SourceContext&gt;
+	&lt;SourceContext&gt;  697       	// Register the new object&lt;/SourceContext&gt;
+	&lt;SourceContext&gt;  698       	if (NewObject != NULL)&lt;/SourceContext&gt;
+	&lt;SourceContext&gt;  699       	{&lt;/SourceContext&gt;
+	&lt;UserDescription&gt;Sent in the unattended mode&lt;/UserDescription&gt;
+	&lt;UserActivityHint /&gt;
+	&lt;ErrorMessage&gt;Ensure condition failed: NewObject-&amp;gt;IsA(this-&amp;gt;GeneratedClass) [File:D:\Build\++UE4+Release-4.12+Compile\Sync\Engine\Source\Runtime\Engine\Private\Blueprint.cpp] [Line: 689] &lt;/ErrorMessage&gt;
+	&lt;ErrorMessage&gt;Type mismatch: Expected Untitled_C, Found LevelScriptBlueprint&lt;/ErrorMessage&gt;
+	&lt;CrashGUID&gt;UE4CC-Windows-7188BC094DE45EE1A7D39EBF6FB434F5_0000&lt;/CrashGUID&gt;
+	&lt;TimeofCrash&gt;2016-07-13T14:58:43.402875Z&lt;/TimeofCrash&gt;
+	&lt;bHasMiniDump&gt;true&lt;/bHasMiniDump&gt;
+	&lt;bHasLog&gt;true&lt;/bHasLog&gt;
+	&lt;bHasDiags&gt;false&lt;/bHasDiags&gt;
+	&lt;bHasWERData&gt;false&lt;/bHasWERData&gt;
+	&lt;bHasVideo&gt;false&lt;/bHasVideo&gt;
+	&lt;bAllowToBeContacted&gt;true&lt;/bAllowToBeContacted&gt;
+&lt;/CrashDescription&gt;</value>
+      </setting>
+      <setting name="AWSS3CompressedSuffix" serializeAs="String">
+        <value>.gz</value>
+      </setting>
+      <setting name="CrashesCacheInMinutes" serializeAs="String">
+        <value>15</value>
+      </setting>
+      <setting name="SearchFilterCacheInMinutes" serializeAs="String">
+        <value>60</value>
+      </setting>
     </Tools.CrashReporter.CrashReportWebSite.Properties.Settings>
   </applicationSettings>
   <runtime>
diff --git a/Engine/Source/Programs/UnrealBuildTool/Configuration/ModuleRules.cs b/Engine/Source/Programs/UnrealBuildTool/Configuration/ModuleRules.cs
--- a/Engine/Source/Programs/UnrealBuildTool/Configuration/ModuleRules.cs
+++ b/Engine/Source/Programs/UnrealBuildTool/Configuration/ModuleRules.cs
@@ -316,11 +316,6 @@ namespace UnrealBuildTool
 		/// </summary>
 		public List<string> Definitions = new List<string>();
 
-		/// <summary>
-		/// CLR modules only: The assemblies referenced by the module's private implementation.
-		/// </summary>
-		public List<string> PrivateAssemblyReferences = new List<string>();
-
 		/// <summary>
 		/// Addition modules this module may require at run-time 
 		/// </summary>
diff --git a/Engine/Source/Programs/UnrealBuildTool/Configuration/TargetRules.cs b/Engine/Source/Programs/UnrealBuildTool/Configuration/TargetRules.cs
--- a/Engine/Source/Programs/UnrealBuildTool/Configuration/TargetRules.cs
+++ b/Engine/Source/Programs/UnrealBuildTool/Configuration/TargetRules.cs
@@ -211,60 +211,60 @@ namespace UnrealBuildTool
 		}
 
 		/// <summary>
-		/// Name of this target
+		/// The name of this target.
 		/// </summary>
 		public readonly string Name;
 
 		/// <summary>
-		/// The platform that the target is being built for
+		/// Platform that this target is being built for.
 		/// </summary>
 		public readonly UnrealTargetPlatform Platform;
 
 		/// <summary>
-		/// The configuration being built
+		/// The configuration being built.
 		/// </summary>
 		public readonly UnrealTargetConfiguration Configuration;
 
 		/// <summary>
-		/// Architecture that the target is being built for (or an empty string for the default)
+		/// Architecture that the target is being built for (or an empty string for the default).
 		/// </summary>
 		public readonly string Architecture;
 
 		/// <summary>
-		/// The path to the project containing this target
+		/// Path to the project file for the project containing this target.
 		/// </summary>
 		public readonly FileReference ProjectFile;
 
 		/// <summary>
-		/// Type of target
+		/// The type of target.
 		/// </summary>
 		public global::UnrealBuildTool.TargetType Type = global::UnrealBuildTool.TargetType.Game;
 
 		/// <summary>
-		/// Whether the target uses Steam (todo: substitute with more generic functionality)
+		/// Whether the target uses Steam.
 		/// </summary>
 		public bool bUsesSteam;
 
 		/// <summary>
-		/// Whether the target uses CEF3
+		/// Whether the target uses CEF3.
 		/// </summary>
 		public bool bUsesCEF3;
 
 		/// <summary>
-		/// Whether the project uses visual Slate UI (as opposed to the low level windowing/messaging which is always used)
+		/// Whether the project uses visual Slate UI (as opposed to the low level windowing/messaging, which is always available).
 		/// </summary>
 		public bool bUsesSlate = true;
 
 		/// <summary>
-		/// Forces linking against the static CRT. This is not supported across the engine due to the need for allocator implementations to be shared (for example), and TPS 
+		/// Forces linking against the static CRT. This is not fully supported across the engine due to the need for allocator implementations to be shared (for example), and TPS 
 		/// libraries to be consistent with each other, but can be used for utility programs.
 		/// </summary>
 		public bool bUseStaticCRT = false;
 
 		/// <summary>
-		/// By default we use the Release C++ Runtime (CRT), even when compiling Debug builds.  This is because the Debug C++
-		/// Runtime isn't very useful when debugging Unreal Engine projects, and linking against the Debug CRT libraries forces
-		/// our third party library dependencies to also be compiled using the Debug CRT (and often perform more slowly.)  Often
+		/// Enables the debug C++ runtime (CRT) for debug builds. By default we always use the release runtime, since the debug
+		/// version isn't particularly useful when debugging Unreal Engine projects, and linking against the debug CRT libraries forces
+		/// our third party library dependencies to also be compiled using the debug CRT (and often perform more slowly). Often
 		/// it can be inconvenient to require a separate copy of the debug versions of third party static libraries simply
 		/// so that you can debug your program's code.
 		/// </summary>
@@ -272,9 +272,8 @@ namespace UnrealBuildTool
 		public bool bDebugBuildsActuallyUseDebugCRT = false;
 
 		/// <summary>
-		/// Whether the output from this target can be publicly distributed, even if it has
-		/// dependencies on modules that are not (i.e. CarefullyRedist, NotForLicensees, NoRedist).
-		/// This should be used when you plan to release binaries but not source.
+		/// Whether the output from this target can be publicly distributed, even if it has dependencies on modules that are in folders 
+		/// with special restrictions (eg. CarefullyRedist, NotForLicensees, NoRedist).
 		/// </summary>
 		public bool bOutputPubliclyDistributable = false;
 
@@ -297,22 +296,22 @@ namespace UnrealBuildTool
 		public List<string> AdditionalPlugins = new List<string>();
 
 		/// <summary>
-		/// Path to the set of pak signing keys to embed in the executable
+		/// Path to the set of pak signing keys to embed in the executable.
 		/// </summary>
 		public string PakSigningKeysFile = "";
 
 		/// <summary>
-		/// Allows a Program Target to specify it's own solution folder path
+		/// Allows a Program Target to specify it's own solution folder path.
 		/// </summary>
 		public string SolutionDirectory = String.Empty;
 
 		/// <summary>
-		/// If true, the built target goes into the Engine/Binaries/{Platform} folder
+		/// Output the executable to the engine binaries folder.
 		/// </summary>
 		public bool bOutputToEngineBinaries = false;
 
 		/// <summary>
-		/// Sub folder where the built target goes: Engine/Binaries/{Platform}/{SubDir}
+		/// Subfolder to place executables in, relative to the default location.
 		/// </summary>
 		public string ExeBinariesSubFolder = String.Empty;
 
@@ -322,34 +321,34 @@ namespace UnrealBuildTool
 		public EGeneratedCodeVersion GeneratedCodeVersion = EGeneratedCodeVersion.None;
 
 		/// <summary>
-		/// Whether to include PhysX support
+		/// Whether to include PhysX support.
 		/// </summary>
 		public bool bCompilePhysX = true;
 
 		/// <summary>
-		/// Whether to include PhysX APEX support
+		/// Whether to include PhysX APEX support.
 		/// </summary>
 		[ConfigFile(ConfigHierarchyType.Engine, "/Script/BuildSettings.BuildSettings", "bCompileApex")]
 		public bool bCompileAPEX = true;
 
 		/// <summary>
-		/// Whether to include NvCloth
+		/// Whether to include NvCloth.
 		/// </summary>
 		public bool bCompileNvCloth = false;
 
 		/// <summary>
-		/// Whether to allow runtime cooking of physics
+		/// Whether to allow runtime cooking of physics.
 		/// </summary>
 		public bool bRuntimePhysicsCooking = true;
 
 		/// <summary>
-		/// Whether to include Box2D support
+		/// Whether to include Box2D support.
 		/// </summary>
 		[ConfigFile(ConfigHierarchyType.Engine, "/Script/BuildSettings.BuildSettings", "bCompileBox2D")]
 		public bool bCompileBox2D = true;
 
 		/// <summary>
-		/// Whether to include ICU unicode/i18n support in core
+		/// Whether to include ICU unicode/i18n support in Core.
 		/// </summary>
 		[ConfigFile(ConfigHierarchyType.Engine, "/Script/BuildSettings.BuildSettings", "bCompileICU")]
 		public bool bCompileICU = true;
@@ -361,7 +360,7 @@ namespace UnrealBuildTool
 		public bool bCompileCEF3 = true;
 
 		/// <summary>
-		/// Whether to compile the editor or not. Only desktop platforms (Windows or Mac) will use this, other platforms force this to false
+		/// Whether to compile the editor or not. Only desktop platforms (Windows or Mac) will use this, other platforms force this to false.
 		/// </summary>
 		[CommandLine("-NoEditor", Value = "false")]
 		public bool bBuildEditor = true;
@@ -372,7 +371,7 @@ namespace UnrealBuildTool
 		public bool bBuildRequiresCookedData = false;
 
 		/// <summary>
-		/// Whether to compile WITH_EDITORONLY_DATA disabled. Only Windows will use this, other platforms force this to false
+		/// Whether to compile WITH_EDITORONLY_DATA disabled. Only Windows will use this, other platforms force this to false.
 		/// </summary>
 		[CommandLine("-NoEditorOnlyData", Value = "false")]
 		public bool bBuildWithEditorOnlyData = true;
@@ -383,7 +382,7 @@ namespace UnrealBuildTool
 		public bool bBuildDeveloperTools = true;
 
 		/// <summary>
-		/// Whether to force compiling the target platform modules, even if they wouldn't normally be built
+		/// Whether to force compiling the target platform modules, even if they wouldn't normally be built.
 		/// </summary>
 		public bool bForceBuildTargetPlatforms = false;
 
@@ -422,13 +421,13 @@ namespace UnrealBuildTool
 		public bool bCompileAgainstCoreUObject = true;
 
 		/// <summary>
-		/// If true, include ADO database support in core
+		/// If true, include ADO database support in core.
 		/// </summary>
 		[ConfigFile(ConfigHierarchyType.Engine, "/Script/BuildSettings.BuildSettings", "bIncludeADO")]
 		public bool bIncludeADO;
 
 		/// <summary>
-		/// Whether to compile Recast navmesh generation
+		/// Whether to compile Recast navmesh generation.
 		/// </summary>
 		[ConfigFile(ConfigHierarchyType.Engine, "/Script/BuildSettings.BuildSettings", "bCompileRecast")]
 		public bool bCompileRecast = true;
@@ -440,12 +439,12 @@ namespace UnrealBuildTool
 		public bool bCompileSpeedTree = true;
 
 		/// <summary>
-		/// Enable exceptions for all modules
+		/// Enable exceptions for all modules.
 		/// </summary>
 		public bool bForceEnableExceptions = false;
 
 		/// <summary>
-		/// Enable RTTI for all modules
+		/// Enable RTTI for all modules.
 		/// </summary>
 		public bool bForceEnableRTTI = false;
 
@@ -455,50 +454,50 @@ namespace UnrealBuildTool
 		public bool bWithServerCode = true;
 
 		/// <summary>
-		/// Whether to include stats support even without the engine
+		/// Whether to include stats support even without the engine.
 		/// </summary>
 		public bool bCompileWithStatsWithoutEngine = false;
 
 		/// <summary>
-		/// Whether to include plugin support
+		/// Whether to include plugin support.
 		/// </summary>
 		[ConfigFile(ConfigHierarchyType.Engine, "/Script/BuildSettings.BuildSettings", "bCompileWithPluginSupport")]
 		public bool bCompileWithPluginSupport = false;
 
         /// <summary>
-        /// Whether to include PerfCounters support
+        /// Whether to include PerfCounters support.
         /// </summary>
 		[ConfigFile(ConfigHierarchyType.Engine, "/Script/BuildSettings.BuildSettings", "bWithPerfCounters")]
         public bool bWithPerfCounters = false;
 
         /// <summary>
-        /// Whether to turn on logging for test/shipping builds
+        /// Whether to turn on logging for test/shipping builds.
         /// </summary>
         public bool bUseLoggingInShipping = false;
 
 		/// <summary>
-		/// Whether to turn on logging to memory for test/shipping builds
+		/// Whether to turn on logging to memory for test/shipping builds.
 		/// </summary>
 		public bool bLoggingToMemoryEnabled;
 
 		/// <summary>
-		/// Whether to check that the process was launched through an external launcher
+		/// Whether to check that the process was launched through an external launcher.
 		/// </summary>
         public bool bUseLauncherChecks = false;
 
 		/// <summary>
-		/// Whether to turn on checks (asserts) for test/shipping builds
+		/// Whether to turn on checks (asserts) for test/shipping builds.
 		/// </summary>
 		public bool bUseChecksInShipping = false;
 
 		/// <summary>
-		/// True if we need FreeType support
+		/// True if we need FreeType support.
 		/// </summary>
 		[ConfigFile(ConfigHierarchyType.Engine, "/Script/BuildSettings.BuildSettings", "bCompileFreeType")]
 		public bool bCompileFreeType = true;
 
 		/// <summary>
-		/// True if we want to favor optimizing size over speed
+		/// True if we want to favor optimizing size over speed.
 		/// </summary>
 		[ConfigFile(ConfigHierarchyType.Engine, "/Script/BuildSettings.BuildSettings", "bCompileForSize")]
 		public bool bCompileForSize = false;
@@ -514,7 +513,7 @@ namespace UnrealBuildTool
         public bool bForceCompilePerformanceAutomationTests = false;
 
 		/// <summary>
-		/// If true, event driven loader will be used in cooked builds. @todoio This needs to be replaced by a runtime solution after async loading refactor
+		/// If true, event driven loader will be used in cooked builds. @todoio This needs to be replaced by a runtime solution after async loading refactor.
 		/// </summary>
 		public bool bEventDrivenLoader;
 
@@ -524,18 +523,18 @@ namespace UnrealBuildTool
 		public bool bEnforceIWYU = true;
 
 		/// <summary>
-		/// Whether the final executable should export symbols
+		/// Whether the final executable should export symbols.
 		/// </summary>
 		public bool bHasExports = true;
 
 		/// <summary>
-		/// Make static libraries for all engine modules as intermediates for this target
+		/// Make static libraries for all engine modules as intermediates for this target.
 		/// </summary>
 		[CommandLine("-Precompile")]
 		public bool bPrecompile = false;
 
 		/// <summary>
-		/// Use existing static libraries for all engine modules in this target
+		/// Use existing static libraries for all engine modules in this target.
 		/// </summary>
 		[CommandLine("-UsePrecompiled")]
 		public bool bUsePrecompiled = false;
@@ -546,7 +545,7 @@ namespace UnrealBuildTool
 		public bool bEnableOSX109Support = false;
 
 		/// <summary>
-		/// True if this is a console application that's being built
+		/// True if this is a console application that's being built.
 		/// </summary>
 		public bool bIsBuildingConsoleApplication = false;
 
@@ -563,22 +562,28 @@ namespace UnrealBuildTool
 		public bool bUseUnityBuild = true;
 
 		/// <summary>
-		/// Whether to _force_ unify C++ code into larger files for faster compilation.
+		/// Whether to force C++ source files to be combined into larger files for faster compilation.
 		/// </summary>
 		[CommandLine("-ForceUnity")]
 		[XmlConfigFile(Category = "BuildConfiguration")]
 		public bool bForceUnityBuild = false;
 
 		/// <summary>
-		/// An experimental new feature that, when enabled, will disable optimization of files that are omitted from Unity so they they become
-		/// easier to debug.
-		/// IMPORTANT: This feature is not complete yet!  Currently, it requires source files to be read-only that you aren't actively working with!
+		/// Use a heuristic to determine which files are currently being iterated on and exclude them from unity blobs, result in faster
+		/// incremental compile times. The current implementation uses the read-only flag to distinguish the working set, assuming that files will
+		/// be made writable by the source control system if they are being modified. This is true for Perforce, but not for Git.
+		/// </summary>
+		[XmlConfigFile(Category = "BuildConfiguration")]
+		public bool bUseAdaptiveUnityBuild = true;
+
+		/// <summary>
+		/// Disable optimization for files that are in the adaptive non-unity working set.
 		/// </summary>
 		[XmlConfigFile(Category = "BuildConfiguration")]
 		public bool bAdaptiveUnityDisablesOptimizations = false;
 
 		/// <summary>
-		/// Disables any force-included PCHs for files that are excluded from the unity build.
+		/// Disables force-included PCHs for files that are in the adaptive non-unity working set.
 		/// </summary>
 		[XmlConfigFile(Category = "BuildConfiguration")]
 		public bool bAdaptiveUnityDisablesPCH = false;
@@ -599,7 +604,7 @@ namespace UnrealBuildTool
 		public bool bShadowVariableErrors = false;
 
 		/// <summary>
-		/// Forces the use of undefined identifiers in conditional expressions to be treated as errors on platforms that support it.
+		/// Forces the use of undefined identifiers in conditional expressions to be treated as errors.
 		/// </summary>
 		[XmlConfigFile(Category = "BuildConfiguration")]
 		public bool bUndefinedIdentifierErrors = false;
@@ -633,7 +638,7 @@ namespace UnrealBuildTool
 		public bool bStressTestUnity = false;
 
 		/// <summary>
-		/// Whether to force debug info to be generated
+		/// Whether to force debug info to be generated.
 		/// </summary>
 		[CommandLine("-ForceDebugInfo")]
 		public bool bForceDebugInfo = false;
@@ -652,7 +657,7 @@ namespace UnrealBuildTool
 		public bool bDisableDebugInfoForGeneratedCode = true;
 
 		/// <summary>
-		/// Whether to disable debug info on PC in development builds (for faster developer iteration, as link times are extremely fast with debug info disabled.)
+		/// Whether to disable debug info on PC in development builds (for faster developer iteration, as link times are extremely fast with debug info disabled).
 		/// </summary>
 		[XmlConfigFile(Category = "BuildConfiguration")]
 		public bool bOmitPCDebugInfoInDevelopment = false;
@@ -681,20 +686,20 @@ namespace UnrealBuildTool
 		/// When enabled, a precompiled header is always generated for game modules, even if there are only a few source files
 		/// in the module.  This greatly improves compile times for iterative changes on a few files in the project, at the expense of slower
 		/// full rebuild times for small game projects.  This can be overridden by setting MinFilesUsingPrecompiledHeaderOverride in
-		/// a module's Build.cs file
+		/// a module's Build.cs file.
 		/// </summary>
 		[XmlConfigFile(Category = "BuildConfiguration")]
 		public bool bForcePrecompiledHeaderForGameModules = true;
 
 		/// <summary>
-		/// Whether to use incremental linking or not. Incremental linking can yield faster iteration times when making small changes
-		/// NOTE: We currently don't use incremental linking because it tends to behave a bit buggy on some computers (PDB-related compile errors)
+		/// Whether to use incremental linking or not. Incremental linking can yield faster iteration times when making small changes.
+		/// Currently disabled by default because it tends to behave a bit buggy on some computers (PDB-related compile errors).
 		/// </summary>
 		[XmlConfigFile(Category = "BuildConfiguration")]
 		public bool bUseIncrementalLinking = false;
 
 		/// <summary>
-		/// Whether to allow the use of LTCG (link time code generation) 
+		/// Whether to allow the use of link time code generation (LTCG).
 		/// </summary>
 		[CommandLine("-NoLTCG", Value = "false")]
 		[XmlConfigFile(Category = "BuildConfiguration")]
@@ -714,7 +719,7 @@ namespace UnrealBuildTool
 		public bool bSupportEditAndContinue = false;
 
 		/// <summary>
-		/// Whether to omit frame pointers or not. Disabling is useful for e.g. memory profiling on the PC
+		/// Whether to omit frame pointers or not. Disabling is useful for e.g. memory profiling on the PC.
 		/// </summary>
 		[XmlConfigFile(Category = "BuildConfiguration")]
 		public bool bOmitFramePointers = true;
@@ -726,14 +731,14 @@ namespace UnrealBuildTool
 		public bool bStripSymbolsOnIOS = false;
 
 		/// <summary>
-		/// If true, then enable memory profiling in the build (defines USE_MALLOC_PROFILER=1 and forces bOmitFramePointers=false)
+		/// If true, then enable memory profiling in the build (defines USE_MALLOC_PROFILER=1 and forces bOmitFramePointers=false).
 		/// </summary>
 		[XmlConfigFile(Category = "BuildConfiguration")]
 		public bool bUseMallocProfiler = false;
 
 		/// <summary>
 		/// Enables "Shared PCHs", a feature which significantly speeds up compile times by attempting to
-		/// share certain PCH files between modules that UBT detects is including those PCH's header files
+		/// share certain PCH files between modules that UBT detects is including those PCH's header files.
 		/// </summary>
 		[CommandLine("-NoSharedPCH", Value = "false")]
 		[XmlConfigFile(Category = "BuildConfiguration")]
@@ -746,7 +751,7 @@ namespace UnrealBuildTool
 		public bool bUseShippingPhysXLibraries = false;
 
         /// <summary>
-        /// True if Development and Release builds should use the checked configuration of PhysX/APEX. if bUseShippingPhysXLibraries is true this is ignored
+        /// True if Development and Release builds should use the checked configuration of PhysX/APEX. if bUseShippingPhysXLibraries is true this is ignored.
         /// </summary>
 		[XmlConfigFile(Category = "BuildConfiguration")]
         public bool bUseCheckedPhysXLibraries = false;
@@ -771,7 +776,7 @@ namespace UnrealBuildTool
 		public bool bUseFastPDBLinking = false;
 
 		/// <summary>
-		/// Whether to request the linker create a map file as part of the build
+		/// Outputs a map file as part of the build.
 		/// </summary>
 		[CommandLine("-MapFile")]
 		[XmlConfigFile(Category = "BuildConfiguration")]
@@ -780,14 +785,14 @@ namespace UnrealBuildTool
 		/// <summary>
 		/// Enables code analysis mode.  Currently, this has specific requirements.  It only works on Windows
 		/// platform with the MSVC compiler.  Also, it requires a version of the compiler that supports the
-		/// /analyze option, such as Visual Studio 2013
+		/// /analyze option, such as Visual Studio 2013.
 		/// </summary>
 		[CommandLine("-EnableCodeAnalysis")]
 		[XmlConfigFile(Category = "BuildConfiguration")]
 		public bool bEnableCodeAnalysis = false;
 
 		/// <summary>
-		/// Bundle version for Mac apps
+		/// Bundle version for Mac apps.
 		/// </summary>
 		[CommandLine("-BundleVersion")]
 		public string BundleVersion = null;
@@ -799,13 +804,13 @@ namespace UnrealBuildTool
 		public bool bDeployAfterCompile = false;
 
 		/// <summary>
-		/// If true, then a stub IPA will be generated when compiling is done (minimal files needed for a valid IPA)
+		/// If true, then a stub IPA will be generated when compiling is done (minimal files needed for a valid IPA).
 		/// </summary>
 		[CommandLine("-NoCreateStub", Value = "false")]
 		public bool bCreateStubIPA = true;
 
 		/// <summary>
-		/// If true, then a stub IPA will be generated when compiling is done (minimal files needed for a valid IPA)
+		/// If true, then a stub IPA will be generated when compiling is done (minimal files needed for a valid IPA).
 		/// </summary>
 		[CommandLine("-CopyAppBundleBackToDevice")]
 		public bool bCopyAppBundleBackToDevice = false;
@@ -822,7 +827,7 @@ namespace UnrealBuildTool
 		public bool bCheckSystemHeadersForModification;
 
 		/// <summary>
-		/// Whether to disable linking for this target
+		/// Whether to disable linking for this target.
 		/// </summary>
 		[CommandLine("-NoLink")]
 		public bool bDisableLinking = false;
@@ -836,23 +841,14 @@ namespace UnrealBuildTool
 		public bool bFormalBuild = false;
 
 		/// <summary>
-		/// An experimental new feature that, when enabled, will try to determine source files that you are actively iteratively changing, 
-		/// and break those files out of their unity blobs so that you can compile them as individual translation units, much faster than 
-		/// recompiling the entire unity blob each time.
-		/// IMPORTANT: This feature is not complete yet!  Currently, it requires source files to be read-only that you aren't actively working with!
-		/// </summary>
-		[XmlConfigFile(Category = "BuildConfiguration")]
-		public bool bUseAdaptiveUnityBuild = true;
-
-		/// <summary>
-		/// Whether to clean Builds directory on a remote Mac before building
+		/// Whether to clean Builds directory on a remote Mac before building.
 		/// </summary>
 		[CommandLine("-FlushMac")]
 		[XmlConfigFile(Category = "BuildConfiguration")]
 		public bool bFlushBuildDirOnRemoteMac = false;
 
 		/// <summary>
-		/// Whether to write detailed timing info from the compiler and linker
+		/// Whether to write detailed timing info from the compiler and linker.
 		/// </summary>
 		[CommandLine("-Timing")]
 		[XmlConfigFile(Category = "BuildConfiguration")]
@@ -885,10 +881,10 @@ namespace UnrealBuildTool
 		/// </summary>
 		[CommandLine("-Monolithic", Value ="Monolithic")]
 		[CommandLine("-Modular", Value ="Modular")]
-		public TargetLinkType LinkTypePrivate = TargetLinkType.Default;
+		TargetLinkType LinkTypePrivate = TargetLinkType.Default;
 
 		/// <summary>
-		/// Macros to define globally across the whole target
+		/// Macros to define globally across the whole target.
 		/// </summary>
 		[CommandLine("-Define", ValueAfterSpace = true)]
 		public List<string> GlobalDefinitions = new List<string>();
@@ -914,7 +910,7 @@ namespace UnrealBuildTool
 		private string LaunchModuleNamePrivate;
 
 		/// <summary>
-		/// List of additional modules to be compiled into the target
+		/// List of additional modules to be compiled into the target.
 		/// </summary>
 		public List<string> ExtraModuleNames = new List<string>();
 
@@ -924,27 +920,27 @@ namespace UnrealBuildTool
 		public TargetBuildEnvironment BuildEnvironment = TargetBuildEnvironment.Default;
 
 		/// <summary>
-		/// Android-specific target settings
+		/// Android-specific target settings.
 		/// </summary>
 		public AndroidTargetRules AndroidPlatform = new AndroidTargetRules();
 
 		/// <summary>
-		/// Mac-specific target settings
+		/// Mac-specific target settings.
 		/// </summary>
 		public MacTargetRules MacPlatform = new MacTargetRules();
 
 		/// <summary>
-		/// PS4-specific target settings
+		/// PS4-specific target settings.
 		/// </summary>
 		public PS4TargetRules PS4Platform = new PS4TargetRules();
 
 		/// <summary>
-		/// Windows-specific target settings
+		/// Windows-specific target settings.
 		/// </summary>
 		public WindowsTargetRules WindowsPlatform = new WindowsTargetRules();
 
 		/// <summary>
-		/// XboxOne-specific target settings
+		/// Xbox One-specific target settings.
 		/// </summary>
 		public XboxOneTargetRules XboxOnePlatform = new XboxOneTargetRules();
 
diff --git a/Engine/Source/Programs/UnrealBuildTool/Configuration/UEBuildTarget.cs b/Engine/Source/Programs/UnrealBuildTool/Configuration/UEBuildTarget.cs
--- a/Engine/Source/Programs/UnrealBuildTool/Configuration/UEBuildTarget.cs
+++ b/Engine/Source/Programs/UnrealBuildTool/Configuration/UEBuildTarget.cs
@@ -590,39 +590,37 @@ namespace UnrealBuildTool
 			if (PossibleTargetNames.Count > 0)
 			{
 				// We have possible targets!
-				foreach (string PossibleTargetName in PossibleTargetNames)
-				{
-					// If Engine is installed, the PossibleTargetName could contain a path
-					string TargetName = PossibleTargetName;
+				string PossibleTargetName = PossibleTargetNames[0];
 
-					// If a project file was not specified see if we can find one
-					if (ProjectFile == null && UProjectInfo.TryGetProjectForTarget(TargetName, out ProjectFile))
-					{
-						Log.TraceVerbose("Found project file for {0} - {1}", TargetName, ProjectFile);
-					}
+				// If Engine is installed, the PossibleTargetName could contain a path
+				string TargetName = PossibleTargetName;
 
-					UEBuildPlatform BuildPlatform = UEBuildPlatform.GetBuildPlatform(Platform);
+				// If a project file was not specified see if we can find one
+				if (ProjectFile == null && UProjectInfo.TryGetProjectForTarget(TargetName, out ProjectFile))
+				{
+					Log.TraceVerbose("Found project file for {0} - {1}", TargetName, ProjectFile);
+				}
 
-					if(Architecture == null)
-					{
-						Architecture = BuildPlatform.GetDefaultArchitecture(ProjectFile);
-					}
+				UEBuildPlatform BuildPlatform = UEBuildPlatform.GetBuildPlatform(Platform);
 
-					Targets.Add(new TargetDescriptor()
-						{
-							ProjectFile = ProjectFile,
-							TargetName = TargetName,
-							Platform = Platform,
-							Configuration = Configuration,
-							Architecture = Architecture,
-							bIsEditorRecompile = bIsEditorRecompile,
-							RemoteRoot = RemoteRoot,
-							OnlyModules = OnlyModules,
-							ForeignPlugins = ForeignPlugins,
-							ForceReceiptFileName = ForceReceiptFileName
-						});
-					break;
+				if(Architecture == null)
+				{
+					Architecture = BuildPlatform.GetDefaultArchitecture(ProjectFile);
 				}
+
+				Targets.Add(new TargetDescriptor()
+					{
+						ProjectFile = ProjectFile,
+						TargetName = TargetName,
+						Platform = Platform,
+						Configuration = Configuration,
+						Architecture = Architecture,
+						bIsEditorRecompile = bIsEditorRecompile,
+						RemoteRoot = RemoteRoot,
+						OnlyModules = OnlyModules,
+						ForeignPlugins = ForeignPlugins,
+						ForceReceiptFileName = ForceReceiptFileName
+					});
 			}
 			if (Targets.Count == 0)
 			{
@@ -2022,6 +2020,13 @@ namespace UnrealBuildTool
 						Receipt.PrecompiledBuildDependencies.Add(VariablePath);
 					}
 				}
+
+				// Also add the Shared Build Id File if it's been specified
+				if (SharedBuildIdFile != null)
+				{
+					string VariablePath = TargetReceipt.InsertPathVariables(SharedBuildIdFile, UnrealBuildTool.EngineDirectory, ProjectDirectory);
+					Receipt.PrecompiledBuildDependencies.Add(VariablePath);
+				}
 			}
 
 			// Prepare all the version manifests
@@ -2347,6 +2352,22 @@ namespace UnrealBuildTool
 
 				// Check for linking against modules prohibited by the EULA
 				CheckForEULAViolation();
+
+				// Check there aren't any engine binaries with dependencies on game modules. This can happen when game-specific plugins override engine plugins.
+				foreach(UEBuildModule Module in Modules.Values)
+				{
+					if(Module.Binary != null && Module.RulesFile.IsUnderDirectory(UnrealBuildTool.EngineDirectory))
+					{
+						foreach(UEBuildModule ReferencedModule in Module.GetDirectDependencyModules())
+						{
+							// Hard-code specific exceptions until these are properly fixed up
+							if(ReferencedModule.RulesFile != null && !ReferencedModule.RulesFile.IsUnderDirectory(UnrealBuildTool.EngineDirectory) && !IsLegacyEngineToGameReference(Platform, Module.Name, ReferencedModule.Name))
+							{
+								throw new BuildException("Engine module '{0}' should not depend on game module '{1}'", Module.Name, ReferencedModule.Name);
+							}
+						}
+					}
+				}
 			}
 
 			// Execute the pre-build steps
@@ -2499,6 +2520,7 @@ namespace UnrealBuildTool
 				if (Platform == UnrealTargetPlatform.Win32 || Platform == UnrealTargetPlatform.Win64)
 				{
 					FileItem VersionFile = FileItem.GetExistingItemByFileReference(FileReference.Combine(UnrealBuildTool.EngineSourceDirectory, "Runtime", "Core", "Resources", "Windows", "ModuleVersionResource.rc.inl"));
+					VersionFile.CachedIncludePaths = GlobalCompileEnvironment.IncludePaths;
 					CPPOutput VersionOutput = TargetToolChain.CompileRCFiles(GlobalCompileEnvironment, new List<FileItem> { VersionFile }, ActionGraph);
 					GlobalLinkEnvironment.CommonResourceFiles.AddRange(VersionOutput.ObjectFiles);
 
@@ -2508,6 +2530,7 @@ namespace UnrealBuildTool
 						DefaultResourceCompileEnvironment.Definitions.Add("ORIGINAL_FILE_NAME=\"UE4\"");
 
 						FileItem DefaultResourceFile = FileItem.GetExistingItemByFileReference(FileReference.Combine(UnrealBuildTool.EngineSourceDirectory, "Runtime", "Launch", "Resources", "Windows", "PCLaunch.rc"));
+						DefaultResourceFile.CachedIncludePaths = DefaultResourceCompileEnvironment.IncludePaths;
 						CPPOutput DefaultResourceOutput = TargetToolChain.CompileRCFiles(DefaultResourceCompileEnvironment, new List<FileItem> { DefaultResourceFile }, ActionGraph);
 
 						GlobalLinkEnvironment.DefaultResourceFiles.AddRange(DefaultResourceOutput.ObjectFiles);
@@ -2539,6 +2562,22 @@ namespace UnrealBuildTool
 			return ECompilationResult.Succeeded;
 		}
 
+		/// <summary>
+		/// Checks whether the reference from a game module
+		/// </summary>
+		/// <param name="Platform">The platform being built</param>
+		/// <param name="EngineModuleName">Name of the engine module</param>
+		/// <param name="GameModuleName">Name of the game module</param>
+		/// <returns>True if the reference is a known engine->game reference</returns>
+		static bool IsLegacyEngineToGameReference(UnrealTargetPlatform Platform, string EngineModuleName, string GameModuleName)
+		{
+			if(Platform == UnrealTargetPlatform.PS4 && EngineModuleName == "WebBrowser" && GameModuleName == "OnlineSubsystem")
+			{
+				return true;
+			}
+			return false;
+		}
+
 		/// <summary>
 		/// Export the definition of this target to a JSON file
 		/// </summary>
@@ -3656,7 +3695,7 @@ namespace UnrealBuildTool
 			// If we're compiling against the engine, add the plugins enabled for this target
 			if (Rules.bCompileAgainstEngine)
 			{
-				ProjectDescriptor Project = (Rules.bCompileAgainstEngine && ProjectFile != null) ? ProjectDescriptor.FromFile(ProjectFile.FullName) : null;
+				ProjectDescriptor Project = (ProjectFile != null) ? ProjectDescriptor.FromFile(ProjectFile.FullName) : null;
 				foreach (PluginInfo ValidPlugin in ValidPlugins)
 				{
 					if(UProjectInfo.IsPluginEnabledForProject(ValidPlugin, Project, Platform, TargetType))
@@ -4161,11 +4200,6 @@ namespace UnrealBuildTool
 			// Validate rules object
 			if (RulesObject.Type == ModuleRules.ModuleType.CPlusPlus)
 			{
-				if (RulesObject.PrivateAssemblyReferences.Count > 0)
-				{
-					throw new BuildException("Module rules for '{0}' may not specify PrivateAssemblyReferences unless it is a CPlusPlusCLR module type.", ModuleName);
-				}
-
 				List<string> InvalidDependencies = RulesObject.DynamicallyLoadedModuleNames.Intersect(RulesObject.PublicDependencyModuleNames.Concat(RulesObject.PrivateDependencyModuleNames)).ToList();
 				if (InvalidDependencies.Count != 0)
 				{
@@ -4191,8 +4225,17 @@ namespace UnrealBuildTool
 					}
 					else
 					{
-						// Game module.  Do not enable shared PCHs by default, because games usually have a large precompiled header of their own and compile times would suffer.
-						RulesObject.PCHUsage = ModuleRules.PCHUsageMode.NoSharedPCHs;
+						PluginInfo Plugin;
+						if(RulesAssembly.TryGetPluginForModule(ModuleFileName, out Plugin))
+						{
+							// Game plugin.  Enable shared PCHs by default, since they aren't typically large enough to warrant their own PCH.
+							RulesObject.PCHUsage = ModuleRules.PCHUsageMode.UseSharedPCHs;
+						}
+						else
+						{
+							// Game module.  Do not enable shared PCHs by default, because games usually have a large precompiled header of their own and compile times would suffer.
+							RulesObject.PCHUsage = ModuleRules.PCHUsageMode.NoSharedPCHs;
+						}
 					}
 				}
 			}
diff --git a/Engine/Source/Programs/UnrealBuildTool/HTML5/UEBuildHTML5.cs b/Engine/Source/Programs/UnrealBuildTool/HTML5/UEBuildHTML5.cs
--- a/Engine/Source/Programs/UnrealBuildTool/HTML5/UEBuildHTML5.cs
+++ b/Engine/Source/Programs/UnrealBuildTool/HTML5/UEBuildHTML5.cs
@@ -10,7 +10,9 @@ namespace UnrealBuildTool
 {
 	class HTML5Platform : UEBuildPlatform
 	{
-		// use -win32 for win32 builds. ( build html5 platform as a win32 binary for debugging )
+		/// <summary>
+		/// Architecture to build for. Use -win32 for win32 builds (build html5 platform as a win32 binary for debugging).
+		/// </summary>
 		[XmlConfigFile]
 		public static string HTML5Architecture = "";
 
diff --git a/Engine/Source/Programs/UnrealBuildTool/IOS/UEBuildIOS.cs b/Engine/Source/Programs/UnrealBuildTool/IOS/UEBuildIOS.cs
--- a/Engine/Source/Programs/UnrealBuildTool/IOS/UEBuildIOS.cs
+++ b/Engine/Source/Programs/UnrealBuildTool/IOS/UEBuildIOS.cs
@@ -157,11 +157,11 @@ namespace UnrealBuildTool
 		{
 			get
 			{
-				if(bShipForArmV7 || (!bShipForArmV7 && !bShipForArm64 && !bShipForArmV7S))
+				if(bShipForArmV7 || (!bShipForArm64 && !bShipForArmV7S))
 				{
 					yield return "armv7";
 				}
-				if(bShipForArm64 || (!bShipForArmV7 && !bShipForArm64 && !bShipForArmV7S))
+				if(bShipForArm64 || (!bShipForArmV7 && !bShipForArmV7S))
 				{
 					yield return "arm64";
 				}
diff --git a/Engine/Source/Programs/UnrealBuildTool/System/DependencyCache.cs b/Engine/Source/Programs/UnrealBuildTool/System/DependencyCache.cs
--- a/Engine/Source/Programs/UnrealBuildTool/System/DependencyCache.cs
+++ b/Engine/Source/Programs/UnrealBuildTool/System/DependencyCache.cs
@@ -53,7 +53,7 @@ namespace UnrealBuildTool
 		/// <summary>
 		/// The version number for binary serialization
 		/// </summary>
-		const int FileVersion = 1;
+		const int FileVersion = 2;
 
 		/// <summary>
 		/// The file signature for binary serialization
diff --git a/Engine/Source/Programs/UnrealBuildTool/System/DirectoryReference.cs b/Engine/Source/Programs/UnrealBuildTool/System/DirectoryReference.cs
new file mode 100644
--- /dev/null
+++ b/Engine/Source/Programs/UnrealBuildTool/System/DirectoryReference.cs
@@ -0,0 +1,393 @@
+﻿using System;
+using System.Collections.Generic;
+using System.IO;
+using System.Linq;
+using System.Text;
+using System.Threading.Tasks;
+
+namespace UnrealBuildTool
+{
+	/// <summary>
+	/// Representation of an absolute directory path. Allows fast hashing and comparisons.
+	/// </summary>
+	[Serializable]
+	public class DirectoryReference : FileSystemReference, IEquatable<DirectoryReference>
+	{
+		/// <summary>
+		/// Default constructor.
+		/// </summary>
+		/// <param name="InPath">Path to this directory.</param>
+		public DirectoryReference(string InPath)
+			: base(FixTrailingPathSeparator(Path.GetFullPath(InPath)))
+		{
+		}
+
+		/// <summary>
+		/// Construct a DirectoryReference from a DirectoryInfo object.
+		/// </summary>
+		/// <param name="InInfo">Path to this file</param>
+		public DirectoryReference(DirectoryInfo InInfo)
+			: base(FixTrailingPathSeparator(InInfo.FullName))
+		{
+		}
+
+		/// <summary>
+		/// Constructor for creating a directory object directly from two strings.
+		/// </summary>
+		/// <param name="InFullName"></param>
+		/// <param name="InCanonicalName"></param>
+		protected DirectoryReference(string InFullName, string InCanonicalName)
+			: base(InFullName, InCanonicalName)
+		{
+		}
+
+		/// <summary>
+		/// Ensures that the correct trailing path separator is appended. On Windows, the root directory (eg. C:\) always has a trailing path separator, but no other
+		/// path does.
+		/// </summary>
+		/// <param name="DirName">Absolute path to the directory</param>
+		/// <returns>Path to the directory, with the correct trailing path separator</returns>
+		private static string FixTrailingPathSeparator(string DirName)
+		{
+			if(DirName.Length == 2 && DirName[1] == ':')
+			{
+				return DirName + Path.DirectorySeparatorChar;
+			}
+			else if(DirName.Length == 3 && DirName[1] == ':' && DirName[2] == Path.DirectorySeparatorChar)
+			{
+				return DirName;
+			}
+			else if(DirName.Length > 1 && DirName[DirName.Length - 1] == Path.DirectorySeparatorChar)
+			{
+				return DirName.TrimEnd(Path.DirectorySeparatorChar);
+			}
+			else
+			{
+				return DirName;
+			}
+		}
+
+		/// <summary>
+		/// Gets the top level directory name
+		/// </summary>
+		/// <returns>The name of the directory</returns>
+		public string GetDirectoryName()
+		{
+			return Path.GetFileName(FullName);
+		}
+
+		/// <summary>
+		/// Gets the directory containing this object
+		/// </summary>
+		/// <returns>A new directory object representing the directory containing this object</returns>
+		public DirectoryReference ParentDirectory
+		{
+			get
+			{
+				if (IsRootDirectory())
+				{
+					return null;
+				}
+
+				int ParentLength = CanonicalName.LastIndexOf(Path.DirectorySeparatorChar);
+				if (ParentLength == 2 && CanonicalName[1] == ':')
+				{
+					ParentLength++;
+				}
+
+				return new DirectoryReference(FullName.Substring(0, ParentLength), CanonicalName.Substring(0, ParentLength));
+			}
+		}
+
+		/// <summary>
+		/// Gets the parent directory for a file
+		/// </summary>
+		/// <param name="File">The file to get directory for</param>
+		/// <returns>The full directory name containing the given file</returns>
+		public static DirectoryReference GetParentDirectory(FileReference File)
+		{
+			int ParentLength = File.CanonicalName.LastIndexOf(Path.DirectorySeparatorChar);
+			return new DirectoryReference(File.FullName.Substring(0, ParentLength), File.CanonicalName.Substring(0, ParentLength));
+		}
+
+		/// <summary>
+		/// Determines whether this path represents a root directory in the filesystem
+		/// </summary>
+		/// <returns>True if this path is a root directory, false otherwise</returns>
+		public bool IsRootDirectory()
+		{
+			return CanonicalName[CanonicalName.Length - 1] == Path.DirectorySeparatorChar;
+		}
+
+		/// <summary>
+		/// Combine several fragments with a base directory, to form a new directory name
+		/// </summary>
+		/// <param name="BaseDirectory">The base directory</param>
+		/// <param name="Fragments">Fragments to combine with the base directory</param>
+		/// <returns>The new directory name</returns>
+		public static DirectoryReference Combine(DirectoryReference BaseDirectory, params string[] Fragments)
+		{
+			string FullName = FileSystemReference.CombineStrings(BaseDirectory, Fragments);
+			return new DirectoryReference(FullName, FullName.ToLowerInvariant());
+		}
+
+		/// <summary>
+		/// Compares two filesystem object names for equality. Uses the canonical name representation, not the display name representation.
+		/// </summary>
+		/// <param name="A">First object to compare.</param>
+		/// <param name="B">Second object to compare.</param>
+		/// <returns>True if the names represent the same object, false otherwise</returns>
+		public static bool operator ==(DirectoryReference A, DirectoryReference B)
+		{
+			if ((object)A == null)
+			{
+				return (object)B == null;
+			}
+			else
+			{
+				return (object)B != null && A.CanonicalName == B.CanonicalName;
+			}
+		}
+
+		/// <summary>
+		/// Compares two filesystem object names for inequality. Uses the canonical name representation, not the display name representation.
+		/// </summary>
+		/// <param name="A">First object to compare.</param>
+		/// <param name="B">Second object to compare.</param>
+		/// <returns>False if the names represent the same object, true otherwise</returns>
+		public static bool operator !=(DirectoryReference A, DirectoryReference B)
+		{
+			return !(A == B);
+		}
+
+		/// <summary>
+		/// Compares against another object for equality.
+		/// </summary>
+		/// <param name="Obj">other instance to compare.</param>
+		/// <returns>True if the names represent the same object, false otherwise</returns>
+		public override bool Equals(object Obj)
+		{
+			return (Obj is DirectoryReference) && ((DirectoryReference)Obj) == this;
+		}
+
+		/// <summary>
+		/// Compares against another object for equality.
+		/// </summary>
+		/// <param name="Obj">other instance to compare.</param>
+		/// <returns>True if the names represent the same object, false otherwise</returns>
+		public bool Equals(DirectoryReference Obj)
+		{
+			return Obj == this;
+		}
+
+		/// <summary>
+		/// Returns a hash code for this object
+		/// </summary>
+		/// <returns></returns>
+		public override int GetHashCode()
+		{
+			return CanonicalName.GetHashCode();
+		}
+
+		/// <summary>
+		/// Helper function to create a remote directory reference. Unlike normal DirectoryReference objects, these aren't converted to a full path in the local filesystem.
+		/// </summary>
+		/// <param name="AbsolutePath">The absolute path in the remote file system</param>
+		/// <returns>New directory reference</returns>
+		public static DirectoryReference MakeRemote(string AbsolutePath)
+		{
+			return new DirectoryReference(AbsolutePath, AbsolutePath.ToLowerInvariant());
+		}
+
+		/// <summary>
+		/// Helper function to create a directory reference from a raw platform path. The path provided *MUST* be exactly the same as that returned by Path.GetFullPath(). 
+		/// </summary>
+		/// <param name="AbsolutePath">The absolute path in the file system</param>
+		/// <returns>New file reference</returns>
+		public static DirectoryReference MakeFromNormalizedFullPath(string AbsolutePath)
+		{
+			return new DirectoryReference(AbsolutePath, AbsolutePath.ToLowerInvariant());
+		}
+
+		/// <summary>
+		/// Gets the parent directory for a file, or returns null if it's null.
+		/// </summary>
+		/// <param name="File">The file to create a directory reference for</param>
+		/// <returns>The directory containing the file  </returns>
+		public static DirectoryReference FromFile(FileReference File)
+		{
+			return (File == null)? null : File.Directory;
+		}
+
+		#region System.IO.Directory Wrapper Methods
+
+		/// <summary>
+		/// Finds the current directory
+		/// </summary>
+		/// <returns>The current directory</returns>
+		public static DirectoryReference GetCurrentDirectory()
+		{
+			return new DirectoryReference(Directory.GetCurrentDirectory());
+		}
+
+		/// <summary>
+		/// Creates a directory
+		/// </summary>
+		/// <param name="Location">Location of the directory</param>
+		public static void CreateDirectory(DirectoryReference Location)
+		{
+			Directory.CreateDirectory(Location.FullName);
+		}
+
+        /// <summary>
+        /// Deletes a directory
+        /// </summary>
+		/// <param name="Location">Location of the directory</param>
+        public static void Delete(DirectoryReference Location)
+        {
+            Directory.Delete(Location.FullName);
+        }
+
+        /// <summary>
+        /// Deletes a directory
+        /// </summary>
+		/// <param name="Location">Location of the directory</param>
+		/// <param name="bRecursive">Whether to remove directories recursively</param>
+        public static void Delete(DirectoryReference Location, bool bRecursive)
+        {
+            Directory.Delete(Location.FullName, bRecursive);
+        }
+
+        /// <summary>
+        /// Checks whether the directory exists
+        /// </summary>
+		/// <param name="Location">Location of the directory</param>
+        /// <returns>True if this directory exists</returns>
+        public static bool Exists(DirectoryReference Location)
+		{
+			return Directory.Exists(Location.FullName);
+		}
+
+		/// <summary>
+		/// Enumerate files from a given directory
+		/// </summary>
+		/// <param name="BaseDir">Base directory to search in</param>
+		/// <returns>Sequence of file references</returns>
+		public static IEnumerable<FileReference> EnumerateFiles(DirectoryReference BaseDir)
+		{
+			foreach (string FileName in Directory.EnumerateFiles(BaseDir.FullName))
+			{
+				yield return FileReference.MakeFromNormalizedFullPath(FileName);
+			}
+		}
+
+		/// <summary>
+		/// Enumerate files from a given directory
+		/// </summary>
+		/// <param name="BaseDir">Base directory to search in</param>
+		/// <param name="Pattern">Pattern for matching files</param>
+		/// <returns>Sequence of file references</returns>
+		public static IEnumerable<FileReference> EnumerateFiles(DirectoryReference BaseDir, string Pattern)
+		{
+			foreach (string FileName in Directory.EnumerateFiles(BaseDir.FullName, Pattern))
+			{
+				yield return FileReference.MakeFromNormalizedFullPath(FileName);
+			}
+		}
+
+		/// <summary>
+		/// Enumerate files from a given directory
+		/// </summary>
+		/// <param name="BaseDir">Base directory to search in</param>
+		/// <param name="Pattern">Pattern for matching files</param>
+		/// <param name="Option">Options for the search</param>
+		/// <returns>Sequence of file references</returns>
+		public static IEnumerable<FileReference> EnumerateFiles(DirectoryReference BaseDir, string Pattern, SearchOption Option)
+		{
+			foreach (string FileName in Directory.EnumerateFiles(BaseDir.FullName, Pattern, Option))
+			{
+				yield return FileReference.MakeFromNormalizedFullPath(FileName);
+			}
+		}
+
+		/// <summary>
+		/// Enumerate subdirectories in a given directory
+		/// </summary>
+		/// <param name="BaseDir">Base directory to search in</param>
+		/// <returns>Sequence of directory references</returns>
+		public static IEnumerable<DirectoryReference> EnumerateDirectories(DirectoryReference BaseDir)
+		{
+			foreach (string DirectoryName in Directory.EnumerateDirectories(BaseDir.FullName))
+			{
+				yield return DirectoryReference.MakeFromNormalizedFullPath(DirectoryName);
+			}
+		}
+
+		/// <summary>
+		/// Enumerate subdirectories in a given directory
+		/// </summary>
+		/// <param name="BaseDir">Base directory to search in</param>
+		/// <param name="Pattern">Pattern for matching directories</param>
+		/// <returns>Sequence of directory references</returns>
+		public static IEnumerable<DirectoryReference> EnumerateDirectories(DirectoryReference BaseDir, string Pattern)
+		{
+			foreach (string DirectoryName in Directory.EnumerateDirectories(BaseDir.FullName, Pattern))
+			{
+				yield return DirectoryReference.MakeFromNormalizedFullPath(DirectoryName);
+			}
+		}
+
+		/// <summary>
+		/// Enumerate subdirectories in a given directory
+		/// </summary>
+		/// <param name="BaseDir">Base directory to search in</param>
+		/// <param name="Pattern">Pattern for matching files</param>
+		/// <param name="Option">Options for the search</param>
+		/// <returns>Sequence of directory references</returns>
+		public static IEnumerable<DirectoryReference> EnumerateDirectories(DirectoryReference BaseDir, string Pattern, SearchOption Option)
+		{
+			foreach (string DirectoryName in Directory.EnumerateDirectories(BaseDir.FullName, Pattern, Option))
+			{
+				yield return DirectoryReference.MakeFromNormalizedFullPath(DirectoryName);
+			}
+		}
+
+		/// <summary>
+		/// Sets the current directory
+		/// </summary>
+		/// <param name="Location">Location of the new current directory</param>
+		public static void SetCurrentDirectory(DirectoryReference Location)
+		{
+			Directory.SetCurrentDirectory(Location.FullName);
+		}
+
+		#endregion
+	}
+
+	/// <summary>
+	/// Extension methods for passing DirectoryReference arguments
+	/// </summary>
+	static class DirectoryReferenceExtensionMethods
+	{
+		/// <summary>
+		/// Manually serialize a file reference to a binary stream.
+		/// </summary>
+		/// <param name="Writer">Binary writer to write to</param>
+		/// <param name="Directory">The directory reference to write</param>
+		public static void Write(this BinaryWriter Writer, DirectoryReference Directory)
+		{
+			Writer.Write((Directory == null) ? String.Empty : Directory.FullName);
+		}
+
+		/// <summary>
+		/// Manually deserialize a directory reference from a binary stream.
+		/// </summary>
+		/// <param name="Reader">Binary reader to read from</param>
+		/// <returns>New DirectoryReference object</returns>
+		public static DirectoryReference ReadDirectoryReference(this BinaryReader Reader)
+		{
+			string FullName = Reader.ReadString();
+			return (FullName.Length == 0) ? null : DirectoryReference.MakeFromNormalizedFullPath(FullName);
+		}
+	}
+}
diff --git a/Engine/Source/Programs/UnrealBuildTool/System/FileFilter.cs b/Engine/Source/Programs/UnrealBuildTool/System/FileFilter.cs
--- a/Engine/Source/Programs/UnrealBuildTool/System/FileFilter.cs
+++ b/Engine/Source/Programs/UnrealBuildTool/System/FileFilter.cs
@@ -201,10 +201,10 @@ namespace UnrealBuildTool
 		/// <param name="FileName"></param>
 		/// <param name="SectionName"></param>
 		/// <param name="AllowTags"></param>
-		public void ReadRulesFromFile(string FileName, string SectionName, params string[] AllowTags)
+		public void ReadRulesFromFile(FileReference FileName, string SectionName, params string[] AllowTags)
 		{
 			bool bInSection = false;
-			foreach(string Line in File.ReadAllLines(FileName))
+			foreach(string Line in File.ReadAllLines(FileName.FullName))
 			{
 				string TrimLine = Line.Trim();
 				if(!TrimLine.StartsWith(";") && TrimLine.Length > 0)
@@ -402,19 +402,6 @@ namespace UnrealBuildTool
 			return FileNames.Where(x => Matches(x.MakeRelativeTo(BaseDirectory)));
 		}
 
-		/// <summary>
-		/// Finds a list of files within a given directory which match the filter.
-		/// </summary>
-		/// <param name="DirectoryName">File to match</param>
-		/// <param name="bIgnoreSymlinks">Whether to ignore symlinks in the output</param>
-		/// <returns>List of files that pass the filter</returns>
-		public List<string> ApplyToDirectory(string DirectoryName, bool bIgnoreSymlinks)
-		{
-			List<string> MatchingFileNames = new List<string>();
-			FindMatchesFromDirectory(new DirectoryInfo(DirectoryName), "", bIgnoreSymlinks, MatchingFileNames);
-			return MatchingFileNames;
-		}
-
 		/// <summary>
 		/// Finds a list of files within a given directory which match the filter.
 		/// </summary>
@@ -461,34 +448,6 @@ namespace UnrealBuildTool
 			return Result;
 		}
 
-		/// <summary>
-		/// Finds a list of files within a given directory which match the filter.
-		/// </summary>
-		/// <param name="CurrentDirectory"></param>
-		/// <param name="NamePrefix"></param>
-		/// <param name="bIgnoreSymlinks">Whether to ignore symlinks in the output</param>
-		/// <param name="MatchingFileNames">Receives a list of matching filenames</param>
-		/// <returns>True if the file passes the filter</returns>
-		void FindMatchesFromDirectory(DirectoryInfo CurrentDirectory, string NamePrefix, bool bIgnoreSymlinks, List<string> MatchingFileNames)
-		{
-			foreach (FileInfo NextFile in CurrentDirectory.EnumerateFiles())
-			{
-				string FileName = NamePrefix + NextFile.Name;
-				if (Matches(FileName) && (!bIgnoreSymlinks || !NextFile.Attributes.HasFlag(FileAttributes.ReparsePoint)))
-				{
-					MatchingFileNames.Add(FileName);
-				}
-			}
-			foreach (DirectoryInfo NextDirectory in CurrentDirectory.EnumerateDirectories())
-			{
-				string NextNamePrefix = NamePrefix + NextDirectory.Name;
-				if (PossiblyMatches(NextNamePrefix))
-				{
-					FindMatchesFromDirectory(NextDirectory, NextNamePrefix + "/", bIgnoreSymlinks, MatchingFileNames);
-				}
-			}
-		}
-
 		/// <summary>
 		/// Finds a list of files within a given directory which match the filter.
 		/// </summary>
