Copying //UE4/Dev-Core to //UE4/Dev-Main (Source: //UE4/Dev-Core @ 3345728)

diff --git a/Engine/Source/Programs/UnrealBuildTool/System/FileReference.cs b/Engine/Source/Programs/UnrealBuildTool/System/FileReference.cs
new file mode 100644
--- /dev/null
+++ b/Engine/Source/Programs/UnrealBuildTool/System/FileReference.cs
@@ -0,0 +1,547 @@
+﻿using System;
+using System.Collections.Generic;
+using System.IO;
+using System.Linq;
+using System.Text;
+using System.Threading.Tasks;
+
+namespace UnrealBuildTool
+{
+	/// <summary>
+	/// Representation of an absolute file path. Allows fast hashing and comparisons.
+	/// </summary>
+	[Serializable]
+	public class FileReference : FileSystemReference, IEquatable<FileReference>
+	{
+		/// <summary>
+		/// Default constructor.
+		/// </summary>
+		/// <param name="InPath">Path to this file</param>
+		public FileReference(string InPath)
+			: base(Path.GetFullPath(InPath))
+		{
+			if(FullName.EndsWith("\\") || FullName.EndsWith("/"))
+			{
+				throw new ArgumentException("File names may not be terminated by a path separator character");
+			}
+		}
+
+		/// <summary>
+		/// Construct a FileReference from a FileInfo object.
+		/// </summary>
+		/// <param name="InInfo">Path to this file</param>
+		public FileReference(FileInfo InInfo)
+			: base(InInfo.FullName)
+		{
+		}
+
+		/// <summary>
+		/// Default constructor.
+		/// </summary>
+		protected FileReference(string InFullName, string InCanonicalName)
+			: base(InFullName, InCanonicalName)
+		{
+		}
+
+		/// <summary>
+		/// Gets the file name without path information
+		/// </summary>
+		/// <returns>A string containing the file name</returns>
+		public string GetFileName()
+		{
+			return Path.GetFileName(FullName);
+		}
+
+		/// <summary>
+		/// Gets the file name without path information or an extension
+		/// </summary>
+		/// <returns>A string containing the file name without an extension</returns>
+		public string GetFileNameWithoutExtension()
+		{
+			return Path.GetFileNameWithoutExtension(FullName);
+		}
+
+		/// <summary>
+		/// Gets the file name without path or any extensions
+		/// </summary>
+		/// <returns>A string containing the file name without an extension</returns>
+		public string GetFileNameWithoutAnyExtensions()
+		{
+			int StartIdx = FullName.LastIndexOf(Path.DirectorySeparatorChar) + 1;
+
+			int EndIdx = FullName.IndexOf('.', StartIdx);
+			if (EndIdx < StartIdx)
+			{
+				return FullName.Substring(StartIdx);
+			}
+			else
+			{
+				return FullName.Substring(StartIdx, EndIdx - StartIdx);
+			}
+		}
+
+		/// <summary>
+		/// Gets the extension for this filename
+		/// </summary>
+		/// <returns>A string containing the extension of this filename</returns>
+		public string GetExtension()
+		{
+			return Path.GetExtension(FullName);
+		}
+
+		/// <summary>
+		/// Change the file's extension to something else
+		/// </summary>
+		/// <param name="Extension">The new extension</param>
+		/// <returns>A FileReference with the same path and name, but with the new extension</returns>
+		public FileReference ChangeExtension(string Extension)
+		{
+			string NewFullName = Path.ChangeExtension(FullName, Extension);
+			return new FileReference(NewFullName, NewFullName.ToLowerInvariant());
+		}
+
+		/// <summary>
+		/// Gets the directory containing this file
+		/// </summary>
+		/// <returns>A new directory object representing the directory containing this object</returns>
+		public DirectoryReference Directory
+		{
+			get { return DirectoryReference.GetParentDirectory(this); }
+		}
+
+		/// <summary>
+		/// Combine several fragments with a base directory, to form a new filename
+		/// </summary>
+		/// <param name="BaseDirectory">The base directory</param>
+		/// <param name="Fragments">Fragments to combine with the base directory</param>
+		/// <returns>The new file name</returns>
+		public static FileReference Combine(DirectoryReference BaseDirectory, params string[] Fragments)
+		{
+			string FullName = FileSystemReference.CombineStrings(BaseDirectory, Fragments);
+			return new FileReference(FullName, FullName.ToLowerInvariant());
+		}
+
+		/// <summary>
+		/// Append a string to the end of a filename
+		/// </summary>
+		/// <param name="A">The base file reference</param>
+		/// <param name="B">Suffix to be appended</param>
+		/// <returns>The new file reference</returns>
+		public static FileReference operator +(FileReference A, string B)
+		{
+			return new FileReference(A.FullName + B, A.CanonicalName + B.ToLowerInvariant());
+		}
+
+		/// <summary>
+		/// Compares two filesystem object names for equality. Uses the canonical name representation, not the display name representation.
+		/// </summary>
+		/// <param name="A">First object to compare.</param>
+		/// <param name="B">Second object to compare.</param>
+		/// <returns>True if the names represent the same object, false otherwise</returns>
+		public static bool operator ==(FileReference A, FileReference B)
+		{
+			if ((object)A == null)
+			{
+				return (object)B == null;
+			}
+			else
+			{
+				return (object)B != null && A.CanonicalName == B.CanonicalName;
+			}
+		}
+
+		/// <summary>
+		/// Compares two filesystem object names for inequality. Uses the canonical name representation, not the display name representation.
+		/// </summary>
+		/// <param name="A">First object to compare.</param>
+		/// <param name="B">Second object to compare.</param>
+		/// <returns>False if the names represent the same object, true otherwise</returns>
+		public static bool operator !=(FileReference A, FileReference B)
+		{
+			return !(A == B);
+		}
+
+		/// <summary>
+		/// Compares against another object for equality.
+		/// </summary>
+		/// <param name="Obj">other instance to compare.</param>
+		/// <returns>True if the names represent the same object, false otherwise</returns>
+		public override bool Equals(object Obj)
+		{
+			return (Obj is FileReference) && ((FileReference)Obj) == this;
+		}
+
+		/// <summary>
+		/// Compares against another object for equality.
+		/// </summary>
+		/// <param name="Obj">other instance to compare.</param>
+		/// <returns>True if the names represent the same object, false otherwise</returns>
+		public bool Equals(FileReference Obj)
+		{
+			return Obj == this;
+		}
+
+		/// <summary>
+		/// Returns a hash code for this object
+		/// </summary>
+		/// <returns></returns>
+		public override int GetHashCode()
+		{
+			return CanonicalName.GetHashCode();
+		}
+
+		/// <summary>
+		/// Helper function to create a remote file reference. Unlike normal FileReference objects, these aren't converted to a full path in the local filesystem, but are
+		/// left as they are passed in.
+		/// </summary>
+		/// <param name="AbsolutePath">The absolute path in the remote file system</param>
+		/// <returns>New file reference</returns>
+		public static FileReference MakeRemote(string AbsolutePath)
+		{
+			return new FileReference(AbsolutePath, AbsolutePath.ToLowerInvariant());
+		}
+
+		/// <summary>
+		/// Helper function to create a file reference from a raw platform path. The path provided *MUST* be exactly the same as that returned by Path.GetFullPath(). 
+		/// </summary>
+		/// <param name="AbsolutePath">The absolute path in the file system</param>
+		/// <returns>New file reference</returns>
+		public static FileReference MakeFromNormalizedFullPath(string AbsolutePath)
+		{
+			return new FileReference(AbsolutePath, AbsolutePath.ToLowerInvariant());
+		}
+
+		/// <summary>
+		/// Makes a file location writeable; 
+		/// </summary>
+		/// <param name="Location">Location of the file</param>
+		public static void MakeWriteable(FileReference Location)
+		{
+			if(Exists(Location))
+			{
+				FileAttributes Attributes = GetAttributes(Location);
+				if((Attributes & FileAttributes.ReadOnly) != 0)
+				{
+					SetAttributes(Location, Attributes & ~FileAttributes.ReadOnly);
+				}
+			}
+		}
+
+		/// <summary>
+		/// Constructs a FileInfo object from this reference
+		/// </summary>
+		/// <returns>New FileInfo object</returns>
+		public FileInfo ToFileInfo()
+		{
+			return new FileInfo(FullName);
+		}
+
+		#region System.IO.File methods
+
+		/// <summary>
+		/// Copies a file from one location to another
+		/// </summary>
+		/// <param name="SourceLocation">Location of the source file</param>
+		/// <param name="TargetLocation">Location of the target file</param>
+		public static void Copy(FileReference SourceLocation, FileReference TargetLocation)
+		{
+			File.Copy(SourceLocation.FullName, TargetLocation.FullName);
+		}
+
+		/// <summary>
+		/// Copies a file from one location to another
+		/// </summary>
+		/// <param name="SourceLocation">Location of the source file</param>
+		/// <param name="TargetLocation">Location of the target file</param>
+		/// <param name="bOverwrite">Whether to overwrite the file in the target location</param>
+		public static void Copy(FileReference SourceLocation, FileReference TargetLocation, bool bOverwrite)
+		{
+			File.Copy(SourceLocation.FullName, TargetLocation.FullName, bOverwrite);
+		}
+
+		/// <summary>
+		/// Deletes this file
+		/// </summary>
+		public static void Delete(FileReference Location)
+		{
+			File.Delete(Location.FullName);
+		}
+
+		/// <summary>
+		/// Determines whether the given filename exists
+		/// </summary>
+		/// <returns>True if it exists, false otherwise</returns>
+		public static bool Exists(FileReference Location)
+		{
+			return File.Exists(Location.FullName);
+		}
+
+		/// <summary>
+		/// Gets the attributes for a file
+		/// </summary>
+		/// <param name="Location">Location of the file</param>
+		/// <returns>Attributes for the file</returns>
+		public static FileAttributes GetAttributes(FileReference Location)
+		{
+			return File.GetAttributes(Location.FullName);
+		}
+
+		/// <summary>
+		/// Gets the time that the file was last written to
+		/// </summary>
+		/// <param name="Location">Location of the file</param>
+		/// <returns>Last write time, in local time</returns>
+		public static DateTime GetLastWriteTime(FileReference Location)
+		{
+			return File.GetLastWriteTime(Location.FullName);
+		}
+
+		/// <summary>
+		/// Gets the time that the file was last written to
+		/// </summary>
+		/// <param name="Location">Location of the file</param>
+		/// <returns>Last write time, in UTC time</returns>
+		public static DateTime GetLastWriteTimeUtc(FileReference Location)
+		{
+			return File.GetLastWriteTimeUtc(Location.FullName);
+		}
+
+		/// <summary>
+		/// Moves a file from one location to another
+		/// </summary>
+		/// <param name="SourceLocation">Location of the source file</param>
+		/// <param name="TargetLocation">Location of the target file</param>
+		public static void Move(FileReference SourceLocation, FileReference TargetLocation)
+		{
+			File.Move(SourceLocation.FullName, TargetLocation.FullName);
+		}
+
+		/// <summary>
+		/// Opens a FileStream on the specified path with read/write access
+		/// </summary>
+		/// <param name="Location">Location of the file</param>
+		/// <param name="Mode">Mode to use when opening the file</param>
+		/// <returns>New filestream for the given file</returns>
+		public static FileStream Open(FileReference Location, FileMode Mode)
+		{
+			return File.Open(Location.FullName, Mode);
+		}
+
+		/// <summary>
+		/// Opens a FileStream on the specified path
+		/// </summary>
+		/// <param name="Location">Location of the file</param>
+		/// <param name="Mode">Mode to use when opening the file</param>
+		/// <param name="Access">Sharing mode for the new file</param>
+		/// <returns>New filestream for the given file</returns>
+		public static FileStream Open(FileReference Location, FileMode Mode, FileAccess Access)
+		{
+			return File.Open(Location.FullName, Mode, Access);
+		}
+
+		/// <summary>
+		/// Opens a FileStream on the specified path
+		/// </summary>
+		/// <param name="Location">Location of the file</param>
+		/// <param name="Mode">Mode to use when opening the file</param>
+		/// <param name="Access">Access mode for the new file</param>
+		/// <param name="Share">Sharing mode for the open file</param>
+		/// <returns>New filestream for the given file</returns>
+		public static FileStream Open(FileReference Location, FileMode Mode, FileAccess Access, FileShare Share)
+		{
+			return File.Open(Location.FullName, Mode, Access, Share);
+		}
+
+		/// <summary>
+		/// Reads the contents of a file
+		/// </summary>
+		/// <param name="Location">Location of the file</param>
+		/// <returns>Byte array containing the contents of the file</returns>
+		public static byte[] ReadAllBytes(FileReference Location)
+		{
+			return File.ReadAllBytes(Location.FullName);
+		}
+
+		/// <summary>
+		/// Reads the contents of a file
+		/// </summary>
+		/// <param name="Location">Location of the file</param>
+		/// <returns>String array containing the contents of the file</returns>
+		public static string[] ReadAllLines(FileReference Location)
+		{
+			return File.ReadAllLines(Location.FullName);
+		}
+
+		/// <summary>
+		/// Reads the contents of a file
+		/// </summary>
+		/// <param name="Location">Location of the file</param>
+		/// <param name="Encoding">The encoding to use when parsing the file</param>
+		/// <returns>String array containing the contents of the file</returns>
+		public static string[] ReadAllLines(FileReference Location, Encoding Encoding)
+		{
+			return File.ReadAllLines(Location.FullName, Encoding);
+		}
+
+		/// <summary>
+		/// Sets the attributes for a file
+		/// </summary>
+		/// <param name="Location">Location of the file</param>
+		/// <param name="Attributes">New attributes for the file</param>
+		public static void SetAttributes(FileReference Location, FileAttributes Attributes)
+		{
+			File.SetAttributes(Location.FullName, Attributes);
+		}
+
+		/// <summary>
+		/// Sets the time that the file was last written to
+		/// </summary>
+		/// <param name="Location">Location of the file</param>
+		/// <param name="LastWriteTime">Last write time, in local time</param>
+		public static void SetLastWriteTime(FileReference Location, DateTime LastWriteTime)
+		{
+			File.SetLastWriteTime(Location.FullName, LastWriteTime);
+		}
+
+		/// <summary>
+		/// Sets the time that the file was last written to
+		/// </summary>
+		/// <param name="Location">Location of the file</param>
+		/// <param name="LastWriteTimeUtc">Last write time, in UTC time</param>
+		public static void SetLastWriteTimeUtc(FileReference Location, DateTime LastWriteTimeUtc)
+		{
+			File.SetLastWriteTimeUtc(Location.FullName, LastWriteTimeUtc);
+		}
+
+		/// <summary>
+		/// Writes the contents of a file
+		/// </summary>
+		/// <param name="Location">Location of the file</param>
+		/// <param name="Contents">Contents of the file</param>
+		public static void WriteAllBytes(FileReference Location, byte[] Contents)
+		{
+			File.WriteAllBytes(Location.FullName, Contents);
+		}
+
+		/// <summary>
+		/// Writes the contents of a file
+		/// </summary>
+		/// <param name="Location">Location of the file</param>
+		/// <param name="Contents">Contents of the file</param>
+		public static void WriteAllLines(FileReference Location, IEnumerable<string> Contents)
+		{
+			File.WriteAllLines(Location.FullName, Contents);
+		}
+
+		/// <summary>
+		/// Writes the contents of a file
+		/// </summary>
+		/// <param name="Location">Location of the file</param>
+		/// <param name="Contents">Contents of the file</param>
+		/// <param name="Encoding">The encoding to use when parsing the file</param>
+		public static void WriteAllLines(FileReference Location, IEnumerable<string> Contents, Encoding Encoding)
+		{
+			File.WriteAllLines(Location.FullName, Contents, Encoding);
+		}
+
+		/// <summary>
+		/// Writes the contents of a file
+		/// </summary>
+		/// <param name="Location">Location of the file</param>
+		/// <param name="Contents">Contents of the file</param>
+		public static void WriteAllLines(FileReference Location, string[] Contents)
+		{
+			File.WriteAllLines(Location.FullName, Contents);
+		}
+
+		/// <summary>
+		/// Writes the contents of a file
+		/// </summary>
+		/// <param name="Location">Location of the file</param>
+		/// <param name="Contents">Contents of the file</param>
+		/// <param name="Encoding">The encoding to use when parsing the file</param>
+		public static void WriteAllLines(FileReference Location, string[] Contents, Encoding Encoding)
+		{
+			File.WriteAllLines(Location.FullName, Contents, Encoding);
+		}
+
+		#endregion
+	}
+
+	/// <summary>
+	/// Extension methods for FileReference functionality
+	/// </summary>
+	static class FileReferenceExtensionMethods
+	{
+		/// <summary>
+		/// Manually serialize a file reference to a binary stream.
+		/// </summary>
+		/// <param name="Writer">Binary writer to write to</param>
+		/// <param name="File">The file reference to write</param>
+		public static void Write(this BinaryWriter Writer, FileReference File)
+		{
+			Writer.Write((File == null) ? String.Empty : File.FullName);
+		}
+
+		/// <summary>
+		/// Serializes a file reference, using a lookup table to avoid serializing the same name more than once.
+		/// </summary>
+		/// <param name="Writer">The writer to save this reference to</param>
+		/// <param name="File">A file reference to output; may be null</param>
+		/// <param name="FileToUniqueId">A lookup table that caches previous files that have been output, and maps them to unique id's.</param>
+		public static void Write(this BinaryWriter Writer, FileReference File, Dictionary<FileReference, int> FileToUniqueId)
+		{
+			int UniqueId;
+			if (File == null)
+			{
+				Writer.Write(-1);
+			}
+			else if (FileToUniqueId.TryGetValue(File, out UniqueId))
+			{
+				Writer.Write(UniqueId);
+			}
+			else
+			{
+				Writer.Write(FileToUniqueId.Count);
+				Writer.Write(File);
+				FileToUniqueId.Add(File, FileToUniqueId.Count);
+			}
+		}
+
+		/// <summary>
+		/// Manually deserialize a file reference from a binary stream.
+		/// </summary>
+		/// <param name="Reader">Binary reader to read from</param>
+		/// <returns>New FileReference object</returns>
+		public static FileReference ReadFileReference(this BinaryReader Reader)
+		{
+			string FullName = Reader.ReadString();
+			return (FullName.Length == 0) ? null : FileReference.MakeFromNormalizedFullPath(FullName);
+		}
+
+		/// <summary>
+		/// Deserializes a file reference, using a lookup table to avoid writing the same name more than once.
+		/// </summary>
+		/// <param name="Reader">The source to read from</param>
+		/// <param name="UniqueFiles">List of previously read file references. The index into this array is used in place of subsequent ocurrences of the file.</param>
+		/// <returns>The file reference that was read</returns>
+		public static FileReference ReadFileReference(this BinaryReader Reader, List<FileReference> UniqueFiles)
+		{
+			int UniqueId = Reader.ReadInt32();
+			if (UniqueId == -1)
+			{
+				return null;
+			}
+			else if (UniqueId < UniqueFiles.Count)
+			{
+				return UniqueFiles[UniqueId];
+			}
+			else
+			{
+				FileReference Result = Reader.ReadFileReference();
+				UniqueFiles.Add(Result);
+				return Result;
+			}
+		}
+	}
+}
diff --git a/Engine/Source/Programs/UnrealBuildTool/System/FileSystemName.cs b/Engine/Source/Programs/UnrealBuildTool/System/FileSystemName.cs
new file mode 100644
--- /dev/null
+++ b/Engine/Source/Programs/UnrealBuildTool/System/FileSystemName.cs
@@ -0,0 +1,91 @@
+﻿using System;
+using System.Collections.Generic;
+using System.Linq;
+using System.Text;
+using System.Threading.Tasks;
+
+namespace UnrealBuildTool
+{
+	/// <summary>
+	/// Represents a case-insensitive name in the filesystem
+	/// </summary>
+	public class FileSystemName
+	{
+		/// <summary>
+		/// Name as it should be displayed
+		/// </summary>
+		public readonly string DisplayName;
+
+		/// <summary>
+		/// The canonical form of the name
+		/// </summary>
+		public readonly string CanonicalName;
+
+		/// <summary>
+		/// Constructor
+		/// </summary>
+		/// <param name="DisplayName">The display name</param>
+		public FileSystemName(string DisplayName)
+		{
+			this.DisplayName = DisplayName;
+			CanonicalName = DisplayName.ToLowerInvariant();
+		}
+
+		/// <summary>
+		/// Compares two filesystem object names for equality. Uses the canonical name representation, not the display name representation.
+		/// </summary>
+		/// <param name="A">First object to compare.</param>
+		/// <param name="B">Second object to compare.</param>
+		/// <returns>True if the names represent the same object, false otherwise</returns>
+		public static bool operator ==(FileSystemName A, FileSystemName B)
+		{
+			if ((object)A == null)
+			{
+				return (object)B == null;
+			}
+			else
+			{
+				return (object)B != null && A.CanonicalName == B.CanonicalName;
+			}
+		}
+
+		/// <summary>
+		/// Compares two filesystem object names for inequality. Uses the canonical name representation, not the display name representation.
+		/// </summary>
+		/// <param name="A">First object to compare.</param>
+		/// <param name="B">Second object to compare.</param>
+		/// <returns>False if the names represent the same object, true otherwise</returns>
+		public static bool operator !=(FileSystemName A, FileSystemName B)
+		{
+			return !(A == B);
+		}
+
+		/// <summary>
+		/// Compares against another object for equality.
+		/// </summary>
+		/// <param name="Obj">other instance to compare.</param>
+		/// <returns>True if the names represent the same object, false otherwise</returns>
+		public override bool Equals(object Obj)
+		{
+			return (Obj is FileSystemName) && ((FileSystemName)Obj) == this;
+		}
+
+		/// <summary>
+		/// Returns a hash code for this object
+		/// </summary>
+		/// <returns>Hash code for this object</returns>
+		public override int GetHashCode()
+		{
+			return base.GetHashCode();
+		}
+
+		/// <summary>
+		/// Returns the display name
+		/// </summary>
+		/// <returns>The display name</returns>
+		public override string ToString()
+		{
+			return DisplayName;
+		}
+	}
+}
diff --git a/Engine/Source/Programs/UnrealBuildTool/System/FileSystemReference.cs b/Engine/Source/Programs/UnrealBuildTool/System/FileSystemReference.cs
--- a/Engine/Source/Programs/UnrealBuildTool/System/FileSystemReference.cs
+++ b/Engine/Source/Programs/UnrealBuildTool/System/FileSystemReference.cs
@@ -9,89 +9,6 @@ using System.Threading.Tasks;
 
 namespace UnrealBuildTool
 {
-	/// <summary>
-	/// Represents a case-insensitive name in the filesystem
-	/// </summary>
-	public class FileSystemName
-	{
-		/// <summary>
-		/// Name as it should be displayed
-		/// </summary>
-		public readonly string DisplayName;
-
-		/// <summary>
-		/// The canonical form of the name
-		/// </summary>
-		public readonly string CanonicalName;
-
-		/// <summary>
-		/// Constructor
-		/// </summary>
-		/// <param name="DisplayName">The display name</param>
-		public FileSystemName(string DisplayName)
-		{
-			this.DisplayName = DisplayName;
-			CanonicalName = DisplayName.ToLowerInvariant();
-		}
-
-		/// <summary>
-		/// Compares two filesystem object names for equality. Uses the canonical name representation, not the display name representation.
-		/// </summary>
-		/// <param name="A">First object to compare.</param>
-		/// <param name="B">Second object to compare.</param>
-		/// <returns>True if the names represent the same object, false otherwise</returns>
-		public static bool operator ==(FileSystemName A, FileSystemName B)
-		{
-			if ((object)A == null)
-			{
-				return (object)B == null;
-			}
-			else
-			{
-				return (object)B != null && A.CanonicalName == B.CanonicalName;
-			}
-		}
-
-		/// <summary>
-		/// Compares two filesystem object names for inequality. Uses the canonical name representation, not the display name representation.
-		/// </summary>
-		/// <param name="A">First object to compare.</param>
-		/// <param name="B">Second object to compare.</param>
-		/// <returns>False if the names represent the same object, true otherwise</returns>
-		public static bool operator !=(FileSystemName A, FileSystemName B)
-		{
-			return !(A == B);
-		}
-
-		/// <summary>
-		/// Compares against another object for equality.
-		/// </summary>
-		/// <param name="Obj">other instance to compare.</param>
-		/// <returns>True if the names represent the same object, false otherwise</returns>
-		public override bool Equals(object Obj)
-		{
-			return (Obj is FileSystemName) && ((FileSystemName)Obj) == this;
-		}
-
-		/// <summary>
-		/// Returns a hash code for this object
-		/// </summary>
-		/// <returns>Hash code for this object</returns>
-		public override int GetHashCode()
-		{
-			return base.GetHashCode();
-		}
-
-		/// <summary>
-		/// Returns the display name
-		/// </summary>
-		/// <returns>The display name</returns>
-		public override string ToString()
-		{
-			return DisplayName;
-		}
-	}
-
 	/// <summary>
 	/// Base class for file system objects (files or directories).
 	/// </summary>
@@ -403,640 +320,4 @@ namespace UnrealBuildTool
 			return FullName;
 		}
 	}
-
-	/// <summary>
-	/// Representation of an absolute directory path. Allows fast hashing and comparisons.
-	/// </summary>
-	[Serializable]
-	public class DirectoryReference : FileSystemReference, IEquatable<DirectoryReference>
-	{
-		/// <summary>
-		/// Default constructor.
-		/// </summary>
-		/// <param name="InPath">Path to this directory.</param>
-		public DirectoryReference(string InPath)
-			: base(FixTrailingPathSeparator(Path.GetFullPath(InPath)))
-		{
-		}
-
-		/// <summary>
-		/// Construct a DirectoryReference from a DirectoryInfo object.
-		/// </summary>
-		/// <param name="InInfo">Path to this file</param>
-		public DirectoryReference(DirectoryInfo InInfo)
-			: base(FixTrailingPathSeparator(InInfo.FullName))
-		{
-		}
-
-		/// <summary>
-		/// Constructor for creating a directory object directly from two strings.
-		/// </summary>
-		/// <param name="InFullName"></param>
-		/// <param name="InCanonicalName"></param>
-		protected DirectoryReference(string InFullName, string InCanonicalName)
-			: base(InFullName, InCanonicalName)
-		{
-		}
-
-		/// <summary>
-		/// Ensures that the correct trailing path separator is appended. On Windows, the root directory (eg. C:\) always has a trailing path separator, but no other
-		/// path does.
-		/// </summary>
-		/// <param name="DirName">Absolute path to the directory</param>
-		/// <returns>Path to the directory, with the correct trailing path separator</returns>
-		private static string FixTrailingPathSeparator(string DirName)
-		{
-			if(DirName.Length == 2 && DirName[1] == ':')
-			{
-				return DirName + Path.DirectorySeparatorChar;
-			}
-			else if(DirName.Length == 3 && DirName[1] == ':' && DirName[2] == Path.DirectorySeparatorChar)
-			{
-				return DirName;
-			}
-			else if(DirName.Length > 1 && DirName[DirName.Length - 1] == Path.DirectorySeparatorChar)
-			{
-				return DirName.TrimEnd(Path.DirectorySeparatorChar);
-			}
-			else
-			{
-				return DirName;
-			}
-		}
-
-		/// <summary>
-		/// Gets the top level directory name
-		/// </summary>
-		/// <returns>The name of the directory</returns>
-		public string GetDirectoryName()
-		{
-			return Path.GetFileName(FullName);
-		}
-
-		/// <summary>
-		/// Gets the directory containing this object
-		/// </summary>
-		/// <returns>A new directory object representing the directory containing this object</returns>
-		public DirectoryReference ParentDirectory
-		{
-			get
-			{
-				if (IsRootDirectory())
-				{
-					return null;
-				}
-
-				int ParentLength = CanonicalName.LastIndexOf(Path.DirectorySeparatorChar);
-				if (ParentLength == 2 && CanonicalName[1] == ':')
-				{
-					ParentLength++;
-				}
-
-				return new DirectoryReference(FullName.Substring(0, ParentLength), CanonicalName.Substring(0, ParentLength));
-			}
-		}
-
-		/// <summary>
-		/// Gets the parent directory for a file
-		/// </summary>
-		/// <param name="File">The file to get directory for</param>
-		/// <returns>The full directory name containing the given file</returns>
-		public static DirectoryReference GetParentDirectory(FileReference File)
-		{
-			int ParentLength = File.CanonicalName.LastIndexOf(Path.DirectorySeparatorChar);
-			return new DirectoryReference(File.FullName.Substring(0, ParentLength), File.CanonicalName.Substring(0, ParentLength));
-		}
-
-		/// <summary>
-		/// Contains the current directory
-		/// </summary>
-		public static DirectoryReference GetCurrentDirectory()
-		{
-			return new DirectoryReference(Directory.GetCurrentDirectory());
-		}
-
-		/// <summary>
-		/// Creates a directory
-		/// </summary>
-		public static void CreateDirectory(DirectoryReference Location)
-		{
-			Directory.CreateDirectory(Location.FullName);
-		}
-
-        /// <summary>
-        /// Deletes a directory
-        /// </summary>
-        public static void Delete(DirectoryReference Location, bool bRecursive)
-        {
-            Directory.Delete(Location.FullName, bRecursive);
-        }
-
-        /// <summary>
-        /// Checks whether the directory exists
-        /// </summary>
-        /// <returns>True if this directory exists</returns>
-        public static bool Exists(DirectoryReference Location)
-		{
-			return Directory.Exists(Location.FullName);
-		}
-
-		/// <summary>
-		/// Enumerate files from a given directory
-		/// </summary>
-		/// <returns>Sequence of file references</returns>
-		public static IEnumerable<FileReference> EnumerateFiles(DirectoryReference BaseDir)
-		{
-			foreach (string FileName in Directory.EnumerateFiles(BaseDir.FullName))
-			{
-				yield return FileReference.MakeFromNormalizedFullPath(FileName);
-			}
-		}
-
-		/// <summary>
-		/// Enumerate files from a given directory
-		/// </summary>
-		/// <returns>Sequence of file references</returns>
-		public static IEnumerable<FileReference> EnumerateFiles(DirectoryReference BaseDir, string Pattern)
-		{
-			foreach (string FileName in Directory.EnumerateFiles(BaseDir.FullName, Pattern))
-			{
-				yield return FileReference.MakeFromNormalizedFullPath(FileName);
-			}
-		}
-
-		/// <summary>
-		/// Enumerate files from a given directory
-		/// </summary>
-		/// <returns>Sequence of file references</returns>
-		public static IEnumerable<FileReference> EnumerateFiles(DirectoryReference BaseDir, string Pattern, SearchOption Option)
-		{
-			foreach (string FileName in Directory.EnumerateFiles(BaseDir.FullName, Pattern, Option))
-			{
-				yield return FileReference.MakeFromNormalizedFullPath(FileName);
-			}
-		}
-
-		/// <summary>
-		/// Enumerate subdirectories in a given directory
-		/// </summary>
-		/// <returns>Sequence of directory references</returns>
-		public static IEnumerable<DirectoryReference> EnumerateDirectories(DirectoryReference BaseDir)
-		{
-			foreach (string DirectoryName in Directory.EnumerateDirectories(BaseDir.FullName))
-			{
-				yield return DirectoryReference.MakeFromNormalizedFullPath(DirectoryName);
-			}
-		}
-
-		/// <summary>
-		/// Enumerate subdirectories in a given directory
-		/// </summary>
-		/// <returns>Sequence of directory references</returns>
-		public static IEnumerable<DirectoryReference> EnumerateDirectories(DirectoryReference BaseDir, string Pattern)
-		{
-			foreach (string DirectoryName in Directory.EnumerateDirectories(BaseDir.FullName, Pattern))
-			{
-				yield return DirectoryReference.MakeFromNormalizedFullPath(DirectoryName);
-			}
-		}
-
-		/// <summary>
-		/// Enumerate subdirectories in a given directory
-		/// </summary>
-		/// <returns>Sequence of directory references</returns>
-		public static IEnumerable<DirectoryReference> EnumerateDirectories(DirectoryReference BaseDir, string Pattern, SearchOption Option)
-		{
-			foreach (string DirectoryName in Directory.EnumerateDirectories(BaseDir.FullName, Pattern, Option))
-			{
-				yield return DirectoryReference.MakeFromNormalizedFullPath(DirectoryName);
-			}
-		}
-
-		/// <summary>
-		/// Determines whether this path represents a root directory in the filesystem
-		/// </summary>
-		/// <returns>True if this path is a root directory, false otherwise</returns>
-		public bool IsRootDirectory()
-		{
-			return CanonicalName[CanonicalName.Length - 1] == Path.DirectorySeparatorChar;
-		}
-
-		/// <summary>
-		/// Combine several fragments with a base directory, to form a new directory name
-		/// </summary>
-		/// <param name="BaseDirectory">The base directory</param>
-		/// <param name="Fragments">Fragments to combine with the base directory</param>
-		/// <returns>The new directory name</returns>
-		public static DirectoryReference Combine(DirectoryReference BaseDirectory, params string[] Fragments)
-		{
-			string FullName = FileSystemReference.CombineStrings(BaseDirectory, Fragments);
-			return new DirectoryReference(FullName, FullName.ToLowerInvariant());
-		}
-
-		/// <summary>
-		/// Compares two filesystem object names for equality. Uses the canonical name representation, not the display name representation.
-		/// </summary>
-		/// <param name="A">First object to compare.</param>
-		/// <param name="B">Second object to compare.</param>
-		/// <returns>True if the names represent the same object, false otherwise</returns>
-		public static bool operator ==(DirectoryReference A, DirectoryReference B)
-		{
-			if ((object)A == null)
-			{
-				return (object)B == null;
-			}
-			else
-			{
-				return (object)B != null && A.CanonicalName == B.CanonicalName;
-			}
-		}
-
-		/// <summary>
-		/// Compares two filesystem object names for inequality. Uses the canonical name representation, not the display name representation.
-		/// </summary>
-		/// <param name="A">First object to compare.</param>
-		/// <param name="B">Second object to compare.</param>
-		/// <returns>False if the names represent the same object, true otherwise</returns>
-		public static bool operator !=(DirectoryReference A, DirectoryReference B)
-		{
-			return !(A == B);
-		}
-
-		/// <summary>
-		/// Compares against another object for equality.
-		/// </summary>
-		/// <param name="Obj">other instance to compare.</param>
-		/// <returns>True if the names represent the same object, false otherwise</returns>
-		public override bool Equals(object Obj)
-		{
-			return (Obj is DirectoryReference) && ((DirectoryReference)Obj) == this;
-		}
-
-		/// <summary>
-		/// Compares against another object for equality.
-		/// </summary>
-		/// <param name="Obj">other instance to compare.</param>
-		/// <returns>True if the names represent the same object, false otherwise</returns>
-		public bool Equals(DirectoryReference Obj)
-		{
-			return Obj == this;
-		}
-
-		/// <summary>
-		/// Returns a hash code for this object
-		/// </summary>
-		/// <returns></returns>
-		public override int GetHashCode()
-		{
-			return CanonicalName.GetHashCode();
-		}
-
-		/// <summary>
-		/// Helper function to create a remote directory reference. Unlike normal DirectoryReference objects, these aren't converted to a full path in the local filesystem.
-		/// </summary>
-		/// <param name="AbsolutePath">The absolute path in the remote file system</param>
-		/// <returns>New directory reference</returns>
-		public static DirectoryReference MakeRemote(string AbsolutePath)
-		{
-			return new DirectoryReference(AbsolutePath, AbsolutePath.ToLowerInvariant());
-		}
-
-		/// <summary>
-		/// Helper function to create a directory reference from a raw platform path. The path provided *MUST* be exactly the same as that returned by Path.GetFullPath(). 
-		/// </summary>
-		/// <param name="AbsolutePath">The absolute path in the file system</param>
-		/// <returns>New file reference</returns>
-		public static DirectoryReference MakeFromNormalizedFullPath(string AbsolutePath)
-		{
-			return new DirectoryReference(AbsolutePath, AbsolutePath.ToLowerInvariant());
-		}
-
-		/// <summary>
-		/// Gets the parent directory for a file, or returns null if it's null.
-		/// </summary>
-		/// <param name="File">The file to create a directory reference for</param>
-		/// <returns>The directory containing the file  </returns>
-		public static DirectoryReference FromFile(FileReference File)
-		{
-			return (File == null)? null : File.Directory;
-		}
-	}
-
-	/// <summary>
-	/// Representation of an absolute file path. Allows fast hashing and comparisons.
-	/// </summary>
-	[Serializable]
-	public class FileReference : FileSystemReference, IEquatable<FileReference>
-	{
-		/// <summary>
-		/// Default constructor.
-		/// </summary>
-		/// <param name="InPath">Path to this file</param>
-		public FileReference(string InPath)
-			: base(Path.GetFullPath(InPath))
-		{
-			if(FullName.EndsWith("\\") || FullName.EndsWith("/"))
-			{
-				throw new ArgumentException("File names may not be terminated by a path separator character");
-			}
-		}
-
-		/// <summary>
-		/// Construct a FileReference from a FileInfo object.
-		/// </summary>
-		/// <param name="InInfo">Path to this file</param>
-		public FileReference(FileInfo InInfo)
-			: base(InInfo.FullName)
-		{
-		}
-
-		/// <summary>
-		/// Default constructor.
-		/// </summary>
-		protected FileReference(string InFullName, string InCanonicalName)
-			: base(InFullName, InCanonicalName)
-		{
-		}
-
-		/// <summary>
-		/// Gets the file name without path information
-		/// </summary>
-		/// <returns>A string containing the file name</returns>
-		public string GetFileName()
-		{
-			return Path.GetFileName(FullName);
-		}
-
-		/// <summary>
-		/// Gets the file name without path information or an extension
-		/// </summary>
-		/// <returns>A string containing the file name without an extension</returns>
-		public string GetFileNameWithoutExtension()
-		{
-			return Path.GetFileNameWithoutExtension(FullName);
-		}
-
-		/// <summary>
-		/// Gets the file name without path or any extensions
-		/// </summary>
-		/// <returns>A string containing the file name without an extension</returns>
-		public string GetFileNameWithoutAnyExtensions()
-		{
-			int StartIdx = FullName.LastIndexOf(Path.DirectorySeparatorChar) + 1;
-
-			int EndIdx = FullName.IndexOf('.', StartIdx);
-			if (EndIdx < StartIdx)
-			{
-				return FullName.Substring(StartIdx);
-			}
-			else
-			{
-				return FullName.Substring(StartIdx, EndIdx - StartIdx);
-			}
-		}
-
-		/// <summary>
-		/// Gets the extension for this filename
-		/// </summary>
-		/// <returns>A string containing the extension of this filename</returns>
-		public string GetExtension()
-		{
-			return Path.GetExtension(FullName);
-		}
-
-		/// <summary>
-		/// Change the file's extension to something else
-		/// </summary>
-		/// <param name="Extension">The new extension</param>
-		/// <returns>A FileReference with the same path and name, but with the new extension</returns>
-		public FileReference ChangeExtension(string Extension)
-		{
-			string NewFullName = Path.ChangeExtension(FullName, Extension);
-			return new FileReference(NewFullName, NewFullName.ToLowerInvariant());
-		}
-
-		/// <summary>
-		/// Gets the directory containing this file
-		/// </summary>
-		/// <returns>A new directory object representing the directory containing this object</returns>
-		public DirectoryReference Directory
-		{
-			get { return DirectoryReference.GetParentDirectory(this); }
-		}
-
-		/// <summary>
-		/// Determines whether the given filename exists
-		/// </summary>
-		/// <returns>True if it exists, false otherwise</returns>
-		public static bool Exists(FileReference Location)
-		{
-			return File.Exists(Location.FullName);
-		}
-
-		/// <summary>
-		/// Deletes this file
-		/// </summary>
-		public static void Delete(FileReference Location)
-		{
-			File.Delete(Location.FullName);
-		}
-
-		/// <summary>
-		/// Combine several fragments with a base directory, to form a new filename
-		/// </summary>
-		/// <param name="BaseDirectory">The base directory</param>
-		/// <param name="Fragments">Fragments to combine with the base directory</param>
-		/// <returns>The new file name</returns>
-		public static FileReference Combine(DirectoryReference BaseDirectory, params string[] Fragments)
-		{
-			string FullName = FileSystemReference.CombineStrings(BaseDirectory, Fragments);
-			return new FileReference(FullName, FullName.ToLowerInvariant());
-		}
-
-		/// <summary>
-		/// Append a string to the end of a filename
-		/// </summary>
-		/// <param name="A">The base file reference</param>
-		/// <param name="B">Suffix to be appended</param>
-		/// <returns>The new file reference</returns>
-		public static FileReference operator +(FileReference A, string B)
-		{
-			return new FileReference(A.FullName + B, A.CanonicalName + B.ToLowerInvariant());
-		}
-
-		/// <summary>
-		/// Compares two filesystem object names for equality. Uses the canonical name representation, not the display name representation.
-		/// </summary>
-		/// <param name="A">First object to compare.</param>
-		/// <param name="B">Second object to compare.</param>
-		/// <returns>True if the names represent the same object, false otherwise</returns>
-		public static bool operator ==(FileReference A, FileReference B)
-		{
-			if ((object)A == null)
-			{
-				return (object)B == null;
-			}
-			else
-			{
-				return (object)B != null && A.CanonicalName == B.CanonicalName;
-			}
-		}
-
-		/// <summary>
-		/// Compares two filesystem object names for inequality. Uses the canonical name representation, not the display name representation.
-		/// </summary>
-		/// <param name="A">First object to compare.</param>
-		/// <param name="B">Second object to compare.</param>
-		/// <returns>False if the names represent the same object, true otherwise</returns>
-		public static bool operator !=(FileReference A, FileReference B)
-		{
-			return !(A == B);
-		}
-
-		/// <summary>
-		/// Compares against another object for equality.
-		/// </summary>
-		/// <param name="Obj">other instance to compare.</param>
-		/// <returns>True if the names represent the same object, false otherwise</returns>
-		public override bool Equals(object Obj)
-		{
-			return (Obj is FileReference) && ((FileReference)Obj) == this;
-		}
-
-		/// <summary>
-		/// Compares against another object for equality.
-		/// </summary>
-		/// <param name="Obj">other instance to compare.</param>
-		/// <returns>True if the names represent the same object, false otherwise</returns>
-		public bool Equals(FileReference Obj)
-		{
-			return Obj == this;
-		}
-
-		/// <summary>
-		/// Returns a hash code for this object
-		/// </summary>
-		/// <returns></returns>
-		public override int GetHashCode()
-		{
-			return CanonicalName.GetHashCode();
-		}
-
-		/// <summary>
-		/// Helper function to create a remote file reference. Unlike normal FileReference objects, these aren't converted to a full path in the local filesystem, but are
-		/// left as they are passed in.
-		/// </summary>
-		/// <param name="AbsolutePath">The absolute path in the remote file system</param>
-		/// <returns>New file reference</returns>
-		public static FileReference MakeRemote(string AbsolutePath)
-		{
-			return new FileReference(AbsolutePath, AbsolutePath.ToLowerInvariant());
-		}
-
-		/// <summary>
-		/// Helper function to create a file reference from a raw platform path. The path provided *MUST* be exactly the same as that returned by Path.GetFullPath(). 
-		/// </summary>
-		/// <param name="AbsolutePath">The absolute path in the file system</param>
-		/// <returns>New file reference</returns>
-		public static FileReference MakeFromNormalizedFullPath(string AbsolutePath)
-		{
-			return new FileReference(AbsolutePath, AbsolutePath.ToLowerInvariant());
-		}
-	}
-
-	static class FileReferenceExtensionMethods
-	{
-		/// <summary>
-		/// Manually serialize a file reference to a binary stream.
-		/// </summary>
-		/// <param name="Writer">Binary writer to write to</param>
-		/// <param name="File">The file reference to write</param>
-		public static void Write(this BinaryWriter Writer, FileReference File)
-		{
-			Writer.Write((File == null) ? String.Empty : File.FullName);
-		}
-
-		/// <summary>
-		/// Manually serialize a file reference to a binary stream.
-		/// </summary>
-		/// <param name="Writer">Binary writer to write to</param>
-		/// <param name="Directory">The directory reference to write</param>
-		public static void Write(this BinaryWriter Writer, DirectoryReference Directory)
-		{
-			Writer.Write((Directory == null) ? String.Empty : Directory.FullName);
-		}
-
-		/// <summary>
-		/// Serializes a file reference, using a lookup table to avoid serializing the same name more than once.
-		/// </summary>
-		/// <param name="Writer">The writer to save this reference to</param>
-		/// <param name="File">A file reference to output; may be null</param>
-		/// <param name="FileToUniqueId">A lookup table that caches previous files that have been output, and maps them to unique id's.</param>
-		public static void Write(this BinaryWriter Writer, FileReference File, Dictionary<FileReference, int> FileToUniqueId)
-		{
-			int UniqueId;
-			if (File == null)
-			{
-				Writer.Write(-1);
-			}
-			else if (FileToUniqueId.TryGetValue(File, out UniqueId))
-			{
-				Writer.Write(UniqueId);
-			}
-			else
-			{
-				Writer.Write(FileToUniqueId.Count);
-				Writer.Write(File);
-				FileToUniqueId.Add(File, FileToUniqueId.Count);
-			}
-		}
-
-		/// <summary>
-		/// Manually deserialize a file reference from a binary stream.
-		/// </summary>
-		/// <param name="Reader">Binary reader to read from</param>
-		/// <returns>New FileReference object</returns>
-		public static FileReference ReadFileReference(this BinaryReader Reader)
-		{
-			string FullName = Reader.ReadString();
-			return (FullName.Length == 0) ? null : FileReference.MakeFromNormalizedFullPath(FullName);
-		}
-
-		/// <summary>
-		/// Manually deserialize a directory reference from a binary stream.
-		/// </summary>
-		/// <param name="Reader">Binary reader to read from</param>
-		/// <returns>New DirectoryReference object</returns>
-		public static DirectoryReference ReadDirectoryReference(this BinaryReader Reader)
-		{
-			string FullName = Reader.ReadString();
-			return (FullName.Length == 0) ? null : DirectoryReference.MakeFromNormalizedFullPath(FullName);
-		}
-
-
-		/// <summary>
-		/// Deserializes a file reference, using a lookup table to avoid writing the same name more than once.
-		/// </summary>
-		/// <param name="Reader">The source to read from</param>
-		/// <param name="UniqueFiles">List of previously read file references. The index into this array is used in place of subsequent ocurrences of the file.</param>
-		/// <returns>The file reference that was read</returns>
-		public static FileReference ReadFileReference(this BinaryReader Reader, List<FileReference> UniqueFiles)
-		{
-			int UniqueId = Reader.ReadInt32();
-			if (UniqueId == -1)
-			{
-				return null;
-			}
-			else if (UniqueId < UniqueFiles.Count)
-			{
-				return UniqueFiles[UniqueId];
-			}
-			else
-			{
-				FileReference Result = Reader.ReadFileReference();
-				UniqueFiles.Add(Result);
-				return Result;
-			}
-		}
-	}
 }
diff --git a/Engine/Source/Programs/UnrealBuildTool/System/FlatCPPIncludeDepencencyCache.cs b/Engine/Source/Programs/UnrealBuildTool/System/FlatCPPIncludeDepencencyCache.cs
--- a/Engine/Source/Programs/UnrealBuildTool/System/FlatCPPIncludeDepencencyCache.cs
+++ b/Engine/Source/Programs/UnrealBuildTool/System/FlatCPPIncludeDepencencyCache.cs
@@ -34,7 +34,7 @@ namespace UnrealBuildTool
 		/// <summary>
 		/// The version number for binary serialization
 		/// </summary>
-		const int FileVersion = 1;
+		const int FileVersion = 2;
 
 		/// <summary>
 		/// The file signature for binary serialization
diff --git a/Engine/Source/Programs/UnrealBuildTool/System/PlatformExports.cs b/Engine/Source/Programs/UnrealBuildTool/System/PlatformExports.cs
--- a/Engine/Source/Programs/UnrealBuildTool/System/PlatformExports.cs
+++ b/Engine/Source/Programs/UnrealBuildTool/System/PlatformExports.cs
@@ -58,7 +58,7 @@ namespace UnrealBuildTool
 		public static bool HasDefaultBuildConfig(FileReference ProjectFile, UnrealTargetPlatform Platform)
 		{
 			UEBuildPlatform BuildPlat = UEBuildPlatform.GetBuildPlatform(Platform, true);
-			return (BuildPlat == null)? true : (BuildPlat as UEBuildPlatform).HasDefaultBuildConfig(Platform, ProjectFile.Directory);
+			return (BuildPlat == null)? true : BuildPlat.HasDefaultBuildConfig(Platform, ProjectFile.Directory);
 		}
 
 		/// <summary>
diff --git a/Engine/Source/Programs/UnrealBuildTool/System/ProgressWriter.cs b/Engine/Source/Programs/UnrealBuildTool/System/ProgressWriter.cs
new file mode 100644
--- /dev/null
+++ b/Engine/Source/Programs/UnrealBuildTool/System/ProgressWriter.cs
@@ -0,0 +1,85 @@
+﻿using System;
+using System.Collections.Generic;
+using System.Linq;
+using System.Text;
+using System.Threading.Tasks;
+
+namespace UnrealBuildTool
+{
+	/// <summary>
+	/// Class to display an incrementing progress percentage. Handles progress markup and direct console output.
+	/// </summary>
+	public class ProgressWriter : IDisposable
+	{
+		/// <summary>
+		/// Global setting controlling whether to output markup
+		/// </summary>
+		public static bool bWriteMarkup = false;
+
+		/// <summary>
+		/// Whether to write messages to the console
+		/// </summary>
+		bool bWriteToConsole;
+		string Message;
+		int NumCharsToBackspaceOver;
+		string CurrentProgressString;
+
+		/// <summary>
+		/// Constructor
+		/// </summary>
+		/// <param name="InMessage">The message to display before the progress percentage</param>
+		/// <param name="bInWriteToConsole">Whether to write progress message to the console</param>
+		public ProgressWriter(string InMessage, bool bInWriteToConsole)
+		{
+			Message = InMessage;
+			bWriteToConsole = bInWriteToConsole;
+			if (!bWriteMarkup && bWriteToConsole)
+			{
+				Console.Write(Message + " ");
+			}
+			Write(0, 100);
+		}
+
+		/// <summary>
+		/// Write the terminating newline
+		/// </summary>
+		public void Dispose()
+		{
+			if (!bWriteMarkup && bWriteToConsole)
+			{
+				Console.WriteLine();
+			}
+		}
+
+		/// <summary>
+		/// Writes the current progress
+		/// </summary>
+		/// <param name="Numerator">Numerator for the progress fraction</param>
+		/// <param name="Denominator">Denominator for the progress fraction</param>
+		public void Write(int Numerator, int Denominator)
+		{
+			float ProgressValue = Denominator > 0 ? ((float)Numerator / (float)Denominator) : 1.0f;
+			string ProgressString = String.Format("{0}%", Math.Round(ProgressValue * 100.0f));
+			if (ProgressString != CurrentProgressString)
+			{
+				CurrentProgressString = ProgressString;
+				if (bWriteMarkup)
+				{
+					Log.WriteLine(LogEventType.Console, "@progress '{0}' {1}", Message, ProgressString);
+				}
+				else if (bWriteToConsole)
+				{
+					// Backspace over previous progress value
+					while (NumCharsToBackspaceOver-- > 0)
+					{
+						Console.Write("\b");
+					}
+
+					// Display updated progress string and keep track of how long it was
+					NumCharsToBackspaceOver = ProgressString.Length;
+					Console.Write(ProgressString);
+				}
+			}
+		}
+	}
+}
diff --git a/Engine/Source/Programs/UnrealBuildTool/System/QMakefileGenerator.cs b/Engine/Source/Programs/UnrealBuildTool/System/QMakefileGenerator.cs
--- a/Engine/Source/Programs/UnrealBuildTool/System/QMakefileGenerator.cs
+++ b/Engine/Source/Programs/UnrealBuildTool/System/QMakefileGenerator.cs
@@ -302,13 +302,9 @@ namespace UnrealBuildTool
 								{
 									QMakeSourceFilesList += ("\t\"" + SourceFileRelativeToRoot.Substring(3) + "\" \\\n");
 								}
-								else if (!String.IsNullOrEmpty(GameProjectName))
-								{
-									QMakeSourceFilesList += ("\t\"$$" + GameProjectName + "RootPath/" + Utils.MakePathRelativeTo(CurSourceFile.FullName, GameProjectPath) + "\" \\\n");
-								}
 								else
 								{
-									System.Console.WriteLine("Error!, you should not be here.");
+									QMakeSourceFilesList += ("\t\"$$" + GameProjectName + "RootPath/" + Utils.MakePathRelativeTo(CurSourceFile.FullName, GameProjectPath) + "\" \\\n");
 								}
 							}
 						}
@@ -326,13 +322,9 @@ namespace UnrealBuildTool
 									// SourceFileRelativeToRoot = SourceFileRelativeToRoot.Substring (3);
 									QMakeHeaderFilesList += ("\t\"" + SourceFileRelativeToRoot.Substring(3) + "\" \\\n");
 								}
-								else if (!String.IsNullOrEmpty(GameProjectName))
-								{
-									QMakeHeaderFilesList += ("\t\"$$" + GameProjectName + "RootPath/" + Utils.MakePathRelativeTo(CurSourceFile.FullName, GameProjectPath) + "\" \\\n");
-								}
 								else
 								{
-									System.Console.WriteLine("Error!, you should not be here.");
+									QMakeHeaderFilesList += ("\t\"$$" + GameProjectName + "RootPath/" + Utils.MakePathRelativeTo(CurSourceFile.FullName, GameProjectPath) + "\" \\\n");
 								}
 							}
 						}
@@ -351,13 +343,9 @@ namespace UnrealBuildTool
 									// SourceFileRelativeToRoot = SourceFileRelativeToRoot.Substring (3);
 									QMakeConfigFilesList += ("\t\"" + SourceFileRelativeToRoot.Substring(3) + "\" \\\n");
 								}
-								else if (!String.IsNullOrEmpty(GameProjectName))
-								{
-									QMakeConfigFilesList += ("\t\"$$" + GameProjectName + "RootPath/" + Utils.MakePathRelativeTo(CurSourceFile.FullName, GameProjectPath) + "\" \\\n");
-								}
 								else
 								{
-									System.Console.WriteLine("Error!, you should not be here.");
+									QMakeConfigFilesList += ("\t\"$$" + GameProjectName + "RootPath/" + Utils.MakePathRelativeTo(CurSourceFile.FullName, GameProjectPath) + "\" \\\n");
 								}
 							}
 						}
diff --git a/Engine/Source/Programs/UnrealBuildTool/System/RulesDocumentation.cs b/Engine/Source/Programs/UnrealBuildTool/System/RulesDocumentation.cs
new file mode 100644
--- /dev/null
+++ b/Engine/Source/Programs/UnrealBuildTool/System/RulesDocumentation.cs
@@ -0,0 +1,147 @@
+﻿using System;
+using System.Collections.Generic;
+using System.IO;
+using System.Linq;
+using System.Reflection;
+using System.Text;
+using System.Threading.Tasks;
+using System.Xml;
+
+namespace UnrealBuildTool
+{
+	static class RulesDocumentation
+	{
+		public static void WriteDocumentation(Type RulesType, FileReference OutputFile)
+		{
+			// Get the path to the XML documentation
+			FileReference InputDocumentationFile = new FileReference(Assembly.GetExecutingAssembly().Location).ChangeExtension(".xml");
+			if(!FileReference.Exists(InputDocumentationFile))
+			{
+				throw new BuildException("Generated assembly documentation not found at {0}.", InputDocumentationFile);
+			}
+
+			// Get the current engine version for versioning the page
+			BuildVersion Version;
+			if(!BuildVersion.TryRead(out Version))
+			{
+				throw new BuildException("Unable to read the current build version");
+			}
+
+			// Read the documentation
+			XmlDocument InputDocumentation = new XmlDocument();
+			InputDocumentation.Load(InputDocumentationFile.FullName);
+
+			// Filter the properties into read-only and read/write lists
+			List<FieldInfo> ReadOnlyFields = new List<FieldInfo>();
+			List<FieldInfo> ReadWriteFields = new List<FieldInfo>();
+			foreach(FieldInfo Field in RulesType.GetFields(BindingFlags.Instance | BindingFlags.SetProperty | BindingFlags.Public))
+			{
+				if(!Field.FieldType.IsClass || !Field.FieldType.Name.EndsWith("TargetRules"))
+				{
+					if(Field.IsInitOnly)
+					{
+						ReadOnlyFields.Add(Field);
+					}
+					else
+					{
+						ReadWriteFields.Add(Field);
+					}
+				}
+			}
+
+			// Make sure the output file is writable
+			FileReference.MakeWriteable(OutputFile);
+
+			// Generate the UDN documentation file
+			using (StreamWriter Writer = new StreamWriter(OutputFile.FullName))
+			{
+				Writer.WriteLine("Availability: NoPublish");
+				Writer.WriteLine("Title: Build Configuration Properties Page");
+				Writer.WriteLine("Crumbs:");
+				Writer.WriteLine("Description: This is a procedurally generated markdown page.");
+				Writer.WriteLine("Version: {0}.{1}", Version.MajorVersion, Version.MinorVersion);
+				Writer.WriteLine("");
+				if(ReadOnlyFields.Count > 0)
+				{
+					Writer.WriteLine("### Read-Only Properties");
+					Writer.WriteLine();
+					foreach(FieldInfo Field in ReadOnlyFields)
+					{
+						OutputField(InputDocumentation, Field, Writer);
+					}
+					Writer.WriteLine();
+				}
+				if(ReadWriteFields.Count > 0)
+				{
+					Writer.WriteLine("### Read/Write Properties");
+					foreach(FieldInfo Field in ReadWriteFields)
+					{
+						OutputField(InputDocumentation, Field, Writer);
+					}
+					Writer.WriteLine("");
+				}
+			}
+
+			// Success!
+			Log.TraceInformation("Written documentation to {0}.", OutputFile);
+		}
+
+		static void OutputField(XmlDocument InputDocumentation, FieldInfo Field, TextWriter Writer)
+		{
+			XmlNode Node = InputDocumentation.SelectSingleNode(String.Format("//member[@name='F:{0}.{1}']/summary", Field.DeclaringType.FullName, Field.Name));
+			if(Node != null)
+			{
+				// Reflow the comments into paragraphs, assuming that each paragraph will be separated by a blank line
+				List<string> Lines = new List<string>(Node.InnerText.Trim().Split('\n').Select(x => x.Trim()));
+				for(int Idx = Lines.Count - 1; Idx > 0; Idx--)
+				{
+					if(Lines[Idx - 1].Length > 0 && !Lines[Idx].StartsWith("*") && !Lines[Idx].StartsWith("-"))
+					{
+						Lines[Idx - 1] += " " + Lines[Idx];
+						Lines.RemoveAt(Idx);
+					}
+				}
+
+				// Write the values of the enum
+/*				if(Field.FieldType.IsEnum)
+				{
+					Lines.Add("Valid values are:");
+					foreach(string Value in Enum.GetNames(Field.FieldType))
+					{
+						Lines.Add(String.Format("* {0}.{1}", Field.FieldType.Name, Value));
+					}
+				}
+*/
+				// Write the result to the .udn file
+				if(Lines.Count > 0)
+				{
+					Writer.WriteLine("$ {0} ({1}): {2}", Field.Name, GetPrettyTypeName(Field.FieldType), Lines[0]);
+					for(int Idx = 1; Idx < Lines.Count; Idx++)
+					{
+						if(Lines[Idx].StartsWith("*") || Lines[Idx].StartsWith("-"))
+						{
+							Writer.WriteLine("        * {0}", Lines[Idx].Substring(1).TrimStart());
+						}
+						else
+						{
+							Writer.WriteLine("    * {0}", Lines[Idx]);
+						}
+					}
+					Writer.WriteLine();
+				}
+			}
+		}
+
+		static string GetPrettyTypeName(Type FieldType)
+		{
+			if(FieldType.IsGenericType)
+			{
+				return String.Format("{0}&lt;{1}&gt;", FieldType.Name.Substring(0, FieldType.Name.IndexOf('`')), String.Join(", ", FieldType.GenericTypeArguments.Select(x => GetPrettyTypeName(x))));
+			}
+			else
+			{
+				return FieldType.Name;
+			}
+		}
+	}
+}
diff --git a/Engine/Source/Programs/UnrealBuildTool/System/ScopedTimer.cs b/Engine/Source/Programs/UnrealBuildTool/System/ScopedTimer.cs
new file mode 100644
--- /dev/null
+++ b/Engine/Source/Programs/UnrealBuildTool/System/ScopedTimer.cs
@@ -0,0 +1,53 @@
+﻿using System;
+using System.Collections.Generic;
+using System.Linq;
+using System.Text;
+using System.Threading.Tasks;
+
+namespace UnrealBuildTool
+{
+	/// <summary>
+	/// Scoped timer, start is in the constructor, end in Dispose. Best used with using(var Timer = new ScopedTimer()). Suports nesting.
+	/// </summary>
+	public class ScopedTimer : IDisposable
+	{
+		DateTime StartTime;
+		string TimerName;
+		LogEventType Verbosity;
+		static int Indent = 0;
+		static object IndentLock = new object();
+
+		/// <summary>
+		/// Constructor
+		/// </summary>
+		/// <param name="Name">Name of the block being measured</param>
+		/// <param name="InVerbosity">Verbosity for output messages</param>
+		public ScopedTimer(string Name, LogEventType InVerbosity = LogEventType.Verbose)
+		{
+			TimerName = Name;
+			lock (IndentLock)
+			{
+				Indent++;
+			}
+			Verbosity = InVerbosity;
+			StartTime = DateTime.UtcNow;
+		}
+
+		/// <summary>
+		/// Prints out the timing message
+		/// </summary>
+		public void Dispose()
+		{
+			double TotalSeconds = (DateTime.UtcNow - StartTime).TotalSeconds;
+			int LogIndent = 0;
+			lock (IndentLock)
+			{
+				LogIndent = --Indent;
+			}
+			StringBuilder IndentText = new StringBuilder(LogIndent * 2);
+			IndentText.Append(' ', LogIndent * 2);
+
+			Log.WriteLine(Verbosity, "{0}{1} took {2}s", IndentText.ToString(), TimerName, TotalSeconds);
+		}
+	}
+}
diff --git a/Engine/Source/Programs/UnrealBuildTool/System/UnrealBuildTool.cs b/Engine/Source/Programs/UnrealBuildTool/System/UnrealBuildTool.cs
--- a/Engine/Source/Programs/UnrealBuildTool/System/UnrealBuildTool.cs
+++ b/Engine/Source/Programs/UnrealBuildTool/System/UnrealBuildTool.cs
@@ -363,7 +363,7 @@ namespace UnrealBuildTool
 		private static int GuardedMain(string[] Arguments)
 		{
 			// Do super early log init as a safeguard. We'll re-init with proper config options later.
-			Log.InitLogging(bLogTimestamps: false, InLogLevel: LogEventType.Log, bLogSeverity: false, bLogSources: false, bColorConsoleOutput: true, TraceListeners: new[] { new ConsoleTraceListener() });
+			Log.InitLogging(bLogTimestamps: false, InLogLevel: LogEventType.Log, bLogSeverity: false, bLogSources: false, bLogSourcesToConsole: false, bColorConsoleOutput: true, TraceListeners: new[] { new ConsoleTraceListener() });
 
 			// ensure we can resolve any external assemblies that are not in the same folder as our assembly.
 			AssemblyUtils.InstallAssemblyResolver(Path.GetDirectoryName(Assembly.GetEntryAssembly().GetOriginalLocation()));
@@ -402,7 +402,7 @@ namespace UnrealBuildTool
 			if(!XmlConfig.ReadConfigFiles())
 			{
 				return 1;
-		}
+			}
 
 			// Create the build configuration object, and read the settings
 			BuildConfiguration BuildConfiguration = new BuildConfiguration();
@@ -444,6 +444,21 @@ namespace UnrealBuildTool
 				{
 					bIsEngineInstalled = false;
 				}
+				else if(LowercaseArg.StartsWith("-buildconfigurationdoc="))
+				{
+					XmlConfig.WriteDocumentation(new FileReference(Argument.Substring("-buildconfigurationdoc=".Length)));
+					return 0;
+				}
+				else if(LowercaseArg.StartsWith("-modulerulesdoc="))
+				{
+					RulesDocumentation.WriteDocumentation(typeof(ModuleRules), new FileReference(Argument.Substring("-modulerulesdoc=".Length)));
+					return 0;
+				}
+				else if(LowercaseArg.StartsWith("-targetrulesdoc="))
+				{
+					RulesDocumentation.WriteDocumentation(typeof(TargetRules), new FileReference(Argument.Substring("-targetrulesdoc=".Length)));
+					return 0;
+				}
 			}
 
 			// If it wasn't set explicitly by a command line option, check for the installed build marker file
@@ -460,6 +475,7 @@ namespace UnrealBuildTool
 				InLogLevel: (LogEventType)Enum.Parse(typeof(LogEventType), BuildConfiguration.LogLevel),
 				bLogSeverity: false,
 				bLogSources: false,
+				bLogSourcesToConsole: false,
 				bColorConsoleOutput: true,
 				TraceListeners: new[] 
                 {
@@ -2250,10 +2266,10 @@ namespace UnrealBuildTool
 
 						// Remove "-####-Platform-Configuration" suffix in Debug configuration
 						int IndexOfSecondLastHyphen = OriginalFileNameWithoutExtension.LastIndexOf('-', IndexOfLastHyphen - 1, IndexOfLastHyphen);
-						if (IndexOfSecondLastHyphen >= 0)
+						if (IndexOfSecondLastHyphen > 0)
 						{
 							int IndexOfThirdLastHyphen = OriginalFileNameWithoutExtension.LastIndexOf('-', IndexOfSecondLastHyphen - 1, IndexOfSecondLastHyphen);
-							if (IndexOfThirdLastHyphen >= 0)
+							if (IndexOfThirdLastHyphen > 0)
 							{
 								if (int.TryParse(OriginalFileNameWithoutExtension.Substring(IndexOfThirdLastHyphen + 1, IndexOfSecondLastHyphen - IndexOfThirdLastHyphen - 1), out NumberSuffix))
 								{
@@ -2419,10 +2435,16 @@ namespace UnrealBuildTool
 		/// <returns>True if file is part of the working set</returns>
 		public static bool ShouldSourceFileBePartOfWorkingSet(string SourceFileAbsolutePath)
 		{
-			// @todo: This logic needs work.  Currently its designed for users working with Perforce, or a similar
-			// source control system that keeps source files read-only when not actively checked out for editing.
-			bool bShouldBePartOfWorkingSourceFileSet =
-				!System.IO.File.GetAttributes(SourceFileAbsolutePath).HasFlag(FileAttributes.ReadOnly);
+			bool bShouldBePartOfWorkingSourceFileSet = false;
+			try
+			{
+				// @todo: This logic needs work.  Currently its designed for users working with Perforce, or a similar
+				// source control system that keeps source files read-only when not actively checked out for editing.
+				bShouldBePartOfWorkingSourceFileSet = !System.IO.File.GetAttributes(SourceFileAbsolutePath).HasFlag(FileAttributes.ReadOnly);
+			}
+			catch (FileNotFoundException)
+			{
+			}
 			return bShouldBePartOfWorkingSourceFileSet;
 		}
 	}
