Copying //UE4/Dev-Core to //UE4/Dev-Main (Source: //UE4/Dev-Core @ 3345728)

diff --git a/Engine/Source/Programs/UnrealBuildTool/System/Utils.cs b/Engine/Source/Programs/UnrealBuildTool/System/Utils.cs
--- a/Engine/Source/Programs/UnrealBuildTool/System/Utils.cs
+++ b/Engine/Source/Programs/UnrealBuildTool/System/Utils.cs
@@ -17,196 +17,6 @@ using System.Web.Script.Serialization;
 
 namespace UnrealBuildTool
 {
-	/// <summary>
-	/// Holds information about the current engine version
-	/// </summary>
-	[Serializable]
-	public class BuildVersion
-	{
-		/// <summary>
-		/// The major engine version (4 for UE4)
-		/// </summary>
-		public int MajorVersion;
-
-		/// <summary>
-		/// The minor engine version
-		/// </summary>
-		public int MinorVersion;
-
-		/// <summary>
-		/// The hotfix/patch version
-		/// </summary>
-		public int PatchVersion;
-
-		/// <summary>
-		/// The changelist that the engine is being built from
-		/// </summary>
-		public int Changelist;
-
-		/// <summary>
-		/// The changelist that the engine maintains compatibility with
-		/// </summary>
-		public int CompatibleChangelist;
-
-		/// <summary>
-		/// Whether the changelist numbers are a licensee changelist
-		/// </summary>
-		public int IsLicenseeVersion;
-
-		/// <summary>
-		/// Whether the current build is a promoted build, that is, built strictly from a clean sync of the given changelist
-		/// </summary>
-		public int IsPromotedBuild;
-
-		/// <summary>
-		/// Name of the current branch, with '/' characters escaped as '+'
-		/// </summary>
-		public string BranchName;
-
-		/// <summary>
-		/// Returns the value which can be used as the compatible changelist. Requires that the regular changelist is also set, and defaults to the 
-		/// regular changelist if a specific compatible changelist is not set.
-		/// </summary>
-		public int EffectiveCompatibleChangelist
-		{
-			get { return (Changelist != 0 && CompatibleChangelist != 0)? CompatibleChangelist : Changelist; }
-		}
-
-		/// <summary>
-		/// Try to read a version file from disk
-		/// </summary>
-		/// <param name="Version">The version information</param>
-		/// <returns>True if the version was read sucessfully, false otherwise</returns>
-		public static bool TryRead(out BuildVersion Version)
-		{
-			return TryRead(GetDefaultFileName(), out Version);
-		}
-
-		/// <summary>
-		/// Try to read a version file from disk
-		/// </summary>
-		/// <param name="FileName">Path to the version file</param>
-		/// <param name="Version">The version information</param>
-		/// <returns>True if the version was read sucessfully, false otherwise</returns>
-		public static bool TryRead(string FileName, out BuildVersion Version)
-		{
-			JsonObject Object;
-			if (!JsonObject.TryRead(FileName, out Object))
-			{
-				Version = null;
-				return false;
-			}
-			return TryParse(Object, out Version);
-		}
-
-		/// <summary>
-		/// Get the default path to the build.version file on disk
-		/// </summary>
-		/// <returns>Path to the Build.version file</returns>
-		public static string GetDefaultFileName()
-		{
-			return FileReference.Combine(UnrealBuildTool.EngineDirectory, "Build", "Build.version").FullName;
-		}
-
-		/// <summary>
-		/// Parses a build version from a JsonObject
-		/// </summary>
-		/// <param name="Object">The object to read from</param>
-		/// <param name="Version">The resulting version field</param>
-		/// <returns>True if the build version could be read, false otherwise</returns>
-		public static bool TryParse(JsonObject Object, out BuildVersion Version)
-		{
-			BuildVersion NewVersion = new BuildVersion();
-			if (!Object.TryGetIntegerField("MajorVersion", out NewVersion.MajorVersion) || !Object.TryGetIntegerField("MinorVersion", out NewVersion.MinorVersion) || !Object.TryGetIntegerField("PatchVersion", out NewVersion.PatchVersion))
-			{
-				Version = null;
-				return false;
-			}
-
-			Object.TryGetIntegerField("Changelist", out NewVersion.Changelist);
-			Object.TryGetIntegerField("CompatibleChangelist", out NewVersion.CompatibleChangelist);
-			Object.TryGetIntegerField("IsLicenseeVersion", out NewVersion.IsLicenseeVersion);
-			Object.TryGetIntegerField("IsPromotedBuild", out NewVersion.IsPromotedBuild);
-			Object.TryGetStringField("BranchName", out NewVersion.BranchName);
-
-			Version = NewVersion;
-			return true;
-		}
-
-		/// <summary>
-		/// Exports this object as Json
-		/// </summary>
-		/// <param name="FileName">The filename to write to</param>
-		/// <returns>True if the build version could be read, false otherwise</returns>
-		public void Write(string FileName)
-		{
-			using (JsonWriter Writer = new JsonWriter(FileName))
-			{
-				Writer.WriteObjectStart();
-				WriteProperties(Writer);
-				Writer.WriteObjectEnd();
-			}
-		}
-
-		/// <summary>
-		/// Exports this object as Json
-		/// </summary>
-		/// <param name="Writer">The json writer to receive the version settings</param>
-		/// <returns>True if the build version could be read, false otherwise</returns>
-		public void WriteProperties(JsonWriter Writer)
-		{
-			Writer.WriteValue("MajorVersion", MajorVersion);
-			Writer.WriteValue("MinorVersion", MinorVersion);
-			Writer.WriteValue("PatchVersion", PatchVersion);
-			Writer.WriteValue("Changelist", Changelist);
-			Writer.WriteValue("CompatibleChangelist", CompatibleChangelist);
-			Writer.WriteValue("IsLicenseeVersion", IsLicenseeVersion);
-			Writer.WriteValue("IsPromotedBuild", IsPromotedBuild);
-			Writer.WriteValue("BranchName", BranchName);
-		}
-	}
-
-	/// 
-	/// Log Event Type
-	///
-	public enum LogEventType
-	{
-		/// <summary>
-		/// The log event is a fatal error
-		/// </summary>
-		Fatal = 1,
-
-		/// <summary>
-		/// The log event is an error
-		/// </summary>
-		Error = 2,
-
-		/// <summary>
-		/// The log event is a warning
-		/// </summary>
-		Warning = 4,
-
-		/// <summary>
-		/// Output the log event to the console
-		/// </summary>
-		Console = 8,
-
-		/// <summary>
-		/// Output the event to the on-disk log
-		/// </summary>
-		Log = 16,
-
-		/// <summary>
-		/// The log event should only be displayed if verbose logging is enabled
-		/// </summary>
-		Verbose = 32,
-
-		/// <summary>
-		/// The log event should only be displayed if very verbose logging is enabled
-		/// </summary>
-		VeryVerbose = 64
-	}
-
 	/// <summary>
 	/// Utility functions
 	/// </summary>
@@ -893,104 +703,6 @@ namespace UnrealBuildTool
 			return Path.GetFullPath(FilePath).StartsWith(DirectoryPathPlusSeparator, StringComparison.InvariantCultureIgnoreCase);
 		}
 
-
-		/// <summary>
-		/// Given a path to a file, strips off the base directory part of the path
-		/// </summary>
-		/// <param name="FilePath">The full path</param>
-		/// <param name="BaseDirectory">The base directory, which must be the first part of the path</param>
-		/// <returns>The part of the path after the base directory</returns>
-		public static string StripBaseDirectory(string FilePath, string BaseDirectory)
-		{
-			return CleanDirectorySeparators(FilePath).Substring(CleanDirectorySeparators(BaseDirectory).Length + 1);
-		}
-
-
-		/// <summary>
-		/// Given a path to a "source" file, re-roots the file path to be located under the "destination" folder.  The part of the source file's path after the root folder is unchanged.
-		/// </summary>
-		/// <param name="FilePath"></param>
-		/// <param name="BaseDirectory"></param>
-		/// <param name="NewBaseDirectory"></param>
-		/// <returns></returns>
-		public static string MakeRerootedFilePath(string FilePath, string BaseDirectory, string NewBaseDirectory)
-		{
-			string RelativeFile = StripBaseDirectory(FilePath, BaseDirectory);
-			string DestFile = Path.Combine(NewBaseDirectory, RelativeFile);
-			return DestFile;
-		}
-
-
-		/// <summary>
-		/// Correctly collapses any ../ or ./ entries in a path.
-		/// </summary>
-		/// <param name="InPath">The path to be collapsed</param>
-		/// <returns>true if the path could be collapsed, false otherwise.</returns>
-		public static bool CollapseRelativeDirectories(ref string InPath)
-		{
-			string LocalString = InPath;
-			bool bHadBackSlashes = false;
-			// look to see what kind of slashes we had
-			if (LocalString.IndexOf("\\") != -1)
-			{
-				LocalString = LocalString.Replace("\\", "/");
-				bHadBackSlashes = true;
-			}
-
-			string ParentDir = "/..";
-			int ParentDirLength = ParentDir.Length;
-
-			for (; ; )
-			{
-				// An empty path is finished
-				if (string.IsNullOrEmpty(LocalString))
-					break;
-
-				// Consider empty paths or paths which start with .. or /.. as invalid
-				if (LocalString.StartsWith("..") || LocalString.StartsWith(ParentDir))
-					return false;
-
-				// If there are no "/.."s left then we're done
-				int Index = LocalString.IndexOf(ParentDir);
-				if (Index == -1)
-					break;
-
-				int PreviousSeparatorIndex = Index;
-				for (; ; )
-				{
-					// Find the previous slash
-					PreviousSeparatorIndex = Math.Max(0, LocalString.LastIndexOf("/", PreviousSeparatorIndex - 1));
-
-					// Stop if we've hit the start of the string
-					if (PreviousSeparatorIndex == 0)
-						break;
-
-					// Stop if we've found a directory that isn't "/./"
-					if ((Index - PreviousSeparatorIndex) > 1 && (LocalString[PreviousSeparatorIndex + 1] != '.' || LocalString[PreviousSeparatorIndex + 2] != '/'))
-						break;
-				}
-
-				// If we're attempting to remove the drive letter, that's illegal
-				int Colon = LocalString.IndexOf(":", PreviousSeparatorIndex);
-				if (Colon >= 0 && Colon < Index)
-					return false;
-
-				LocalString = LocalString.Substring(0, PreviousSeparatorIndex) + LocalString.Substring(Index + ParentDirLength);
-			}
-
-			LocalString = LocalString.Replace("./", "");
-
-			// restore back slashes now
-			if (bHadBackSlashes)
-			{
-				LocalString = LocalString.Replace("/", "\\");
-			}
-
-			// and pass back out
-			InPath = LocalString;
-			return true;
-		}
-
 		/// <summary>
 		/// Checks if given type implements given interface.
 		/// </summary>
@@ -1032,22 +744,6 @@ namespace UnrealBuildTool
 			}
 		}
 
-		/// <summary>
-		/// Helper function to get the minimum of two comparable objects
-		/// </summary>
-		public static T Min<T>(T a, T b) where T : IComparable
-		{
-			return a.CompareTo(b) < 0 ? a : b;
-		}
-
-		/// <summary>
-		/// Helper function to get the maximum of two comparable objects
-		/// </summary>
-		public static T Max<T>(T a, T b) where T : IComparable
-		{
-			return a.CompareTo(b) > 0 ? a : b;
-		}
-
 		/// <summary>
 		/// Gets the number of physical cores, excluding hyper threading.
 		/// </summary>
@@ -1082,874 +778,4 @@ namespace UnrealBuildTool
 			}
 		}
 	}
-
-	/// <summary>
-	/// Class to display an incrementing progress percentage. Handles progress markup and direct console output.
-	/// </summary>
-	public class ProgressWriter : IDisposable
-	{
-		/// <summary>
-		/// Global setting controlling whether to output markup
-		/// </summary>
-		public static bool bWriteMarkup = false;
-
-		/// <summary>
-		/// Whether to write messages to the console
-		/// </summary>
-		bool bWriteToConsole;
-		string Message;
-		int NumCharsToBackspaceOver;
-		string CurrentProgressString;
-
-		/// <summary>
-		/// Constructor
-		/// </summary>
-		/// <param name="InMessage">The message to display before the progress percentage</param>
-		/// <param name="bInWriteToConsole">Whether to write progress message to the console</param>
-		public ProgressWriter(string InMessage, bool bInWriteToConsole)
-		{
-			Message = InMessage;
-			bWriteToConsole = bInWriteToConsole;
-			if (!bWriteMarkup && bWriteToConsole)
-			{
-				Console.Write(Message + " ");
-			}
-			Write(0, 100);
-		}
-
-		/// <summary>
-		/// Write the terminating newline
-		/// </summary>
-		public void Dispose()
-		{
-			if (!bWriteMarkup && bWriteToConsole)
-			{
-				Console.WriteLine();
-			}
-		}
-
-		/// <summary>
-		/// Writes the current progress
-		/// </summary>
-		/// <param name="Numerator">Numerator for the progress fraction</param>
-		/// <param name="Denominator">Denominator for the progress fraction</param>
-		public void Write(int Numerator, int Denominator)
-		{
-			float ProgressValue = Denominator > 0 ? ((float)Numerator / (float)Denominator) : 1.0f;
-			string ProgressString = String.Format("{0}%", Math.Round(ProgressValue * 100.0f));
-			if (ProgressString != CurrentProgressString)
-			{
-				CurrentProgressString = ProgressString;
-				if (bWriteMarkup)
-				{
-					Log.WriteLine(LogEventType.Console, "@progress '{0}' {1}", Message, ProgressString);
-				}
-				else if (bWriteToConsole)
-				{
-					// Backspace over previous progress value
-					while (NumCharsToBackspaceOver-- > 0)
-					{
-						Console.Write("\b");
-					}
-
-					// Display updated progress string and keep track of how long it was
-					NumCharsToBackspaceOver = ProgressString.Length;
-					Console.Write(ProgressString);
-				}
-			}
-		}
-	}
-
-	/// <summary>
-	/// UAT/UBT Custom log system.
-	/// 
-	/// This lets you use any TraceListeners you want, but you should only call the static 
-	/// methods below, not call Trace.XXX directly, as the static methods
-	/// This allows the system to enforce the formatting and filtering conventions we desire.
-	///
-	/// For posterity, we cannot use the Trace or TraceSource class directly because of our special log requirements:
-	///   1. We possibly capture the method name of the logging event. This cannot be done as a macro, so must be done at the top level so we know how many layers of the stack to peel off to get the real function.
-	///   2. We have a verbose filter we would like to apply to all logs without having to have each listener filter individually, which would require our string formatting code to run every time.
-	///   3. We possibly want to ensure severity prefixes are logged, but Trace.WriteXXX does not allow any severity info to be passed down.
-	/// </summary>
-	static public class Log
-	{
-		/// <summary>
-		/// Guard our initialization. Mainly used by top level exception handlers to ensure its safe to call a logging function.
-		/// In general user code should not concern itself checking for this.
-		/// </summary>
-		private static bool bIsInitialized = false;
-		/// <summary>
-		/// When true, verbose loggin is enabled.
-		/// </summary>
-		private static LogEventType LogLevel = LogEventType.Log;
-		/// <summary>
-		/// When true, warnings and errors will have a WARNING: or ERROR: prexifx, respectively.
-		/// </summary>
-		private static bool bLogSeverity = false;
-		/// <summary>
-		/// When true, logs will have the calling mehod prepended to the output as MethodName:
-		/// </summary>
-		private static bool bLogSources = false;
-		/// <summary>
-		/// When true, will detect warnings and errors and set the console output color to yellow and red.
-		/// </summary>
-		private static bool bColorConsoleOutput = false;
-		/// <summary>
-		/// When configured, this tracks time since initialization to prepend a timestamp to each log.
-		/// </summary>
-		private static Stopwatch Timer;
-
-		/// <summary>
-		/// Expose the log level. This is a hack for ProcessResult.LogOutput, which wants to bypass our normal formatting scheme.
-		/// </summary>
-		public static bool bIsVerbose { get { return LogLevel >= LogEventType.Verbose; } }
-
-		/// <summary>
-		/// A collection of strings that have been already written once
-		/// </summary>
-		private static List<string> WriteOnceSet = new List<string>();
-
-		/// <summary>
-		/// Allows code to check if the log system is ready yet.
-		/// End users should NOT need to use this. It pretty much exists
-		/// to work around startup issues since this is a global singleton.
-		/// </summary>
-		/// <returns></returns>
-		public static bool IsInitialized()
-		{
-			return bIsInitialized;
-		}
-
-		/// <summary>
-		/// Allows us to change verbosity after initializing. This can happen since we initialize logging early, 
-		/// but then read the config and command line later, which could change this value.
-		/// </summary>
-		public static void SetLoggingLevel(LogEventType InLogLevel)
-		{
-			Log.LogLevel = InLogLevel;
-		}
-
-		/// <summary>
-		/// This class allows InitLogging to be called more than once to work around chicken and eggs issues with logging and parsing command lines (see UBT startup code).
-		/// </summary>
-		/// <param name="bLogTimestamps">If true, the timestamp from Log init time will be prepended to all logs.</param>
-		/// <param name="InLogLevel"></param>
-		/// <param name="bLogSeverity">If true, warnings and errors will have a WARNING: and ERROR: prefix to them. </param>
-		/// <param name="bLogSources">If true, logs will have the originating method name prepended to them.</param>
-		/// <param name="bColorConsoleOutput"></param>
-		/// <param name="TraceListeners">Collection of trace listeners to attach to the Trace.Listeners, in addition to the Default listener. The existing listeners (except the Default listener) are cleared first.</param>
-		public static void InitLogging(bool bLogTimestamps, LogEventType InLogLevel, bool bLogSeverity, bool bLogSources, bool bColorConsoleOutput, IEnumerable<TraceListener> TraceListeners)
-		{
-			bIsInitialized = true;
-			Timer = (bLogTimestamps && Timer == null) ? Stopwatch.StartNew() : null;
-			Log.LogLevel = InLogLevel;
-			Log.bLogSeverity = bLogSeverity;
-			Log.bLogSources = bLogSources;
-			Log.bColorConsoleOutput = bColorConsoleOutput;
-
-			// ensure that if InitLogging is called more than once we don't stack listeners.
-			// but always leave the default listener around.
-			for (int ListenerNdx = 0; ListenerNdx < Trace.Listeners.Count; )
-			{
-				if (Trace.Listeners[ListenerNdx].GetType() != typeof(DefaultTraceListener))
-				{
-					Trace.Listeners.RemoveAt(ListenerNdx);
-				}
-				else
-				{
-					++ListenerNdx;
-				}
-			}
-			// don't add any null listeners
-			Trace.Listeners.AddRange(TraceListeners.Where(l => l != null).ToArray());
-			Trace.AutoFlush = true;
-		}
-
-		/// <summary>
-		/// Gets the name of the Method N levels deep in the stack frame. Used to trap what method actually made the logging call.
-		/// Only used when bLogSources is true.
-		/// </summary>
-		/// <param name="StackFramesToSkip"></param>
-		/// <returns>ClassName.MethodName</returns>
-		[MethodImplAttribute(MethodImplOptions.NoInlining)]
-		private static string GetSource(int StackFramesToSkip)
-		{
-			StackFrame Frame = new StackFrame(2 + StackFramesToSkip);
-			System.Reflection.MethodBase Method = Frame.GetMethod();
-			return String.Format("{0}.{1}", Method.DeclaringType.Name, Method.Name);
-		}
-
-		/// <summary>
-		/// Converts a LogEventType into a log prefix. Only used when bLogSeverity is true.
-		/// </summary>
-		/// <param name="Severity"></param>
-		/// <returns></returns>
-		private static string GetSeverityPrefix(LogEventType Severity)
-		{
-			switch (Severity)
-			{
-				case LogEventType.Fatal:
-					return "FATAL: ";
-				case LogEventType.Error:
-					return "ERROR: ";
-				case LogEventType.Warning:
-					return "WARNING: ";
-				case LogEventType.Console:
-					return "";
-				case LogEventType.Verbose:
-					return "VERBOSE: ";
-				case LogEventType.VeryVerbose:
-					return "VVERBOSE: ";
-				default:
-					return " ";
-			}
-		}
-
-		/// <summary>
-		/// Converts a LogEventType into a message code
-		/// </summary>
-		/// <param name="Severity"></param>
-		/// <returns></returns>
-		private static int GetMessageCode(LogEventType Severity)
-		{
-			return (int)Severity;
-		}
-
-		/// <summary>
-		/// Formats message for logging. Enforces the configured options.
-		/// </summary>
-		/// <param name="StackFramesToSkip">Number of frames to skip to get to the originator of the log request.</param>
-		/// <param name="CustomSource">Custom source string to use. Use the Class.Method string if null. Only used if bLogSources = true.</param>
-		/// <param name="Verbosity">Message verbosity level</param>
-		/// <param name="Format">Message text format string</param>
-		/// <param name="Args">Message text parameters</param>
-		/// <returns>Formatted message</returns>
-		[MethodImplAttribute(MethodImplOptions.NoInlining)]
-		private static string FormatMessage(int StackFramesToSkip, string CustomSource, LogEventType Verbosity, string Format, params object[] Args)
-		{
-			return string.Format("{0}{1}{2}{3}",
-					Timer != null ? String.Format("[{0:hh\\:mm\\:ss\\.fff}] ", Timer.Elapsed) : "",
-					bLogSources ? string.Format("{0}", CustomSource == null ? GetSource(StackFramesToSkip) + ": ": CustomSource == string.Empty ? "" : CustomSource + ": ") : "",
-					bLogSeverity ? GetSeverityPrefix(Verbosity) : "",
-				// If there are no extra args, don't try to format the string, in case it has any format control characters in it (our LOCTEXT strings tend to).
-					Args.Length > 0 ? string.Format(Format, Args) : Format);
-		}
-
-		/// <summary>
-		/// Writes a formatted message to the console. All other functions should boil down to calling this method.
-		/// </summary>
-		/// <param name="StackFramesToSkip">Number of frames to skip to get to the originator of the log request.</param>
-		/// <param name="CustomSource">Custom source string to use. Use the default if null.</param>
-		/// <param name="bWriteOnce">If true, this message will be written only once</param>
-		/// <param name="Verbosity">Message verbosity level. We only meaningfully use values up to Verbose</param>
-		/// <param name="Format">Message format string.</param>
-		/// <param name="Args">Optional arguments</param>
-		[MethodImplAttribute(MethodImplOptions.NoInlining)]
-		private static void WriteLinePrivate(int StackFramesToSkip, string CustomSource, bool bWriteOnce, LogEventType Verbosity, string Format, params object[] Args)
-		{
-			if (!bIsInitialized)
-			{
-				throw new BuildException("Tried to using Logging system before it was ready");
-			}
-
-			// if we want this message only written one time, check if it was already written out
-			if (bWriteOnce)
-			{
-				string Formatted = string.Format(Format, Args);
-				if (WriteOnceSet.Contains(Formatted))
-				{
-					return;
-				}
-
-				WriteOnceSet.Add(Formatted);
-			}
-
-			if (Verbosity <= LogLevel)
-			{
-				// Do console color highlighting here.
-				ConsoleColor DefaultColor = ConsoleColor.Gray;
-				bool bIsWarning = false;
-				bool bIsError = false;
-				// don't try to touch the console unless we are told to color the output.
-				if (bColorConsoleOutput)
-				{
-					DefaultColor = Console.ForegroundColor;
-					bIsWarning = Verbosity == LogEventType.Warning;
-					bIsError = Verbosity <= LogEventType.Error;
-					// @todo mono - mono doesn't seem to initialize the ForegroundColor properly, so we can't restore it properly.
-					// Avoid touching the console color unless we really need to.
-					if (bIsWarning || bIsError)
-					{
-						Console.ForegroundColor = bIsWarning ? ConsoleColor.Yellow : ConsoleColor.Red;
-					}
-				}
-				try
-				{
-					// @todo mono: mono has some kind of bug where calling mono recursively by spawning
-					// a new process causes Trace.WriteLine to stop functioning (it returns, but does nothing for some reason).
-					// work around this by simulating Trace.WriteLine on mono.
-					// We use UAT to spawn UBT instances recursively a lot, so this bug can effectively
-					// make all build output disappear outside of the top level UAT process.
-					//					#if MONO
-					lock (((System.Collections.ICollection)Trace.Listeners).SyncRoot)
-					{
-						foreach (TraceListener l in Trace.Listeners)
-						{
-							if (Verbosity != LogEventType.Log || (l as ConsoleTraceListener) == null || LogLevel >= LogEventType.Verbose)
-							{
-								l.WriteLine(FormatMessage(StackFramesToSkip + 1, CustomSource, Verbosity, Format, Args));
-							}
-							l.Flush();
-						}
-					}
-					//					#else
-					// Call Trace directly here. Trace ensures that our logging is threadsafe using the GlobalLock.
-					//                    	Trace.WriteLine(FormatMessage(StackFramesToSkip + 1, CustomSource, Verbosity, Format, Args));
-					//					#endif
-				}
-				finally
-				{
-					// make sure we always put the console color back.
-					if (bColorConsoleOutput && (bIsWarning || bIsError))
-					{
-						Console.ForegroundColor = DefaultColor;
-					}
-				}
-			}
-		}
-
-		/// <summary>
-		/// Similar to Trace.WriteLineIf
-		/// </summary>
-		/// <param name="Condition"></param>
-		/// <param name="Verbosity"></param>
-		/// <param name="Format"></param>
-		/// <param name="Args"></param>
-		[MethodImplAttribute(MethodImplOptions.NoInlining)]
-		public static void WriteLineIf(bool Condition, LogEventType Verbosity, string Format, params object[] Args)
-		{
-			if (Condition)
-			{
-				WriteLinePrivate(1, null, false, Verbosity, Format, Args);
-			}
-		}
-
-		/// <summary>
-		/// Mostly an internal function, but expose StackFramesToSkip to allow UAT to use existing wrapper functions and still get proper formatting.
-		/// </summary>
-		/// <param name="StackFramesToSkip"></param>
-		/// <param name="Verbosity"></param>
-		/// <param name="Format"></param>
-		/// <param name="Args"></param>
-		[MethodImplAttribute(MethodImplOptions.NoInlining)]
-		public static void WriteLine(int StackFramesToSkip, LogEventType Verbosity, string Format, params object[] Args)
-		{
-			WriteLinePrivate(StackFramesToSkip + 1, null, false, Verbosity, Format, Args);
-		}
-
-		/// <summary>
-		/// Mostly an internal function, but expose StackFramesToSkip and a custom Source to alow UAT to use existing wrapper functions and still get proper formatting.
-		/// </summary>
-		/// <param name="StackFramesToSkip"></param>
-		/// <param name="CustomSource">Custom source string to use. Use the default if null.</param>
-		/// <param name="Verbosity"></param>
-		/// <param name="Format"></param>
-		/// <param name="Args"></param>
-		[MethodImplAttribute(MethodImplOptions.NoInlining)]
-		public static void WriteLine(int StackFramesToSkip, string CustomSource, LogEventType Verbosity, string Format, params object[] Args)
-		{
-			WriteLinePrivate(StackFramesToSkip + 1, CustomSource, false, Verbosity, Format, Args);
-		}
-
-		/// <summary>
-		/// Similar to Trace.WriteLin
-		/// </summary>
-		/// <param name="Verbosity"></param>
-		/// <param name="Format"></param>
-		/// <param name="Args"></param>
-		[MethodImplAttribute(MethodImplOptions.NoInlining)]
-		public static void WriteLine(LogEventType Verbosity, string Format, params object[] Args)
-		{
-			WriteLinePrivate(1, null, false, Verbosity, Format, Args);
-		}
-
-		/// <summary>
-		/// Writes an error message to the console.
-		/// </summary>
-		/// <param name="Format">Message format string</param>
-		/// <param name="Args">Optional arguments</param>
-		[MethodImplAttribute(MethodImplOptions.NoInlining)]
-		public static void TraceError(string Format, params object[] Args)
-		{
-			WriteLinePrivate(1, null, false, LogEventType.Error, Format, Args);
-		}
-
-		/// <summary>
-		/// Writes a verbose message to the console.
-		/// </summary>
-		/// <param name="Format">Message format string</param>
-		/// <param name="Args">Optional arguments</param>
-		[Conditional("TRACE")]
-		[MethodImplAttribute(MethodImplOptions.NoInlining)]
-		public static void TraceVerbose(string Format, params object[] Args)
-		{
-			WriteLinePrivate(1, null, false, LogEventType.Verbose, Format, Args);
-		}
-
-		/// <summary>
-		/// Writes a message to the console.
-		/// </summary>
-		/// <param name="Format">Message format string</param>
-		/// <param name="Args">Optional arguments</param>
-		[MethodImplAttribute(MethodImplOptions.NoInlining)]
-		public static void TraceInformation(string Format, params object[] Args)
-		{
-			WriteLinePrivate(1, null, false, LogEventType.Console, Format, Args);
-		}
-
-		/// <summary>
-		/// Writes a warning message to the console.
-		/// </summary>
-		/// <param name="Format">Message format string</param>
-		/// <param name="Args">Optional arguments</param>
-		[MethodImplAttribute(MethodImplOptions.NoInlining)]
-		public static void TraceWarning(string Format, params object[] Args)
-		{
-			WriteLinePrivate(1, null, false, LogEventType.Warning, Format, Args);
-		}
-
-		/// <summary>
-		/// Writes a very verbose message to the console.
-		/// </summary>
-		/// <param name="Format">Message format string</param>
-		/// <param name="Args">Optional arguments</param>
-		[Conditional("TRACE")]
-		[MethodImplAttribute(MethodImplOptions.NoInlining)]
-		public static void TraceVeryVerbose(string Format, params object[] Args)
-		{
-			WriteLinePrivate(1, null, false, LogEventType.VeryVerbose, Format, Args);
-		}
-
-		/// <summary>
-		/// Writes a message to the log only.
-		/// </summary>
-		/// <param name="Format">Message format string</param>
-		/// <param name="Args">Optional arguments</param>
-		[Conditional("TRACE")]
-		[MethodImplAttribute(MethodImplOptions.NoInlining)]
-		public static void TraceLog(string Format, params object[] Args)
-		{
-			WriteLinePrivate(1, null, false, LogEventType.Log, Format, Args);
-		}
-
-		/// <summary>
-		/// Similar to Trace.WriteLin
-		/// </summary>
-		/// <param name="Verbosity"></param>
-		/// <param name="Format"></param>
-		/// <param name="Args"></param>
-		[MethodImplAttribute(MethodImplOptions.NoInlining)]
-		public static void WriteLineOnce(LogEventType Verbosity, string Format, params object[] Args)
-		{
-			WriteLinePrivate(1, null, true, Verbosity, Format, Args);
-		}
-
-		/// <summary>
-		/// Writes an error message to the console.
-		/// </summary>
-		/// <param name="Format">Message format string</param>
-		/// <param name="Args">Optional arguments</param>
-		[MethodImplAttribute(MethodImplOptions.NoInlining)]
-		public static void TraceErrorOnce(string Format, params object[] Args)
-		{
-			WriteLinePrivate(1, null, true, LogEventType.Error, Format, Args);
-		}
-
-		/// <summary>
-		/// Writes a verbose message to the console.
-		/// </summary>
-		/// <param name="Format">Message format string</param>
-		/// <param name="Args">Optional arguments</param>
-		[Conditional("TRACE")]
-		[MethodImplAttribute(MethodImplOptions.NoInlining)]
-		public static void TraceVerboseOnce(string Format, params object[] Args)
-		{
-			WriteLinePrivate(1, null, true, LogEventType.Verbose, Format, Args);
-		}
-
-		/// <summary>
-		/// Writes a message to the console.
-		/// </summary>
-		/// <param name="Format">Message format string</param>
-		/// <param name="Args">Optional arguments</param>
-		[MethodImplAttribute(MethodImplOptions.NoInlining)]
-		public static void TraceInformationOnce(string Format, params object[] Args)
-		{
-			WriteLinePrivate(1, null, true, LogEventType.Console, Format, Args);
-		}
-
-		/// <summary>
-		/// Writes a warning message to the console.
-		/// </summary>
-		/// <param name="Format">Message format string</param>
-		/// <param name="Args">Optional arguments</param>
-		[MethodImplAttribute(MethodImplOptions.NoInlining)]
-		public static void TraceWarningOnce(string Format, params object[] Args)
-		{
-			WriteLinePrivate(1, null, true, LogEventType.Warning, Format, Args);
-		}
-
-		/// <summary>
-		/// Writes a very verbose message to the console.
-		/// </summary>
-		/// <param name="Format">Message format string</param>
-		/// <param name="Args">Optional arguments</param>
-		[Conditional("TRACE")]
-		[MethodImplAttribute(MethodImplOptions.NoInlining)]
-		public static void TraceVeryVerboseOnce(string Format, params object[] Args)
-		{
-			WriteLinePrivate(1, null, true, LogEventType.VeryVerbose, Format, Args);
-		}
-
-		/// <summary>
-		/// Writes a message to the log only.
-		/// </summary>
-		/// <param name="Format">Message format string</param>
-		/// <param name="Args">Optional arguments</param>
-		[Conditional("TRACE")]
-		[MethodImplAttribute(MethodImplOptions.NoInlining)]
-		public static void TraceLogOnce(string Format, params object[] Args)
-		{
-			WriteLinePrivate(1, null, true, LogEventType.Log, Format, Args);
-		}
-	}
-
-	#region StreamUtils
-
-	/// <summary>
-	/// Extension methods for Stream classes
-	/// </summary>
-	public static class StreamUtils
-	{
-		/// <summary>
-		/// Read a stream into another, buffering in 4K chunks.
-		/// </summary>
-		/// <param name="output">this</param>
-		/// <param name="input">the Stream to read from</param>
-		/// <returns>same stream for expression chaining.</returns>
-		public static Stream ReadFrom(this Stream output, Stream input)
-		{
-			long bytesRead;
-			return output.ReadFrom(input, out bytesRead);
-		}
-
-		/// <summary>
-		/// Read a stream into another, buffering in 4K chunks.
-		/// </summary>
-		/// <param name="output">this</param>
-		/// <param name="input">the Stream to read from</param>
-		/// <param name="totalBytesRead">returns bytes read</param>
-		/// <returns>same stream for expression chaining.</returns>
-		public static Stream ReadFrom(this Stream output, Stream input, out long totalBytesRead)
-		{
-			totalBytesRead = 0;
-			const int BytesToRead = 4096;
-			byte[] buf = new byte[BytesToRead];
-			int bytesReadThisTime = 0;
-			do
-			{
-				bytesReadThisTime = input.Read(buf, 0, BytesToRead);
-				totalBytesRead += bytesReadThisTime;
-				output.Write(buf, 0, bytesReadThisTime);
-			} while (bytesReadThisTime != 0);
-			return output;
-		}
-
-		/// <summary>
-		/// Read stream into a new MemoryStream. Useful for chaining together expressions.
-		/// </summary>
-		/// <param name="input">Stream to read from.</param>
-		/// <returns>memory stream that contains the stream contents.</returns>
-		public static MemoryStream ReadIntoMemoryStream(this Stream input)
-		{
-			MemoryStream data = new MemoryStream(4096);
-			data.ReadFrom(input);
-			return data;
-		}
-
-		/// <summary>
-		/// Writes the entire contents of a byte array to the stream.
-		/// </summary>
-		/// <param name="stream"></param>
-		/// <param name="arr"></param>
-		/// <returns></returns>
-		public static Stream Write(this Stream stream, byte[] arr)
-		{
-			stream.Write(arr, 0, arr.Length);
-			return stream;
-		}
-	}
-	#endregion
-
-	#region HashCodeHelper
-	/// <summary>
-	/// Helper for generating hashcodes for value types.
-	/// </summary>
-	public static class HashCodeHelper
-	{
-		private const int BasePrimeNumber = 691;
-		private const int FieldMutatorPrimeNumber = 397;
-
-		/// <summary>
-		/// Helper to compute a reasonable hash code from a group of objects, typically base types, but could be anything.
-		/// </summary>
-		/// <typeparam name="T1"></typeparam>
-		/// <param name="Value1"></param>
-		/// <returns></returns>
-		public static int GetHashCodeFromValues<T1>(T1 Value1)
-		{
-			int Result = BasePrimeNumber;
-			if (Value1 != null)
-			{
-				Result *= Value1.GetHashCode() + FieldMutatorPrimeNumber;
-			}
-			return Result;
-		}
-
-		/// <summary>
-		/// Helper to compute a reasonable hash code from a group of objects, typically base types, but could be anything.
-		/// </summary>
-		/// <typeparam name="T1"></typeparam>
-		/// <typeparam name="T2"></typeparam>
-		/// <param name="Value1"></param>
-		/// <param name="Value2"></param>
-		/// <returns></returns>
-		public static int GetHashCodeFromValues<T1, T2>(T1 Value1, T2 Value2)
-		{
-			int Result = BasePrimeNumber;
-			if (Value1 != null)
-			{
-				Result *= Value1.GetHashCode() + FieldMutatorPrimeNumber;
-			}
-			if (Value2 != null)
-			{
-				Result *= Value2.GetHashCode() + FieldMutatorPrimeNumber;
-			}
-			return Result;
-		}
-
-		/// <summary>
-		/// Helper to compute a reasonable hash code from a group of objects, typically base types, but could be anything.
-		/// </summary>
-		/// <typeparam name="T1"></typeparam>
-		/// <typeparam name="T2"></typeparam>
-		/// <typeparam name="T3"></typeparam>
-		/// <param name="Value1"></param>
-		/// <param name="Value2"></param>
-		/// <param name="Value3"></param>
-		/// <returns></returns>
-		public static int GetHashCodeFromValues<T1, T2, T3>(T1 Value1, T2 Value2, T3 Value3)
-		{
-			int Result = BasePrimeNumber;
-			if (Value1 != null)
-			{
-				Result *= Value1.GetHashCode() + FieldMutatorPrimeNumber;
-			}
-			if (Value2 != null)
-			{
-				Result *= Value2.GetHashCode() + FieldMutatorPrimeNumber;
-			}
-			if (Value3 != null)
-			{
-				Result *= Value3.GetHashCode() + FieldMutatorPrimeNumber;
-			}
-			return Result;
-		}
-
-		/// <summary>
-		/// Helper to compute a reasonable hash code from a group of objects, typically base types, but could be anything.
-		/// </summary>
-		/// <typeparam name="T1"></typeparam>
-		/// <typeparam name="T2"></typeparam>
-		/// <typeparam name="T3"></typeparam>
-		/// <typeparam name="T4"></typeparam>
-		/// <param name="Value1"></param>
-		/// <param name="Value2"></param>
-		/// <param name="Value3"></param>
-		/// <param name="Value4"></param>
-		/// <returns></returns>
-		public static int GetHashCodeFromValues<T1, T2, T3, T4>(T1 Value1, T2 Value2, T3 Value3, T4 Value4)
-		{
-			int Result = BasePrimeNumber;
-			if (Value1 != null)
-			{
-				Result *= Value1.GetHashCode() + FieldMutatorPrimeNumber;
-			}
-			if (Value2 != null)
-			{
-				Result *= Value2.GetHashCode() + FieldMutatorPrimeNumber;
-			}
-			if (Value3 != null)
-			{
-				Result *= Value3.GetHashCode() + FieldMutatorPrimeNumber;
-			}
-			if (Value4 != null)
-			{
-				Result *= Value4.GetHashCode() + FieldMutatorPrimeNumber;
-			}
-			return Result;
-		}
-
-		/// <summary>
-		/// Helper to compute a reasonable hash code from a group of objects, typically base types, but could be anything.
-		/// </summary>
-		/// <param name="values"></param>
-		/// <returns></returns>
-		public static int GetHashCodeFromValues(params object[] values)
-		{
-			// prevent arithmetic overflow exceptions
-			unchecked
-			{
-				return values.Where(elt => elt != null).Aggregate(BasePrimeNumber, (result, rhs) => result * (FieldMutatorPrimeNumber + rhs.GetHashCode()));
-			}
-		}
-	}
-	#endregion
-
-	#region Scoped Timers
-
-	/// <summary>
-	/// Scoped timer, start is in the constructor, end in Dispose. Best used with using(var Timer = new ScopedTimer()). Suports nesting.
-	/// </summary>
-	public class ScopedTimer : IDisposable
-	{
-		DateTime StartTime;
-		string TimerName;
-		LogEventType Verbosity;
-		static int Indent = 0;
-		static object IndentLock = new object();
-
-		/// <summary>
-		/// Constructor
-		/// </summary>
-		/// <param name="Name">Name of the block being measured</param>
-		/// <param name="InVerbosity">Verbosity for output messages</param>
-		public ScopedTimer(string Name, LogEventType InVerbosity = LogEventType.Verbose)
-		{
-			TimerName = Name;
-			lock (IndentLock)
-			{
-				Indent++;
-			}
-			Verbosity = InVerbosity;
-			StartTime = DateTime.UtcNow;
-		}
-
-		/// <summary>
-		/// Prints out the timing message
-		/// </summary>
-		public void Dispose()
-		{
-			double TotalSeconds = (DateTime.UtcNow - StartTime).TotalSeconds;
-			int LogIndent = 0;
-			lock (IndentLock)
-			{
-				LogIndent = --Indent;
-			}
-			StringBuilder IndentText = new StringBuilder(LogIndent * 2);
-			IndentText.Append(' ', LogIndent * 2);
-
-			Log.WriteLine(Verbosity, "{0}{1} took {2}s", IndentText.ToString(), TimerName, TotalSeconds);
-		}
-	}
-
-	/// <summary>
-	/// Scoped, accumulative timer. Can have multiple instances. Best used as a static variable with Start() and End() 
-	/// </summary>
-	public class ScopedCounter : IDisposable
-	{
-		class Accumulator
-		{
-			public double Time;
-			public LogEventType Verbosity;
-		}
-		DateTime StartTime;
-		Accumulator ScopedAccumulator;
-		string CounterName;
-
-		static Dictionary<string, Accumulator> Accumulators = new Dictionary<string, Accumulator>();
-
-		/// <summary>
-		/// Constructor
-		/// </summary>
-		/// <param name="Name"></param>
-		/// <param name="InVerbosity"></param>
-		public ScopedCounter(string Name, LogEventType InVerbosity = LogEventType.Verbose)
-		{
-			CounterName = Name;
-			if (!Accumulators.TryGetValue(CounterName, out ScopedAccumulator))
-			{
-				ScopedAccumulator = new Accumulator();
-				Accumulators.Add(CounterName, ScopedAccumulator);
-			}
-			ScopedAccumulator.Verbosity = InVerbosity;
-			StartTime = DateTime.UtcNow;
-		}
-
-		/// <summary>
-		/// (Re-)Starts time measuring
-		/// </summary>
-		/// <returns></returns>
-		public ScopedCounter Start()
-		{
-			StartTime = DateTime.UtcNow;
-			return this;
-		}
-
-		/// <summary>
-		/// Ends measuring time
-		/// </summary>
-		public void Dispose()
-		{
-			double TotalSeconds = (DateTime.UtcNow - StartTime).TotalSeconds;
-			ScopedAccumulator.Time += TotalSeconds;
-		}
-
-		/// <summary>
-		/// Logs the specified counters timings to the output
-		/// </summary>
-		/// <param name="Name">Timer name</param>
-		static public void LogCounter(string Name)
-		{
-			Accumulator ExistingAccumulator;
-			if (Accumulators.TryGetValue(Name, out ExistingAccumulator))
-			{
-				LogAccumulator(Name, ExistingAccumulator);
-			}
-			else
-			{
-				throw new Exception(String.Format("ScopedCounter {0} does not exist!", Name));
-			}
-		}
-
-		static private void LogAccumulator(string Name, Accumulator Counter)
-		{
-			Log.WriteLine(Counter.Verbosity, "{0} took {1}s", Name, Counter.Time);
-		}
-
-		/// <summary>
-		/// Dumps all counters to the log
-		/// </summary>
-		static public void DumpAllCounters()
-		{
-			foreach (KeyValuePair<string, Accumulator> Counter in Accumulators)
-			{
-				LogAccumulator(Counter.Key, Counter.Value);
-			}
-		}
-	}
-
-	#endregion
-
 }
diff --git a/Engine/Source/Programs/UnrealBuildTool/System/XcodeProject.cs b/Engine/Source/Programs/UnrealBuildTool/System/XcodeProject.cs
--- a/Engine/Source/Programs/UnrealBuildTool/System/XcodeProject.cs
+++ b/Engine/Source/Programs/UnrealBuildTool/System/XcodeProject.cs
@@ -1170,6 +1170,11 @@ namespace UnrealBuildTool
 					Content.Append("            isEnabled = \"YES\">" + ProjectFileGenerator.NewLine);
 					Content.Append("         </CommandLineArgument>" + ProjectFileGenerator.NewLine);
 				}
+				// Always add a configuration argument
+				Content.Append("         <CommandLineArgument" + ProjectFileGenerator.NewLine);
+				Content.Append("            argument = \"-RunConfig=$(Configuration)\"" + ProjectFileGenerator.NewLine);
+				Content.Append("            isEnabled = \"YES\">" + ProjectFileGenerator.NewLine);
+				Content.Append("         </CommandLineArgument>" + ProjectFileGenerator.NewLine);
 				Content.Append("      </CommandLineArguments>" + ProjectFileGenerator.NewLine);
 			}
 			Content.Append("      <AdditionalOptions>" + ProjectFileGenerator.NewLine);
diff --git a/Engine/Source/Programs/UnrealBuildTool/System/XmlConfig.cs b/Engine/Source/Programs/UnrealBuildTool/System/XmlConfig.cs
--- a/Engine/Source/Programs/UnrealBuildTool/System/XmlConfig.cs
+++ b/Engine/Source/Programs/UnrealBuildTool/System/XmlConfig.cs
@@ -48,14 +48,7 @@ namespace UnrealBuildTool
 		public static bool ReadConfigFiles()
 		{
 			// Find all the configurable types
-			List<Type> ConfigTypes = new List<Type>();
-			foreach(Type ConfigType in Assembly.GetExecutingAssembly().GetTypes())
-			{
-				if(HasXmlConfigFileAttribute(ConfigType))
-				{
-					ConfigTypes.Add(ConfigType);
-				}
-			}
+			List<Type> ConfigTypes = FindConfigurableTypes();
 
 			// Find all the input files
 			FileReference[] InputFiles = FindInputFiles().Select(x => x.Location).ToArray();
@@ -115,6 +108,23 @@ namespace UnrealBuildTool
 			return true;
 		}
 
+		/// <summary>
+		/// Find all the configurable types in the current assembly
+		/// </summary>
+		/// <returns>List of configurable types</returns>
+		static List<Type> FindConfigurableTypes()
+		{
+			List<Type> ConfigTypes = new List<Type>();
+			foreach(Type ConfigType in Assembly.GetExecutingAssembly().GetTypes())
+			{
+				if(HasXmlConfigFileAttribute(ConfigType))
+				{
+					ConfigTypes.Add(ConfigType);
+				}
+			}
+			return ConfigTypes;
+		}
+
 		/// <summary>
 		/// Determines whether the given type has a field with an XmlConfigFile attribute
 		/// </summary>
@@ -593,5 +603,101 @@ namespace UnrealBuildTool
 			// Otherwise, it's up to date
 			return true;
 		}
+
+		/// <summary>
+		/// Generates documentation files for the available settings, by merging the XML documentation from the compiler.
+		/// </summary>
+		/// <param name="OutputFile">The documentation file to write</param>
+		public static void WriteDocumentation(FileReference OutputFile)
+		{
+			// Find all the configurable types
+			List<Type> ConfigTypes = FindConfigurableTypes();
+
+			// Find all the configurable fields from the given types
+			Dictionary<string, Dictionary<string, FieldInfo>> CategoryToFields = new Dictionary<string, Dictionary<string, FieldInfo>>();
+			FindConfigurableFields(ConfigTypes, CategoryToFields);
+
+			// Get the path to the XML documentation
+			FileReference InputDocumentationFile = new FileReference(Assembly.GetExecutingAssembly().Location).ChangeExtension(".xml");
+			if(!FileReference.Exists(InputDocumentationFile))
+			{
+				throw new BuildException("Generated assembly documentation not found at {0}.", InputDocumentationFile);
+			}
+
+			// Get the current engine version for versioning the page
+			BuildVersion Version;
+			if(!BuildVersion.TryRead(out Version))
+			{
+				throw new BuildException("Unable to read the current build version");
+			}
+
+			// Read the documentation
+			XmlDocument InputDocumentation = new XmlDocument();
+			InputDocumentation.Load(InputDocumentationFile.FullName);
+
+			// Make sure we can write to the output file
+			FileReference.MakeWriteable(OutputFile);
+
+			// Generate the UDN documentation file
+			using (StreamWriter Writer = new StreamWriter(OutputFile.FullName))
+			{
+				Writer.WriteLine("Availability: NoPublish");
+				Writer.WriteLine("Title: Build Configuration Properties Page");
+				Writer.WriteLine("Crumbs:");
+				Writer.WriteLine("Description: This is a procedurally generated markdown page.");
+				Writer.WriteLine("Version: {0}.{1}", Version.MajorVersion, Version.MinorVersion);
+				Writer.WriteLine("");
+
+				foreach(KeyValuePair<string, Dictionary<string, FieldInfo>> CategoryPair in CategoryToFields)
+				{
+					string CategoryName = CategoryPair.Key;
+					Writer.WriteLine("### {0}", CategoryName);
+					Writer.WriteLine();
+
+					Dictionary<string, FieldInfo> Fields = CategoryPair.Value;
+					foreach(KeyValuePair<string, FieldInfo> FieldPair in Fields)
+					{
+						string FieldName = FieldPair.Key;
+
+						FieldInfo Field = FieldPair.Value;
+						XmlNode Node = InputDocumentation.SelectSingleNode(String.Format("//member[@name='F:{0}.{1}']/summary", Field.DeclaringType.FullName, Field.Name));
+						if(Node != null)
+						{
+							// Reflow the comments into paragraphs, assuming that each paragraph will be separated by a blank line
+							List<string> Lines = new List<string>(Node.InnerText.Trim().Split('\n').Select(x => x.Trim()));
+							for(int Idx = Lines.Count - 1; Idx > 0; Idx--)
+							{
+								if(Lines[Idx - 1].Length > 0 && !Lines[Idx].StartsWith("*") && !Lines[Idx].StartsWith("-"))
+								{
+									Lines[Idx - 1] += " " + Lines[Idx];
+									Lines.RemoveAt(Idx);
+								}
+							}
+
+							// Write the result to the .udn file
+							if(Lines.Count > 0)
+							{
+								Writer.WriteLine("$ {0} : {1}", FieldName, Lines[0]);
+								for(int Idx = 1; Idx < Lines.Count; Idx++)
+								{
+									if(Lines[Idx].StartsWith("*") || Lines[Idx].StartsWith("-"))
+									{
+										Writer.WriteLine("        * {0}", Lines[Idx].Substring(1).TrimStart());
+									}
+									else
+									{
+										Writer.WriteLine("    * {0}", Lines[Idx]);
+									}
+								}
+								Writer.WriteLine();
+							}
+						}
+					}
+				}
+			}
+
+			// Success!
+			Log.TraceInformation("Written documentation to {0}.", OutputFile);
+		}
 	}
 }
diff --git a/Engine/Source/Programs/UnrealBuildTool/UnrealBuildTool.csproj b/Engine/Source/Programs/UnrealBuildTool/UnrealBuildTool.csproj
--- a/Engine/Source/Programs/UnrealBuildTool/UnrealBuildTool.csproj
+++ b/Engine/Source/Programs/UnrealBuildTool/UnrealBuildTool.csproj
@@ -57,6 +57,7 @@
     <PlatformTarget>AnyCPU</PlatformTarget>
     <ErrorReport>prompt</ErrorReport>
     <Prefer32Bit>false</Prefer32Bit>
+    <DocumentationFile>..\..\..\Binaries\DotNET\UnrealBuildTool.XML</DocumentationFile>
   </PropertyGroup>
   <ItemGroup>
     <Reference Include="Ionic.Zip.Reduced">
@@ -134,8 +135,14 @@
     <Compile Include="Configuration\UEBuildModuleExternal.cs" />
     <Compile Include="Configuration\UEBuildPlatformFactory.cs" />
     <Compile Include="Configuration\UEBuildPlatformSDK.cs" />
+    <Compile Include="System\BuildVersion.cs" />
+    <Compile Include="System\DirectoryReference.cs" />
+    <Compile Include="System\RulesDocumentation.cs" />
     <Compile Include="System\FileFilter.cs" />
     <Compile Include="System\FileFilterNode.cs" />
+    <Compile Include="System\FileReference.cs" />
+    <Compile Include="System\FileSystemName.cs" />
+    <Compile Include="System\Log.cs" />
     <Compile Include="System\PrecompiledHeaderInstance.cs" />
     <Compile Include="System\PrecompiledHeaderTemplate.cs" />
     <Compile Include="Configuration\TargetInfo.cs" />
@@ -155,9 +162,12 @@
     <Compile Include="System\InstalledPlatformInfo.cs" />
     <Compile Include="System\LocalizationDescriptor.cs" />
     <Compile Include="System\ObsoleteOverrideAttribute.cs" />
+    <Compile Include="System\ProgressWriter.cs" />
     <Compile Include="System\RemoteExecutor.cs" />
     <Compile Include="System\RestrictedFolders.cs" />
     <Compile Include="System\RulesAssembly.cs" />
+    <Compile Include="System\ScopedTimer.cs" />
+    <Compile Include="System\StreamUtils.cs" />
     <Compile Include="System\SupportedConfigurationsAttribute.cs" />
     <Compile Include="System\SupportedPlatformsAttribute.cs" />
     <Compile Include="System\TargetReceipt.cs" />
diff --git a/Engine/Source/Programs/UnrealBuildTool/Windows/VCToolChain.cs b/Engine/Source/Programs/UnrealBuildTool/Windows/VCToolChain.cs
--- a/Engine/Source/Programs/UnrealBuildTool/Windows/VCToolChain.cs
+++ b/Engine/Source/Programs/UnrealBuildTool/Windows/VCToolChain.cs
@@ -1236,7 +1236,7 @@ namespace UnrealBuildTool
 						// the one that produces the PDB (as opposed to no debug info, where the above code
 						// is needed, but not the output PDB, or when multiple files share a single PDB, so
 						// only the action that generates it should count it as output directly)
-						if (CompileEnvironment.bUsePDBFiles && bActionProducesPDB)
+						if (bActionProducesPDB)
 						{
 							CompileAction.ProducedItems.Add(PDBFile);
 							Result.DebugDataFiles.Add(PDBFile);
@@ -1453,7 +1453,7 @@ namespace UnrealBuildTool
 			}
 
 			// @todo UE4 DLL: Why do I need LIBPATHs to build only export libraries with /DEF? (tbbmalloc.lib)
-			if (!LinkEnvironment.bIsBuildingLibrary || (LinkEnvironment.bIsBuildingLibrary && bIncludeDependentLibrariesInLibrary))
+			if (!LinkEnvironment.bIsBuildingLibrary || bIncludeDependentLibrariesInLibrary)
 			{
 				// Add the library paths to the argument list.
 				foreach (string LibraryPath in LinkEnvironment.LibraryPaths)
diff --git a/Engine/Source/Programs/UnrealGameSync/UnrealGameSync/Forms/MainWindow.cs b/Engine/Source/Programs/UnrealGameSync/UnrealGameSync/Forms/MainWindow.cs
--- a/Engine/Source/Programs/UnrealGameSync/UnrealGameSync/Forms/MainWindow.cs
+++ b/Engine/Source/Programs/UnrealGameSync/UnrealGameSync/Forms/MainWindow.cs
@@ -1824,8 +1824,15 @@ namespace UnrealGameSync
 			if(StreamName != null)
 			{
 				IReadOnlyList<string> OtherStreamNames = PerforceMonitor.OtherStreamNames;
+				IReadOnlyList<string> OtherStreamFilter = Workspace.ProjectStreamFilter;
 
 				StreamContextMenu.Items.Clear();
+
+				if(OtherStreamFilter != null && !Settings.bShowAllStreams)
+				{
+					OtherStreamNames = OtherStreamNames.Where(x => OtherStreamFilter.Any(y => y.Equals(x, StringComparison.InvariantCultureIgnoreCase)) || x.Equals(StreamName, StringComparison.InvariantCultureIgnoreCase)).ToList().AsReadOnly();
+				}
+
 				foreach(string OtherStreamName in OtherStreamNames.OrderBy(x => x).Where(x => !x.EndsWith("/Dev-Binaries")))
 				{
 					string ThisStreamName = OtherStreamName; // Local for lambda capture
@@ -1837,10 +1844,34 @@ namespace UnrealGameSync
 					}
 					StreamContextMenu.Items.Add(Item);
 				}
+
+				if(OtherStreamFilter != null)
+				{
+					StreamContextMenu.Items.Add(new ToolStripSeparator());
+					if(Settings.bShowAllStreams)
+					{
+						StreamContextMenu.Items.Add(new ToolStripMenuItem("Show Filtered Stream List...", null, new EventHandler((S, E) => SetShowAllStreamsOption(false, Bounds))));
+					}
+					else
+					{
+						StreamContextMenu.Items.Add(new ToolStripMenuItem("Show All Streams...", null, new EventHandler((S, E) => SetShowAllStreamsOption(true, Bounds))));
+					}
+				}
+
 				StreamContextMenu.Show(StatusPanel, new Point(Bounds.Left, Bounds.Bottom), ToolStripDropDownDirection.BelowRight);
 			}
 		}
 
+		private void SetShowAllStreamsOption(bool bValue, Rectangle Bounds)
+		{
+			Settings.bShowAllStreams = bValue;
+			Settings.Save();
+
+			// Showing the new context menu before the current one has closed seems to fail to display it if the new context menu is shorter (ie. the current cursor position is out of bounds).
+			// BeginInvoke() it so it won't happen until the menu has closed.
+			BeginInvoke(new MethodInvoker(() => SelectOtherStream(Bounds)));
+		}
+
 		private void SelectStream(string StreamName)
 		{
 			if(Workspace.IsBusy())
diff --git a/Engine/Source/Programs/UnrealGameSync/UnrealGameSync/Properties/AssemblyInfo.cs b/Engine/Source/Programs/UnrealGameSync/UnrealGameSync/Properties/AssemblyInfo.cs
--- a/Engine/Source/Programs/UnrealGameSync/UnrealGameSync/Properties/AssemblyInfo.cs
+++ b/Engine/Source/Programs/UnrealGameSync/UnrealGameSync/Properties/AssemblyInfo.cs
@@ -34,5 +34,5 @@ using System.Runtime.InteropServices;
 // You can specify all the values or you can default the Build and Revision Numbers 
 // by using the '*' as shown below:
 // [assembly: AssemblyVersion("1.0.*")]
-[assembly: AssemblyVersion("1.91.*")]
-[assembly: AssemblyFileVersion("1.91.0.0")]
+[assembly: AssemblyVersion("1.93.*")]
+[assembly: AssemblyFileVersion("1.93.0.0")]
diff --git a/Engine/Source/Programs/UnrealGameSync/UnrealGameSync/UserSettings.cs b/Engine/Source/Programs/UnrealGameSync/UnrealGameSync/UserSettings.cs
--- a/Engine/Source/Programs/UnrealGameSync/UnrealGameSync/UserSettings.cs
+++ b/Engine/Source/Programs/UnrealGameSync/UnrealGameSync/UserSettings.cs
@@ -60,6 +60,7 @@ namespace UnrealGameSync
 		public bool bAutoResolveConflicts;
 		public bool bUseIncrementalBuilds;
 		public bool bShowLocalTimes;
+		public bool bShowAllStreams;
 		public bool bKeepInTray;
 		public string LastProjectFileName;
 		public string[] OtherProjectFileNames;
@@ -107,6 +108,7 @@ namespace UnrealGameSync
 			bAutoResolveConflicts = (ConfigFile.GetValue("General.AutoResolveConflicts", "1") != "0");
 			bUseIncrementalBuilds = ConfigFile.GetValue("General.IncrementalBuilds", true);
 			bShowLocalTimes = ConfigFile.GetValue("General.ShowLocalTimes", false);
+			bShowAllStreams = ConfigFile.GetValue("General.ShowAllStreams", false);
 			bKeepInTray = ConfigFile.GetValue("General.KeepInTray", true);
 			LastProjectFileName = ConfigFile.GetValue("General.LastProjectFileName", null);
 			OtherProjectFileNames = ConfigFile.GetValues("General.OtherProjectFileNames", new string[0]);
@@ -311,6 +313,7 @@ namespace UnrealGameSync
 			GeneralSection.SetValue("AutoResolveConflicts", bAutoResolveConflicts);
 			GeneralSection.SetValue("IncrementalBuilds", bUseIncrementalBuilds);
 			GeneralSection.SetValue("ShowLocalTimes", bShowLocalTimes);
+			GeneralSection.SetValue("ShowAllStreams", bShowAllStreams);
 			GeneralSection.SetValue("LastProjectFileName", LastProjectFileName);
 			GeneralSection.SetValue("KeepInTray", bKeepInTray);
 			GeneralSection.SetValues("OtherProjectFileNames", OtherProjectFileNames);
diff --git a/Engine/Source/Programs/UnrealGameSync/UnrealGameSync/Workspace.cs b/Engine/Source/Programs/UnrealGameSync/UnrealGameSync/Workspace.cs
--- a/Engine/Source/Programs/UnrealGameSync/UnrealGameSync/Workspace.cs
+++ b/Engine/Source/Programs/UnrealGameSync/UnrealGameSync/Workspace.cs
@@ -122,6 +122,7 @@ namespace UnrealGameSync
 			this.SyncPaths = SyncPaths.AsReadOnly();
 
 			ProjectConfigFile = ReadProjectConfigFile(InLocalRootPath, InSelectedLocalFileName, Log);
+			ProjectStreamFilter = ReadProjectStreamFilter(Perforce, ProjectConfigFile, Log);
 		}
 
 		public void Dispose()
@@ -134,6 +135,11 @@ namespace UnrealGameSync
 			get; private set;
 		}
 
+		public IReadOnlyList<string> ProjectStreamFilter
+		{
+			get; private set;
+		}
+
 		public void Update(WorkspaceUpdateContext Context)
 		{
 			// Kill any existing sync
@@ -314,6 +320,7 @@ namespace UnrealGameSync
 					{
 						// Read the new config file
 						ProjectConfigFile = ReadProjectConfigFile(LocalRootPath, SelectedLocalFileName, Log);
+						ProjectStreamFilter = ReadProjectStreamFilter(Perforce, ProjectConfigFile, Log);
 
 						// Get the branch name
 						string BranchOrStreamName;
@@ -356,7 +363,8 @@ namespace UnrealGameSync
 						if(ProjectConfigFile.GetValue("Options.UseFastModularVersioning", false))
 						{
 							Dictionary<string, string> BuildVersionStrings = new Dictionary<string,string>();
-							BuildVersionStrings["\"Changelist\":"] = String.Format(" {0},", VersionChangeNumber);
+							BuildVersionStrings["\"Changelist\":"] = String.Format(" {0},", PendingChangeNumber);
+							BuildVersionStrings["\"CompatibleChangelist\":"] = String.Format(" {0},", VersionChangeNumber);
 							BuildVersionStrings["\"BranchName\":"] = String.Format(" \"{0}\"", BranchOrStreamName.Replace('/', '+'));
 							BuildVersionStrings["\"IsPromotedBuild\":"] = " 0,";
 							if(!UpdateVersionFile(ClientRootPath + BuildVersionFileName, BuildVersionStrings, PendingChangeNumber))
@@ -762,6 +770,23 @@ namespace UnrealGameSync
 			return ProjectConfig;
 		}
 
+		static IReadOnlyList<string> ReadProjectStreamFilter(PerforceConnection Perforce, ConfigFile ProjectConfigFile, TextWriter Log)
+		{
+			string StreamListDepotPath = ProjectConfigFile.GetValue("Options.QuickSelectStreamList", null);
+			if(StreamListDepotPath == null)
+			{
+				return null;
+			}
+
+			List<string> Lines;
+			if(!Perforce.Print(StreamListDepotPath, out Lines, Log))
+			{
+				return null;
+			}
+
+			return Lines.Select(x => x.Trim()).Where(x => x.Length > 0).ToList().AsReadOnly();
+		}
+
 		static string FormatTime(long Seconds)
 		{
 			if(Seconds >= 60)
diff --git a/Engine/Source/Programs/UnrealHeaderTool/Private/ClassMaps.cpp b/Engine/Source/Programs/UnrealHeaderTool/Private/ClassMaps.cpp
--- a/Engine/Source/Programs/UnrealHeaderTool/Private/ClassMaps.cpp
+++ b/Engine/Source/Programs/UnrealHeaderTool/Private/ClassMaps.cpp
@@ -10,9 +10,7 @@ TMap<FString, TSharedRef<FUnrealSourceFile> > GUnrealSourceFilesMap;
 TMap<UField*, TSharedRef<FUnrealTypeDefinitionInfo> > GTypeDefinitionInfoMap;
 TMap<UClass*, FString> GClassStrippedHeaderTextMap;
 TMap<UClass*, FString> GClassHeaderNameWithNoPathMap;
-TSet<UClass*> GPublicClassSet;
 TSet<FUnrealSourceFile*> GPublicSourceFileSet;
-TSet<FUnrealSourceFile*> GExportedSourceFiles;
 TMap<UProperty*, FString> GArrayDimensions;
 TMap<UPackage*,  const FManifestModule*> GPackageToManifestModuleMap;
 TMap<UField*, uint32> GGeneratedCodeCRCs;
diff --git a/Engine/Source/Programs/UnrealHeaderTool/Private/ClassMaps.h b/Engine/Source/Programs/UnrealHeaderTool/Private/ClassMaps.h
--- a/Engine/Source/Programs/UnrealHeaderTool/Private/ClassMaps.h
+++ b/Engine/Source/Programs/UnrealHeaderTool/Private/ClassMaps.h
@@ -24,8 +24,6 @@ extern TMap<FString, TSharedRef<FUnrealSourceFile> > GUnrealSourceFilesMap;
 extern TMap<UField*, TSharedRef<FUnrealTypeDefinitionInfo> > GTypeDefinitionInfoMap;
 extern TMap<UClass*, FString> GClassStrippedHeaderTextMap;
 extern TMap<UClass*, FString> GClassHeaderNameWithNoPathMap;
-extern TSet<FUnrealSourceFile*> GExportedSourceFiles;
-extern TSet<UClass*> GPublicClassSet;
 extern TSet<FUnrealSourceFile*> GPublicSourceFileSet;
 extern TMap<UProperty*, FString> GArrayDimensions;
 extern TMap<UPackage*,  const FManifestModule*> GPackageToManifestModuleMap;
