UBT: Add the shared Build ID as a build product to the receipt in precompiled builds, so it is shuttled around with the other build products. Fixes warnings in "Make Installed Build Win64/Mac" build steps in CIS.

diff --git a/Engine/Source/Programs/UnrealBuildTool/Configuration/UEBuildTarget.cs b/Engine/Source/Programs/UnrealBuildTool/Configuration/UEBuildTarget.cs
--- a/Engine/Source/Programs/UnrealBuildTool/Configuration/UEBuildTarget.cs
+++ b/Engine/Source/Programs/UnrealBuildTool/Configuration/UEBuildTarget.cs
@@ -2025,7 +2025,7 @@ namespace UnrealBuildTool
 				if (SharedBuildIdFile != null)
 				{
 					string VariablePath = TargetReceipt.InsertPathVariables(SharedBuildIdFile, UnrealBuildTool.EngineDirectory, ProjectDirectory);
-					Receipt.PrecompiledBuildDependencies.Add(VariablePath);
+					Receipt.BuildProducts.Add(new BuildProduct(VariablePath, BuildProductType.BuildResource));
 				}
 			}
 
diff --git a/Engine/Source/Programs/UnrealBuildTool/System/TargetReceipt.cs b/Engine/Source/Programs/UnrealBuildTool/System/TargetReceipt.cs
--- a/Engine/Source/Programs/UnrealBuildTool/System/TargetReceipt.cs
+++ b/Engine/Source/Programs/UnrealBuildTool/System/TargetReceipt.cs
@@ -50,6 +50,11 @@ namespace UnrealBuildTool
 		/// A resource file which was generated by the build and is required for the executable to run.
 		/// </summary>
 		RequiredResource,
+
+		/// <summary>
+		/// A build resource which was generated by the build, but is not required for the executable to run.
+		/// </summary>
+		BuildResource,
 	}
 
 	/// <summary>
