UBT: Speculative fix for editor build failures in development streams. Modifying the timestamp for the shared BuildId.txt file results in errors from BuildGraph due to modified files. Does not affect //UE4/Main because build nodes happen to not be dependent on one another.

diff --git a/Engine/Source/Programs/UnrealBuildTool/Configuration/UEBuildTarget.cs b/Engine/Source/Programs/UnrealBuildTool/Configuration/UEBuildTarget.cs
--- a/Engine/Source/Programs/UnrealBuildTool/Configuration/UEBuildTarget.cs
+++ b/Engine/Source/Programs/UnrealBuildTool/Configuration/UEBuildTarget.cs
@@ -2197,7 +2197,7 @@ namespace UnrealBuildTool
 					Directory.CreateDirectory(Path.GetDirectoryName(ForceReceiptFileName));
 					Receipt.Write(ForceReceiptFileName);
 				}
-				if(SharedBuildIdFile != null)
+				if(SharedBuildIdFile != null && (!FileReference.Exists(SharedBuildIdFile) || File.ReadAllText(SharedBuildIdFile.FullName) != Receipt.BuildId))
 				{
 					DirectoryReference.CreateDirectory(SharedBuildIdFile.Directory);
 					File.WriteAllText(SharedBuildIdFile.FullName, Receipt.BuildId);
